CIS fixes

#rnx

diff --git a/Engine/Plugins/Experimental/PixelStreaming/Source/PixelStreaming/Private/ProxyConnection.cpp b/Engine/Plugins/Experimental/PixelStreaming/Source/PixelStreaming/Private/ProxyConnection.cpp
--- a/Engine/Plugins/Experimental/PixelStreaming/Source/PixelStreaming/Private/ProxyConnection.cpp
+++ b/Engine/Plugins/Experimental/PixelStreaming/Source/PixelStreaming/Private/ProxyConnection.cpp
@@ -16,7 +16,6 @@
 FProxyConnection::FProxyConnection(const FString& IP, uint16 Port, FStreamer& Streamer) :
 	Streamer(Streamer),
 	InputDevice(FModuleManager::Get().GetModuleChecked<FPixelStreamingPlugin>("PixelStreaming").GetInputDevicePtr()),
-	Thread(TEXT("WebRTC Proxy Connection"), [this, IP, Port]() { Run(IP, Port); }),
 	Socket(nullptr),
 	Listener(nullptr),
 	ExitRequested(false),
