Fix pixel streaming plugin crashes with -dx12 parameter because NvEnc is only supported on DX11

diff --git a/Engine/Plugins/Experimental/PixelStreaming/Source/PixelStreaming/Private/NvVideoEncoder.cpp b/Engine/Plugins/Experimental/PixelStreaming/Source/PixelStreaming/Private/NvVideoEncoder.cpp
--- a/Engine/Plugins/Experimental/PixelStreaming/Source/PixelStreaming/Private/NvVideoEncoder.cpp
+++ b/Engine/Plugins/Experimental/PixelStreaming/Source/PixelStreaming/Private/NvVideoEncoder.cpp
@@ -161,6 +161,8 @@ FNvVideoEncoder::FNvVideoEncoderImpl::FNvVideoEncoderImpl(void* DllHandle, const
 	uint32 Width = Settings.Width;
 	uint32 Height = Settings.Height;
 
+	FString RHIName = GDynamicRHI->GetName();
+	check(RHIName == TEXT("D3D11"));
 	ID3D11Device* Device = static_cast<ID3D11Device*>(GDynamicRHI->RHIGetNativeDevice());
 	checkf(Device != nullptr, TEXT("Cannot initialize NvEnc with invalid device"));
 	checkf(Width > 0 && Height > 0, TEXT("Cannot initialize NvEnc with invalid width/height"));
diff --git a/Engine/Plugins/Experimental/PixelStreaming/Source/PixelStreaming/Private/PixelStreamingPlugin.cpp b/Engine/Plugins/Experimental/PixelStreaming/Source/PixelStreaming/Private/PixelStreamingPlugin.cpp
--- a/Engine/Plugins/Experimental/PixelStreaming/Source/PixelStreaming/Private/PixelStreamingPlugin.cpp
+++ b/Engine/Plugins/Experimental/PixelStreaming/Source/PixelStreaming/Private/PixelStreamingPlugin.cpp
@@ -31,6 +31,18 @@ DEFINE_LOG_CATEGORY(PixelStreamingCapture);
 void FPixelStreamingPlugin::StartupModule()
 {
 	// detect hardware capabilities, init nvidia capture libs, etc
+	check(GDynamicRHI);
+	void* Device = GDynamicRHI->RHIGetNativeDevice();
+	// During cooking RHI device is invalid, skip error logging in this case as it causes the build to fail.
+	if (Device)
+	{
+		FString RHIName = GDynamicRHI->GetName();
+		if (RHIName != TEXT("D3D11"))
+		{
+			UE_LOG(PixelStreaming, Error, TEXT("Failed to initialise Pixel Streaming plugin because it only supports DX11"));
+			return;
+		}
+	}
 
 	// subscribe to engine delegates here for init / framebuffer creation / whatever
 	if (UGameEngine* GameEngine = Cast<UGameEngine>(GEngine))
