Copy of CL 4452669

* Fix incorrectly setting VideoEncoderSettings.FrameRate to 0 on startup due to timing of set for MaxFPS being done afterwards

diff --git a/Engine/Plugins/Experimental/PixelStreaming/Source/PixelStreaming/Private/Streamer.cpp b/Engine/Plugins/Experimental/PixelStreaming/Source/PixelStreaming/Private/Streamer.cpp
--- a/Engine/Plugins/Experimental/PixelStreaming/Source/PixelStreaming/Private/Streamer.cpp
+++ b/Engine/Plugins/Experimental/PixelStreaming/Source/PixelStreaming/Private/Streamer.cpp
@@ -103,7 +103,7 @@ FStreamer::FStreamer(const TCHAR* IP, uint16 Port, const FTexture2DRHIRef& Frame
 
 	ProxyConnection.Reset(new FProxyConnection(IP, Port, *this));
 
-	UpdateEncoderSettings(FrameBuffer);
+	UpdateEncoderSettings(FrameBuffer, InitialMaxFPS);
 	CreateVideoEncoder(FrameBuffer);
 
 	// This needs to be called last, after ProxyConnection is created
@@ -235,7 +235,7 @@ void FStreamer::ForceIdrFrame()
 	VideoEncoder->ForceIdrFrame();
 }
 
-void FStreamer::UpdateEncoderSettings(const FTexture2DRHIRef& FrameBuffer)
+void FStreamer::UpdateEncoderSettings(const FTexture2DRHIRef& FrameBuffer, int32 Fps)
 {
 	float MaxBitrateMbps = CVarEncoderMaxBitrate.GetValueOnRenderThread();
 
@@ -260,7 +260,7 @@ void FStreamer::UpdateEncoderSettings(const FTexture2DRHIRef& FrameBuffer)
 	VideoEncoderSettings.AverageBitRate = ReducedBitrate;
 	SET_DWORD_STAT(STAT_PixelStreaming_EncodingBitrate, VideoEncoderSettings.AverageBitRate);
 
-	VideoEncoderSettings.FrameRate = GEngine->GetMaxFPS();
+	VideoEncoderSettings.FrameRate = Fps >= 0 ? Fps : GEngine->GetMaxFPS();
 	SET_DWORD_STAT(STAT_PixelStreaming_EncodingFramerate, VideoEncoderSettings.FrameRate);
 
 	bool bUseBackBufferSize = CVarEncoderUseBackBufferSize.GetValueOnAnyThread() > 0;
diff --git a/Engine/Plugins/Experimental/PixelStreaming/Source/PixelStreaming/Private/Streamer.h b/Engine/Plugins/Experimental/PixelStreaming/Source/PixelStreaming/Private/Streamer.h
--- a/Engine/Plugins/Experimental/PixelStreaming/Source/PixelStreaming/Private/Streamer.h
+++ b/Engine/Plugins/Experimental/PixelStreaming/Source/PixelStreaming/Private/Streamer.h
@@ -51,7 +51,7 @@ public:
 private:
 	void CreateVideoEncoder(const FTexture2DRHIRef& FrameBuffer);
 	void SendSpsPpsHeader();
-	void UpdateEncoderSettings(const FTexture2DRHIRef& FrameBuffer);
+	void UpdateEncoderSettings(const FTexture2DRHIRef& FrameBuffer, int32 Fps = -1);
 	void Stream(uint64 Timestamp, PixelStreamingProtocol::EToProxyMsg, const uint8* Data, uint32 Size);
 	void SaveEncodedVideoToFile(PixelStreamingProtocol::EToProxyMsg PktType, const uint8* Data, uint32 Size);
 	void SubmitVideoFrame(uint64 Timestamp, bool KeyFrame, const uint8* Data, uint32 Size);
