fix incorrect json deserialization failure check
#rnx

diff --git a/Engine/Plugins/Experimental/PixelStreaming/Source/PixelStreaming/Private/PixelStreamingInputComponent.cpp b/Engine/Plugins/Experimental/PixelStreaming/Source/PixelStreaming/Private/PixelStreamingInputComponent.cpp
--- a/Engine/Plugins/Experimental/PixelStreaming/Source/PixelStreaming/Private/PixelStreamingInputComponent.cpp
+++ b/Engine/Plugins/Experimental/PixelStreaming/Source/PixelStreaming/Private/PixelStreamingInputComponent.cpp
@@ -98,7 +98,7 @@ void UPixelStreamingInputComponent::AddJsonStringValue(const FString& Descriptor
 	if (!Descriptor.IsEmpty())
 	{
 		TSharedRef<TJsonReader<>> JsonReader = TJsonReaderFactory<>::Create(Descriptor);
-		if (!FJsonSerializer::Deserialize(JsonReader, JsonObject) && JsonObject.IsValid())
+		if (!FJsonSerializer::Deserialize(JsonReader, JsonObject) || !JsonObject.IsValid())
 		{
 			Success = false;
 			return;
