Coerce PixelStreaming to compile given the D3D11 RHI changes
#rnx

diff --git a/Engine/Plugins/Experimental/PixelStreaming/Source/PixelStreaming/PixelStreaming.Build.cs b/Engine/Plugins/Experimental/PixelStreaming/Source/PixelStreaming/PixelStreaming.Build.cs
--- a/Engine/Plugins/Experimental/PixelStreaming/Source/PixelStreaming/PixelStreaming.Build.cs
+++ b/Engine/Plugins/Experimental/PixelStreaming/Source/PixelStreaming/PixelStreaming.Build.cs
@@ -144,9 +144,16 @@ namespace UnrealBuildTool.Rules
             if (Target.Platform == UnrealTargetPlatform.Win32 || Target.Platform == UnrealTargetPlatform.Win64)
             {
                 PrivateDependencyModuleNames.AddRange(new string[] { "D3D11RHI" });
-            }
-
-            AddWebRTCProxy();
+				PrivateIncludePaths.AddRange(
+					new string[] {
+					"../../../../Source/Runtime/Windows/D3D11RHI/Private",
+					"../../../../Source/Runtime/Windows/D3D11RHI/Private/Windows",
+					});
+				AddEngineThirdPartyPrivateStaticDependencies(Target, "IntelMetricsDiscovery");
+				AddEngineThirdPartyPrivateStaticDependencies(Target, "NVAftermath");
+			}
+
+			AddWebRTCProxy();
             AddSignallingServer();
             AddMatchmakingServer();
             AddWebRTCServers();
diff --git a/Engine/Plugins/Experimental/PixelStreaming/Source/PixelStreaming/Private/NvVideoEncoder.cpp b/Engine/Plugins/Experimental/PixelStreaming/Source/PixelStreaming/Private/NvVideoEncoder.cpp
--- a/Engine/Plugins/Experimental/PixelStreaming/Source/PixelStreaming/Private/NvVideoEncoder.cpp
+++ b/Engine/Plugins/Experimental/PixelStreaming/Source/PixelStreaming/Private/NvVideoEncoder.cpp
@@ -24,9 +24,8 @@
 #		include "Windows/AllowWindowsPlatformTypes.h"
 #			include "NvEncoder/nvEncodeAPI.h"
 #			include <d3d11.h>
-#			include "D3D11State.h"
-#			include "D3D11Resources.h"
 #		include "Windows/HideWindowsPlatformTypes.h"
+#		include "D3D11RHIPrivate.h"
 #	pragma warning(pop)
 #endif
 
