Ensure the Pixel Streaming plugin doesn't crash on Windows 7 or a machine with an AMD GPU.
Pixel Streaming is not supported on these platforms, so we show a pop-up dialogue message instead.

#rnx

diff --git a/Engine/Plugins/Experimental/PixelStreaming/Source/PixelStreaming/Private/PixelStreamingNvVideoEncoder.cpp b/Engine/Plugins/Experimental/PixelStreaming/Source/PixelStreaming/Private/PixelStreamingNvVideoEncoder.cpp
--- a/Engine/Plugins/Experimental/PixelStreaming/Source/PixelStreaming/Private/PixelStreamingNvVideoEncoder.cpp
+++ b/Engine/Plugins/Experimental/PixelStreaming/Source/PixelStreaming/Private/PixelStreamingNvVideoEncoder.cpp
@@ -16,6 +16,7 @@
 #include "Misc/CoreDelegates.h"
 #include "Modules/ModuleManager.h"
 #include "CommonRenderResources.h"
+#include "Misc/MessageDialog.h"
 
 #if defined PLATFORM_WINDOWS
 // Disable macro redefinition warning for compatibility with Windows SDK 8+
@@ -773,18 +774,25 @@ void FPixelStreamingNvVideoEncoder::FPixelStreamingNvVideoEncoderImpl::Unregiste
 }
 
 
+bool FPixelStreamingNvVideoEncoder::CheckPlatformCompatibility()
+{
+	if (FPlatformProcess::GetDllHandle(GetDllName()) == nullptr)
+	{
+		FString ErrorString(TEXT("Failed to initialize Pixel Streaming plugin because Nvidia NVENC was not installed"));
+		FText ErrorText = FText::FromString(ErrorString);
+		FText TitleText = FText::FromString(TEXT("Pixel Streaming Plugin"));
+		FMessageDialog::Open(EAppMsgType::Ok, ErrorText, &TitleText);
+		UE_LOG(PixelStreaming, Error, TEXT("%s"), *ErrorString);
+		return false;
+	}
+
+	return true;
+}
+
 FPixelStreamingNvVideoEncoder::FPixelStreamingNvVideoEncoder(const FVideoEncoderSettings& Settings, const FTexture2DRHIRef& BackBuffer, const FEncodedFrameReadyCallback& InEncodedFrameReadyCallback)
 	: NvVideoEncoderImpl(nullptr), DllHandle(nullptr)
 {
-#if defined PLATFORM_WINDOWS
-#if defined _WIN64
-	DllHandle = FPlatformProcess::GetDllHandle(TEXT("nvEncodeAPI64.dll"));
-#else
-	DllHandle = FPlatformProcess::GetDllHandle(TEXT("nvEncodeAPI.dll"));
-#endif
-#else
-	DllHandle = FPlatformProcess::GetDllHandle(TEXT("libnvidia-encode.so.1"));
-#endif
+	DllHandle = FPlatformProcess::GetDllHandle(GetDllName());
 	checkf(DllHandle != nullptr, TEXT("Failed to load NvEncode dll"));
 
 	if (DllHandle)
@@ -833,3 +841,17 @@ void FPixelStreamingNvVideoEncoder::ForceIdrFrame()
 	NvVideoEncoderImpl->ForceIdrFrame();
 }
 
+const TCHAR* FPixelStreamingNvVideoEncoder::GetDllName()
+{
+#if defined PLATFORM_WINDOWS
+#if defined _WIN64
+	return TEXT("nvEncodeAPI64.dll");
+#else
+	return TEXT("nvEncodeAPI.dll");
+#endif
+#elif defined PLATFORM_LINUX
+	return TEXT("libnvidia-encode.so.1");
+#else
+	return TEXT("");
+#endif
+}
diff --git a/Engine/Plugins/Experimental/PixelStreaming/Source/PixelStreaming/Private/PixelStreamingNvVideoEncoder.h b/Engine/Plugins/Experimental/PixelStreaming/Source/PixelStreaming/Private/PixelStreamingNvVideoEncoder.h
--- a/Engine/Plugins/Experimental/PixelStreaming/Source/PixelStreaming/Private/PixelStreamingNvVideoEncoder.h
+++ b/Engine/Plugins/Experimental/PixelStreaming/Source/PixelStreaming/Private/PixelStreamingNvVideoEncoder.h
@@ -14,6 +14,12 @@ public:
 	FPixelStreamingNvVideoEncoder(const FVideoEncoderSettings& InSettings, const FTexture2DRHIRef& BackBuffer, const FEncodedFrameReadyCallback& InEncodedFrameReadyCallback);
 	~FPixelStreamingNvVideoEncoder();
 
+	/**
+	 * Check to see if the Nvidia NVENC Video Encoder is available on the
+	 * platform we are running on.
+	 */
+	static bool CheckPlatformCompatibility();
+
 	/**
 	* Return name of the encoder.
 	*/
@@ -49,5 +55,10 @@ private:
 	class FPixelStreamingNvVideoEncoderImpl;
 	FPixelStreamingNvVideoEncoderImpl* NvVideoEncoderImpl;
 	void* DllHandle;
+
+	/**
+	 * Get the name of the Nvidia NVENC Video Encoder DLL.
+	 */
+	static const TCHAR* GetDllName();
 };
 
diff --git a/Engine/Plugins/Experimental/PixelStreaming/Source/PixelStreaming/Private/PixelStreamingPlugin.cpp b/Engine/Plugins/Experimental/PixelStreaming/Source/PixelStreaming/Private/PixelStreamingPlugin.cpp
--- a/Engine/Plugins/Experimental/PixelStreaming/Source/PixelStreaming/Private/PixelStreamingPlugin.cpp
+++ b/Engine/Plugins/Experimental/PixelStreaming/Source/PixelStreaming/Private/PixelStreamingPlugin.cpp
@@ -21,6 +21,7 @@
 #include "GameFramework/GameModeBase.h"
 #include "Dom/JsonObject.h"
 #include "Misc/App.h"
+#include "Misc/MessageDialog.h"
 
 DEFINE_LOG_CATEGORY(PixelStreaming);
 DEFINE_LOG_CATEGORY(PixelStreamingInput);
@@ -30,6 +31,14 @@ DEFINE_LOG_CATEGORY(PixelStreamingCapture);
 /** IModuleInterface implementation */
 void FPixelStreamingPlugin::StartupModule()
 {
+	// Check to see if we can use the Pixel Streaming plugin on this platform.
+	// If not then we avoid setting up our delegates to prevent access to the
+	// plugin.
+	if (!CheckPlatformCompatibility())
+	{
+		return;
+	}
+
 	// detect hardware capabilities, init nvidia capture libs, etc
 	check(GDynamicRHI);
 	void* Device = GDynamicRHI->RHIGetNativeDevice();
@@ -74,6 +83,29 @@ void FPixelStreamingPlugin::ShutdownModule()
 	IModularFeatures::Get().UnregisterModularFeature(GetModularFeatureName(), this);
 }
 
+bool FPixelStreamingPlugin::CheckPlatformCompatibility() const
+{
+	bool bCompatible = true;
+
+	bool bWin8OrHigher = FWindowsPlatformMisc::VerifyWindowsVersion(6, 2);
+	if (!bWin8OrHigher)
+	{
+		FString ErrorString(TEXT("Failed to initialize Pixel Streaming plugin because minimum requirement is Windows 8"));
+		FText ErrorText = FText::FromString(ErrorString);
+		FText TitleText = FText::FromString(TEXT("Pixel Streaming Plugin"));
+		FMessageDialog::Open(EAppMsgType::Ok, ErrorText, &TitleText);
+		UE_LOG(PixelStreaming, Error, TEXT("%s"), *ErrorString);
+		bCompatible = false;
+	}
+
+	if (!FStreamer::CheckPlatformCompatibility())
+	{
+		bCompatible = false;
+	}
+
+	return bCompatible;
+}
+
 void FPixelStreamingPlugin::UpdateViewport(FSceneViewport* Viewport)
 {
 	FRHIViewport* const ViewportRHI = Viewport->GetViewportRHI().GetReference();
diff --git a/Engine/Plugins/Experimental/PixelStreaming/Source/PixelStreaming/Private/PixelStreamingPlugin.h b/Engine/Plugins/Experimental/PixelStreaming/Source/PixelStreaming/Private/PixelStreamingPlugin.h
--- a/Engine/Plugins/Experimental/PixelStreaming/Source/PixelStreaming/Private/PixelStreamingPlugin.h
+++ b/Engine/Plugins/Experimental/PixelStreaming/Source/PixelStreaming/Private/PixelStreamingPlugin.h
@@ -40,6 +40,7 @@ public:
 
 private:
 
+	bool CheckPlatformCompatibility() const;
 	void UpdateViewport(FSceneViewport* Viewport);
 	void OnBackBufferReady_RenderThread(SWindow& SlateWindow, const FTexture2DRHIRef& BackBuffer);
 	void OnPreResizeWindowBackbuffer(void* BackBuffer);
diff --git a/Engine/Plugins/Experimental/PixelStreaming/Source/PixelStreaming/Private/Streamer.cpp b/Engine/Plugins/Experimental/PixelStreaming/Source/PixelStreaming/Private/Streamer.cpp
--- a/Engine/Plugins/Experimental/PixelStreaming/Source/PixelStreaming/Private/Streamer.cpp
+++ b/Engine/Plugins/Experimental/PixelStreaming/Source/PixelStreaming/Private/Streamer.cpp
@@ -138,6 +138,11 @@ void FStreamer::SendSpsPpsHeader()
 	Stream(FPlatformTime::Seconds(), PixelStreamingProtocol::EToProxyMsg::SpsPps, SpsPps.GetData(), SpsPps.Num());
 }
 
+bool FStreamer::CheckPlatformCompatibility()
+{
+	return FPixelStreamingNvVideoEncoder::CheckPlatformCompatibility();
+}
+
 void FStreamer::OnFrameBufferReady(const FTexture2DRHIRef& FrameBuffer)
 {
 	if (!bStreamingStarted)
diff --git a/Engine/Plugins/Experimental/PixelStreaming/Source/PixelStreaming/Private/Streamer.h b/Engine/Plugins/Experimental/PixelStreaming/Source/PixelStreaming/Private/Streamer.h
--- a/Engine/Plugins/Experimental/PixelStreaming/Source/PixelStreaming/Private/Streamer.h
+++ b/Engine/Plugins/Experimental/PixelStreaming/Source/PixelStreaming/Private/Streamer.h
@@ -27,6 +27,8 @@ public:
 	FStreamer(const TCHAR* IP, uint16 Port, const FTexture2DRHIRef& FrameBuffer);
 	virtual ~FStreamer();
 
+	static bool CheckPlatformCompatibility();
+
 	void OnFrameBufferReady(const FTexture2DRHIRef& FrameBuffer);
 	void OnPreResizeWindowBackbuffer();
 	void OnAudioPCMPacketReady(const uint8* Data, int Size);
