For the previous fix in CL 8030129, don't check platform compatability when running the Editor as Pixel Streaming doesn't work in the Editor and the plugin is passive anyway. If we do perform the checks then we just prevent somebody with a AMD video card from packaging a project which has the Pixel Streaming plugin activated.
Also, check to see if the plugin is running on a Nvidia device, in case somebody has a machine which used to have an Nvidia card but was swapped out for an AMD card.

#rnx

diff --git a/Engine/Plugins/Experimental/PixelStreaming/Source/PixelStreaming/Private/PixelStreamingNvVideoEncoder.cpp b/Engine/Plugins/Experimental/PixelStreaming/Source/PixelStreaming/Private/PixelStreamingNvVideoEncoder.cpp
--- a/Engine/Plugins/Experimental/PixelStreaming/Source/PixelStreaming/Private/PixelStreamingNvVideoEncoder.cpp
+++ b/Engine/Plugins/Experimental/PixelStreaming/Source/PixelStreaming/Private/PixelStreamingNvVideoEncoder.cpp
@@ -776,11 +776,19 @@ void FPixelStreamingNvVideoEncoder::FPixelStreamingNvVideoEncoderImpl::Unregiste
 
 bool FPixelStreamingNvVideoEncoder::CheckPlatformCompatibility()
 {
-	if (FPlatformProcess::GetDllHandle(GetDllName()) == nullptr)
+	FText TitleText = FText::FromString(TEXT("Pixel Streaming Plugin"));
+	if (!IsRHIDeviceNVIDIA())
+	{
+		FString ErrorString(TEXT("Failed to initialize Pixel Streaming plugin because platform does not have Nvidia device"));
+		FText ErrorText = FText::FromString(ErrorString);
+		FMessageDialog::Open(EAppMsgType::Ok, ErrorText, &TitleText);
+		UE_LOG(PixelStreaming, Error, TEXT("%s"), *ErrorString);
+		return false;
+	}
+	else if (FPlatformProcess::GetDllHandle(GetDllName()) == nullptr)
 	{
 		FString ErrorString(TEXT("Failed to initialize Pixel Streaming plugin because Nvidia NVENC was not installed"));
 		FText ErrorText = FText::FromString(ErrorString);
-		FText TitleText = FText::FromString(TEXT("Pixel Streaming Plugin"));
 		FMessageDialog::Open(EAppMsgType::Ok, ErrorText, &TitleText);
 		UE_LOG(PixelStreaming, Error, TEXT("%s"), *ErrorString);
 		return false;
diff --git a/Engine/Plugins/Experimental/PixelStreaming/Source/PixelStreaming/Private/PixelStreamingPlugin.cpp b/Engine/Plugins/Experimental/PixelStreaming/Source/PixelStreaming/Private/PixelStreamingPlugin.cpp
--- a/Engine/Plugins/Experimental/PixelStreaming/Source/PixelStreaming/Private/PixelStreamingPlugin.cpp
+++ b/Engine/Plugins/Experimental/PixelStreaming/Source/PixelStreaming/Private/PixelStreamingPlugin.cpp
@@ -33,8 +33,9 @@ void FPixelStreamingPlugin::StartupModule()
 {
 	// Check to see if we can use the Pixel Streaming plugin on this platform.
 	// If not then we avoid setting up our delegates to prevent access to the
-	// plugin.
-	if (!CheckPlatformCompatibility())
+	// plugin. Note that Pixel Streaming is not currently performed in the
+	// Editor.
+	if (!GIsEditor && !CheckPlatformCompatibility())
 	{
 		return;
 	}
