removed the call to FRHICommandList::SetRenderTargets as it's not needed there

diff --git a/Engine/Plugins/Experimental/PixelStreaming/Source/PixelStreaming/Private/PixelStreamingNvVideoEncoder.cpp b/Engine/Plugins/Experimental/PixelStreaming/Source/PixelStreaming/Private/PixelStreamingNvVideoEncoder.cpp
--- a/Engine/Plugins/Experimental/PixelStreaming/Source/PixelStreaming/Private/PixelStreamingNvVideoEncoder.cpp
+++ b/Engine/Plugins/Experimental/PixelStreaming/Source/PixelStreaming/Private/PixelStreamingNvVideoEncoder.cpp
@@ -470,7 +470,6 @@ void FPixelStreamingNvVideoEncoder::FPixelStreamingNvVideoEncoderImpl::CopyBackB
 				EDRF_Default);
 		}
 		RHICmdList.EndRenderPass();
-		RHICmdList.SetRenderTargets(0, nullptr, nullptr, 0, nullptr);
 	}
 }
 
