Merging 9585196 from //UE4/Dev-CharacterTech/... to //UE4/Release-4.24/...

For Pixel Streaming added an OnAllConnectionsClosed delegate which will be broadcast to when everybody has disconnected from the app via Pixel Streaming. This is an opportunity to reset the app, ready for a new user.
In Blueprint, the GetPixelStreamerDelegates function may be used to get access to the OnAllConnectionsClosed delegate so you may bind to it.

[RN] Added an OnAllConnectionsClosed delegate to detect when everybody disconnects from an app via Pixel Streaming.

diff --git a/Engine/Plugins/Experimental/PixelStreaming/Source/PixelStreaming/Private/PixelStreamerDelegates.cpp b/Engine/Plugins/Experimental/PixelStreaming/Source/PixelStreaming/Private/PixelStreamerDelegates.cpp
new file mode 100644
--- /dev/null
+++ b/Engine/Plugins/Experimental/PixelStreaming/Source/PixelStreaming/Private/PixelStreamerDelegates.cpp
@@ -0,0 +1,11 @@
+// Copyright 1998-2019 Epic Games, Inc. All Rights Reserved.
+
+#include "PixelStreamerDelegates.h"
+
+UPixelStreamerDelegates* UPixelStreamerDelegates::Singleton = nullptr;
+
+void UPixelStreamerDelegates::CreateInstance()
+{
+	Singleton = NewObject<UPixelStreamerDelegates>();
+	Singleton->AddToRoot();
+}
diff --git a/Engine/Plugins/Experimental/PixelStreaming/Source/PixelStreaming/Private/PixelStreamerDelegates.h b/Engine/Plugins/Experimental/PixelStreaming/Source/PixelStreaming/Private/PixelStreamerDelegates.h
new file mode 100644
--- /dev/null
+++ b/Engine/Plugins/Experimental/PixelStreaming/Source/PixelStreaming/Private/PixelStreamerDelegates.h
@@ -0,0 +1,41 @@
+// Copyright 1998-2019 Epic Games, Inc. All Rights Reserved.
+
+#pragma once
+
+#include "PixelStreamerDelegates.generated.h"
+
+UCLASS()
+class UPixelStreamerDelegates : public UObject
+{
+	GENERATED_BODY()
+
+public:
+
+	/**
+	 * All connections have closed and nobody is viewing or interacting with
+	 * the app. This is an opportunity to reset the app.
+	 */
+	DECLARE_DYNAMIC_MULTICAST_DELEGATE(FOnAllConnectionsClosed);
+	UPROPERTY(BlueprintAssignable, Category = "Pixel Streamer Delegates")
+	FOnAllConnectionsClosed OnAllConnectionsClosed;
+
+	/**
+	 * Create the singleton.
+	 */
+	static void CreateInstance();
+
+	/**
+	 * Get the singleton. This allows application-specific blueprints to bind
+	 * to delegates of interest.
+	 */
+	UFUNCTION(BlueprintCallable, Category = "Pixel Streamer Delegates")
+	static UPixelStreamerDelegates* GetPixelStreamerDelegates()
+	{
+		return Singleton;
+	}
+
+private:
+
+	// The singleton object.
+	static UPixelStreamerDelegates* Singleton;
+};
diff --git a/Engine/Plugins/Experimental/PixelStreaming/Source/PixelStreaming/Private/PixelStreamingModule.cpp b/Engine/Plugins/Experimental/PixelStreaming/Source/PixelStreaming/Private/PixelStreamingModule.cpp
--- a/Engine/Plugins/Experimental/PixelStreaming/Source/PixelStreaming/Private/PixelStreamingModule.cpp
+++ b/Engine/Plugins/Experimental/PixelStreaming/Source/PixelStreaming/Private/PixelStreamingModule.cpp
@@ -6,6 +6,7 @@
 #include "Streamer.h"
 #include "InputDevice.h"
 #include "PixelStreamerInputComponent.h"
+#include "PixelStreamerDelegates.h"
 #include "SessionMonitorConnection.h"
 #include "SignallingServerConnection.h"
 #include "HUDStats.h"
@@ -108,6 +109,10 @@ void FPixelStreamingModule::InitStreamer()
 
 	FApp::SetUnfocusedVolumeMultiplier(1.0f);
 
+	// Allow Pixel Streaming to broadcast to various delegates bound in the
+	// application-specific blueprint.
+	UPixelStreamerDelegates::CreateInstance();
+
 	// Allow Pixel Streaming to be frozen and a freeze frame image to be used
 	// instead of the video stream.
 	UFreezeFrame::CreateInstance();
diff --git a/Engine/Plugins/Experimental/PixelStreaming/Source/PixelStreaming/Private/Streamer.cpp b/Engine/Plugins/Experimental/PixelStreaming/Source/PixelStreaming/Private/Streamer.cpp
--- a/Engine/Plugins/Experimental/PixelStreaming/Source/PixelStreaming/Private/Streamer.cpp
+++ b/Engine/Plugins/Experimental/PixelStreaming/Source/PixelStreaming/Private/Streamer.cpp
@@ -8,6 +8,7 @@
 #include "PlayerSession.h"
 #include "Codecs/VideoEncoder.h"
 #include "WebRtcLogging.h"
+#include "PixelStreamerDelegates.h"
 
 #include "Modules/ModuleManager.h"
 #include "WebSocketsModule.h"
@@ -228,6 +229,14 @@ void FStreamer::DeletePlayerSession(FPlayerId PlayerId)
 		{
 			Streams.Empty();
 		}
+
+		// Inform the application-specific blueprint that nobody is viewing or
+		// interacting with the app. This is an opportunity to reset the app.
+		UPixelStreamerDelegates* Delegates = UPixelStreamerDelegates::GetPixelStreamerDelegates();
+		if (Delegates)
+		{
+			Delegates->OnAllConnectionsClosed.Broadcast();
+		}
 	}
 	else if (bWasQualityController)
 	{
