Merging 9548630 from //UE4/Dev-CharacterTech/... to //UE4/Release-4.24/...

 Crash : Pixel Streaming v2 with NVIDIA GPU crashes if RHI thread is enabled

diff --git a/Engine/Plugins/Experimental/PixelStreaming/Source/PixelStreaming/Private/Codecs/PixelStreamingNvVideoEncoder.cpp b/Engine/Plugins/Experimental/PixelStreaming/Source/PixelStreaming/Private/Codecs/PixelStreamingNvVideoEncoder.cpp
--- a/Engine/Plugins/Experimental/PixelStreaming/Source/PixelStreaming/Private/Codecs/PixelStreamingNvVideoEncoder.cpp
+++ b/Engine/Plugins/Experimental/PixelStreaming/Source/PixelStreaming/Private/Codecs/PixelStreamingNvVideoEncoder.cpp
@@ -79,6 +79,7 @@ TAutoConsoleVariable<FString> CVarEncoderRateControl(
 	TEXT("PixelStreaming video encoder RateControl mode. Supported modes are `ConstQP`, `VBR`, `CBR`, `VBR_MinQP`"),
 	ECVF_Default);
 
+
 namespace
 {
 	constexpr const TCHAR* GetDllName()
@@ -167,12 +168,51 @@ namespace
 
 }
 
+FPixelStreamingNvVideoEncoder::FEncoderDevice::FEncoderDevice()
+{
+	if (GDynamicRHI)
+	{
+		FString RHIName = GDynamicRHI->GetName();
+		check(RHIName == TEXT("D3D11"));
+		auto D3D11Device = static_cast<ID3D11Device*>(GDynamicRHI->RHIGetNativeDevice());
+		checkf(D3D11Device != nullptr, TEXT("Cannot initialize NvEnc with invalid device"));
+
+		IDXGIDevice* DXGIDevice = nullptr;
+		CHECK_HR_DX9_VOID(D3D11Device->QueryInterface(__uuidof(IDXGIDevice), (LPVOID*)&DXGIDevice));
+
+		IDXGIAdapter* Adapter;
+		CHECK_HR_DX9_VOID(DXGIDevice->GetAdapter(&Adapter));
+
+		D3D_DRIVER_TYPE DriverType = D3D_DRIVER_TYPE_UNKNOWN;
+		uint32 DeviceFlags = D3D11Device->GetCreationFlags();
+		D3D_FEATURE_LEVEL FeatureLevel = D3D11Device->GetFeatureLevel();
+		D3D_FEATURE_LEVEL ActualFeatureLevel;
+
+		CHECK_HR_DX9_VOID(D3D11CreateDevice(
+			Adapter,
+			DriverType,
+			NULL,
+			DeviceFlags,
+			&FeatureLevel,
+			1,
+			D3D11_SDK_VERSION,
+			Device.GetInitReference(),
+			&ActualFeatureLevel,
+			DeviceContext.GetInitReference()
+		));
+
+		DXGIDevice->Release();
+	}
+	else
+	{
+		UE_LOG(PixelStreamer, Error, TEXT("Attempting to create Encoder Device without existing RHI"));
+	}
+}
+
 const uint32 BitstreamSize = 1024 * 1024 * 2;
 
 #define NV_RESULT(NvFunction) NvFunction == NV_ENC_SUCCESS
 
-FThreadSafeCounter FPixelStreamingNvVideoEncoder::ImplCounter(0);
-
 bool FPixelStreamingNvVideoEncoder::CheckPlatformCompatibility()
 {
 	if (!IsRHIDeviceNVIDIA())
@@ -194,7 +234,7 @@ bool FPixelStreamingNvVideoEncoder::CheckPlatformCompatibility()
 	}
 }
 
-FPixelStreamingNvVideoEncoder::FPixelStreamingNvVideoEncoder(bool bEnableAsyncMode)
+FPixelStreamingNvVideoEncoder::FPixelStreamingNvVideoEncoder()
 	: InitialMaxFPS(GEngine->GetMaxFPS())
 {
 	DllHandle = FPlatformProcess::GetDllHandle(GetDllName());
@@ -204,41 +244,24 @@ FPixelStreamingNvVideoEncoder::FPixelStreamingNvVideoEncoder(bool bEnableAsyncMo
 		return;
 	}
 
-	// actual initialization must happen in render thread
-	int32 CurrImplCounter = ImplCounter.GetValue();
-	ENQUEUE_RENDER_COMMAND(NvEncInit)(
-	[this, bEnableAsyncMode, CurrImplCounter](FRHICommandListImmediate& RHICmdList)
-	{
-		if (CurrImplCounter != ImplCounter.GetValue()) // Check if the "this" we captured is still valid
-			return;
-
-		Init_RenderThread(bEnableAsyncMode);
-	});
+	Init();
 }
 
-void FPixelStreamingNvVideoEncoder::Init_RenderThread(bool bEnableAsyncMode)
+void FPixelStreamingNvVideoEncoder::Init()
 {
 	if (InitialMaxFPS == 0)
 	{
 		const float DefaultFPS = 60.0;
 		InitialMaxFPS = DefaultFPS;
 
-		check(IsInRenderingThread());
+		//check(IsInRenderingThread());
 		AsyncTask(ENamedThreads::GameThread, [this]()
 		{
 			GEngine->SetMaxFPS(InitialMaxFPS);
 		});
 	}
 
-	// Bind to the delegates that are triggered when render thread is created or destroyed, so the encoder thread can act accordingly.
-	FCoreDelegates::PostRenderingThreadCreated.AddRaw(this, &FPixelStreamingNvVideoEncoder::PostRenderingThreadCreated);
-	FCoreDelegates::PreRenderingThreadDestroyed.AddRaw(this, &FPixelStreamingNvVideoEncoder::PreRenderingThreadDestroyed);
-
-	FString RHIName = GDynamicRHI->GetName();
-	check(RHIName == TEXT("D3D11"));
-
-	auto Device = static_cast<ID3D11Device*>(GDynamicRHI->RHIGetNativeDevice());
-	checkf(Device != nullptr, TEXT("Cannot initialize NvEnc with invalid device"));
+	EncoderDevice = MakeUnique<FEncoderDevice>();
 
 	_NVENCSTATUS Result;
 
@@ -268,7 +291,7 @@ void FPixelStreamingNvVideoEncoder::Init_RenderThread(bool bEnableAsyncMode)
 		NV_ENC_OPEN_ENCODE_SESSION_EX_PARAMS OpenEncodeSessionExParams;
 		FMemory::Memzero(OpenEncodeSessionExParams);
 		OpenEncodeSessionExParams.version = NV_ENC_OPEN_ENCODE_SESSION_EX_PARAMS_VER;
-		OpenEncodeSessionExParams.device = Device;
+		OpenEncodeSessionExParams.device = EncoderDevice->Device;
 		OpenEncodeSessionExParams.deviceType = NV_ENC_DEVICE_TYPE_DIRECTX;	// Currently only DX11 is supported
 		OpenEncodeSessionExParams.apiVersion = NVENCAPI_VERSION;
 		Result = NvEncodeAPI->nvEncOpenEncodeSessionEx(&OpenEncodeSessionExParams, &EncoderInterface);
@@ -342,6 +365,7 @@ void FPixelStreamingNvVideoEncoder::Init_RenderThread(bool bEnableAsyncMode)
 		// maybe doesn't have an effect, high level is chosen because we aim at high bitrate
 		NvEncConfig.encodeCodecConfig.h264Config.level = NV_ENC_LEVEL_H264_52;
 	}
+
 	// Get encoder capability
 	{
 		NV_ENC_CAPS_PARAM CapsParam;
@@ -351,25 +375,28 @@ void FPixelStreamingNvVideoEncoder::Init_RenderThread(bool bEnableAsyncMode)
 		int32 AsyncMode = 0;
 		Result = NvEncodeAPI->nvEncGetEncodeCaps(EncoderInterface, NvEncInitializeParams.encodeGUID, &CapsParam, &AsyncMode);
 		checkf(NV_RESULT(Result), TEXT("Failed to get NVEncoder capability params (status: %d)"), Result);
-		NvEncInitializeParams.enableEncodeAsync = bEnableAsyncMode ? AsyncMode : 0;
+		if (AsyncMode == 0)
+		{
+			UE_LOG(LogVideoEncoder, Fatal, TEXT("NvEnc doesn't support async mode"));
+			return;
+		}
+
+		NvEncInitializeParams.enableEncodeAsync = true;
 	}
 
 	Result = NvEncodeAPI->nvEncInitializeEncoder(EncoderInterface, &NvEncInitializeParams);
 	checkf(NV_RESULT(Result), TEXT("Failed to initialize NVEncoder (status: %d)"), Result);
 
-	InitializeResources();
-
-	// initially all buffers are free
-	FBufferId BufferId = 0;
-	for (FFrame& frame : BufferedFrames)
+	FBufferId Id = 0;
+	for(FFrame& Frame : BufferedFrames)
 	{
-		FreeBuffers.Enqueue(BufferId++);
+		// We keep `Id` as const, because it not supposed to change at all once initialized
+		*(const_cast<FBufferId*>(&Frame.Id)) = Id++;
 	}
 
-	if (NvEncInitializeParams.enableEncodeAsync)
-	{
-		EncoderThread = MakeUnique<FThread>(TEXT("PixelStreaming Video Send"), [this]() { EncoderCheckLoop(); });
-	}
+	InitializeResources();
+
+	EncoderThread = MakeUnique<FThread>(TEXT("NvVideoEncoder"), [this]() { EncoderCheckLoop(); });
 
 	UE_LOG(LogVideoEncoder, Log, TEXT("NvEnc initialised: %dFPS%s"), static_cast<int32>(InitialMaxFPS), CVarEncoderPrioritiseQuality.GetValueOnAnyThread() != 0 ? TEXT(", prioritise quality") : TEXT(""));
 }
@@ -386,22 +413,13 @@ FPixelStreamingNvVideoEncoder::~FPixelStreamingNvVideoEncoder()
 
 	if (EncoderThread)
 	{
-		// Reset bWaitForRenderThreadToResume so encoder thread can quit
-		bWaitForRenderThreadToResume = false;
-
 		bExitEncoderThread = true;
-		// Trigger all frame events to release encoder thread waiting on them
-		// (we don't know here which frame it's waiting for)
-		for (FFrame& Frame : BufferedFrames)
-		{
-			SetEvent(Frame.OutputFrame.EventHandle);
-		}
+
+		// Trigger an event to ensure we can get out of the encoder thread.
+		SetEvent(EncodeQueue.EncodeEvent);
 
 		// Exit encoder runnable thread before shutting down NvEnc interface
 		EncoderThread->Join();
-		// Increment the counter, so that if any pending render commands sent from EncoderCheckLoop 
-		// to the Render Thread still reference "this", they will be ignored because the counter is different
-		ImplCounter.Increment();
 	}
 
 	ReleaseResources();
@@ -423,20 +441,31 @@ FPixelStreamingNvVideoEncoder::~FPixelStreamingNvVideoEncoder()
 
 bool FPixelStreamingNvVideoEncoder::CopyBackBuffer(const FTexture2DRHIRef& BackBuffer, FTimespan Timestamp, FBufferId& BufferId)
 {
-	if (!FreeBuffers.Dequeue(BufferId))
+
+	// Find a free slot we can use
+	FFrame* Frame = nullptr;
+	for (FFrame& Slot : BufferedFrames)
 	{
-		UE_LOG(LogVideoEncoder, Verbose, TEXT("Frame dropped cos NvEnc queue is full"));
-		return false;
+		if (Slot.State == EFrameState::Free)
+		{
+			Frame = &Slot;
+			BufferId = Slot.Id;
+			break;
+		}
 	}
 
-	FFrame& Frame = BufferedFrames[BufferId];
+	if (!Frame)
+	{
+		UE_LOG(LogVideoEncoder, Verbose, TEXT("Frame dropped because NvEnc queue is full"));
+		return false;
+	}
 
-	Frame.FrameIdx = CapturedFrameCount++;
-	Frame.InputFrame.CaptureTs = Timestamp;
-	CopyBackBuffer(BackBuffer, Frame.InputFrame);
+	Frame->FrameIdx = CapturedFrameCount++;
+	Frame->InputFrame.CaptureTs = Timestamp;
+	CopyBackBuffer(BackBuffer, Frame->InputFrame);
 
-	UE_LOG(LogVideoEncoder, Verbose, TEXT("Buffer #%d (%d) captured"), Frame.FrameIdx, BufferId);
-	CapturedBuffers.Enqueue(BufferId);
+	UE_LOG(LogVideoEncoder, Verbose, TEXT("Buffer #%d (%d) captured"), Frame->FrameIdx, BufferId);
+	Frame->State = EFrameState::Captured;
 
 	return true;
 }
@@ -507,42 +536,26 @@ void FPixelStreamingNvVideoEncoder::CopyBackBuffer(const FTexture2DRHIRef& BackB
 
 void FPixelStreamingNvVideoEncoder::EncodeFrame(FBufferId BufferId, const webrtc::EncodedImage& EncodedFrame, uint32 Bitrate)
 {
-	{
-		FBufferId NextCapturedBufferId;
-		bool bRes = CapturedBuffers.Dequeue(NextCapturedBufferId);
-		checkf(bRes && BufferId == NextCapturedBufferId, TEXT("%d: %d - %d"), bRes, BufferId, NextCapturedBufferId);
-	}
-
 	FFrame& Frame = BufferedFrames[BufferId];
+	checkf(Frame.State==EFrameState::Captured, TEXT("Buffer %d : Expected state %d, but found %d"), BufferId, (int)EFrameState::Captured, (int)Frame.State);
 
 	Frame.OutputFrame.EncodedFrame = EncodedFrame;
 	Frame.OutputFrame.EncodedFrame._encodedWidth = NvEncInitializeParams.encodeWidth;
 	Frame.OutputFrame.EncodedFrame._encodedHeight = NvEncInitializeParams.encodeHeight;
 
-	int32 CurrImplCounter = ImplCounter.GetValue();
-	ENQUEUE_RENDER_COMMAND(NvEncEncodeFrame)(
-	[this, &Frame, BufferId, Bitrate, CurrImplCounter](FRHICommandListImmediate& RHICmdList)
-	{
-		if (CurrImplCounter != ImplCounter.GetValue()) // Check if the "this" we captured is still valid
-			return;
-
-		EncodeFrameInRenderingThread(Frame, Bitrate);
+	Frame.State = EFrameState::Encoding;
+	// Save the requested bitrate, so we can reconfigure the encoder later
+	Frame.BitrateRequested = Bitrate;
+	EncodeQueue.Push(&Frame);
+	UE_LOG(LogVideoEncoder, VeryVerbose, TEXT("Buffer #%d (%d), ts %u started encoding"), Frame.FrameIdx, BufferId, Frame.OutputFrame.EncodedFrame.Timestamp());
 
-		UE_LOG(LogVideoEncoder, VeryVerbose, TEXT("Buffer #%d (%d), ts %u started encoding"), Frame.FrameIdx, BufferId, Frame.OutputFrame.EncodedFrame.Timestamp());
-		BuffersBeingEncoded.Enqueue(BufferId);
-	});
 }
 
 void FPixelStreamingNvVideoEncoder::OnFrameDropped(FBufferId BufferId)
 {
-	{
-		FBufferId NextCapturedBufferId;
-		bool bRes = CapturedBuffers.Dequeue(NextCapturedBufferId);
-		checkf(bRes && NextCapturedBufferId == BufferId, TEXT("%d: %d - %d"), bRes, NextCapturedBufferId, BufferId);
-	}
-
+	checkf(BufferedFrames[BufferId].State==EFrameState::Captured, TEXT("Buffer %d: Expected state %d, found %d"), BufferId, (int)EFrameState::Captured, (int)BufferedFrames[BufferId].State);
+	BufferedFrames[BufferId].State = EFrameState::Free;
 	UE_LOG(LogVideoEncoder, Log, TEXT("Buffer #%d (%d) dropped"), BufferedFrames[BufferId].FrameIdx, BufferId);
-	FreeBuffers.Enqueue(BufferId);
 }
 
 void FPixelStreamingNvVideoEncoder::SubscribeToFrameEncodedEvent(FVideoEncoder& Subscriber)
@@ -562,12 +575,12 @@ void FPixelStreamingNvVideoEncoder::UnsubscribeFromFrameEncodedEvent(FVideoEncod
 
 void FPixelStreamingNvVideoEncoder::UpdateNvEncConfig(const FInputFrame& InputFrame, uint32 Bitrate)
 {
-	check(IsInRenderingThread());
+	check(!IsInRenderingThread());
 
 	bool bSettingsChanged = false;
 	bool bResolutionChanged = false;
 
-	float MaxBitrate = CVarEncoderMaxBitrate.GetValueOnRenderThread();
+	float MaxBitrate = CVarEncoderMaxBitrate.GetValueOnAnyThread();
 	uint32 ClampedBitrate = FMath::Min(Bitrate, static_cast<uint32_t>(MaxBitrate));
 	if (NvEncConfig.rcParams.averageBitRate != ClampedBitrate)
 	{
@@ -576,7 +589,7 @@ void FPixelStreamingNvVideoEncoder::UpdateNvEncConfig(const FInputFrame& InputFr
 		bSettingsChanged = true;
 	}
 
-	uint32 MinQP = static_cast<uint32>(CVarEncoderMinQP.GetValueOnRenderThread());
+	uint32 MinQP = static_cast<uint32>(CVarEncoderMinQP.GetValueOnAnyThread());
 	MinQP = FMath::Clamp(MinQP, 0u, 54u);
 	if (NvEncConfig.rcParams.minQP.qpIntra != MinQP)
 	{
@@ -585,7 +598,7 @@ void FPixelStreamingNvVideoEncoder::UpdateNvEncConfig(const FInputFrame& InputFr
 		bSettingsChanged = true;
 	}
 
-	NV_ENC_PARAMS_RC_MODE RcMode = ToRcMode(CVarEncoderRateControl.GetValueOnRenderThread());
+	NV_ENC_PARAMS_RC_MODE RcMode = ToRcMode(CVarEncoderRateControl.GetValueOnAnyThread());
 	if (RcMode != NvEncConfig.rcParams.rateControlMode)
 	{
 		NvEncConfig.rcParams.rateControlMode = RcMode;
@@ -701,7 +714,7 @@ void FPixelStreamingNvVideoEncoder::UpdateRes(const FTexture2DRHIRef& BackBuffer
 		}
 	}
 
-	// check if target resolution matches of currently allocated `InputFrame.BackBuffer` resolution
+	// check if target resolution matches our currently allocated `InputFrame.BackBuffer` resolution
 	if (InputFrame.BackBuffer->GetSizeX() == Width && InputFrame.BackBuffer->GetSizeY() == Height)
 	{
 		return;
@@ -712,28 +725,10 @@ void FPixelStreamingNvVideoEncoder::UpdateRes(const FTexture2DRHIRef& BackBuffer
 	InitFrameInputBuffer(InputFrame, Width, Height);
 }
 
-void FPixelStreamingNvVideoEncoder::EncodeFrameInRenderingThread(FFrame& Frame, uint32 Bitrate)
-{
-	UpdateNvEncConfig(Frame.InputFrame, Bitrate);
-
-	// Encode frame
-	{
-		FRHICommandList& RHICmdList = FRHICommandListExecutor::GetImmediateCommandList();
-		if (RHICmdList.Bypass())
-		{
-			FRHITransferRenderTargetToNvEnc Command(this, &Frame);
-			Command.Execute(RHICmdList);
-		}
-		else
-		{
-			ALLOC_COMMAND_CL(RHICmdList, FRHITransferRenderTargetToNvEnc)(this, &Frame);
-		}
-	}
-}
-
-void FPixelStreamingNvVideoEncoder::TransferRenderTargetToHWEncoder(FFrame& Frame)
+void FPixelStreamingNvVideoEncoder::SubmitFrameToEncoder(FFrame& Frame)
 {
 	SCOPE_CYCLE_COUNTER(STAT_NvEnc_SendBackBufferToEncoder);
+	check(Frame.State==EFrameState::Encoding);
 
 	Frame.OutputFrame.EncodedFrame.timing_.encode_start_ms = rtc::TimeMicros() / 1000;
 
@@ -754,65 +749,76 @@ void FPixelStreamingNvVideoEncoder::TransferRenderTargetToHWEncoder(FFrame& Fram
 
 	_NVENCSTATUS Result = NvEncodeAPI->nvEncEncodePicture(EncoderInterface, &PicParams);
 	checkf(NV_RESULT(Result), TEXT("Failed to encode frame (status: %d)"), Result);
-
-	if (!NvEncInitializeParams.enableEncodeAsync)
-	{
-		// In synchronous mode, simply process the frame immediately.
-		ProcessFrame(Frame);
-	}
 }
 
 void FPixelStreamingNvVideoEncoder::EncoderCheckLoop()
 {
-	while (!bExitEncoderThread)
+	// This thread will both encode frames and will also wait for the next frame
+	// to finish encoding.
+	TQueue<FFrame*> CurrentlyEncodingQueue;
+
+	while (true)
 	{
-		FBufferId BufferId;
-		if (!BuffersBeingEncoded.Dequeue(BufferId))
+		// Wait for either the command to encode frames or the information
+		// that the next frame has finished encoding.
+		// The signalling events are a pair of handles for windows events so we
+		// can wait for one or the other.
+		enum EncodeEvents
 		{
-			FPlatformProcess::Sleep(0);
-			continue; // to check for exit request
-		}
+			START_ENCODING_EVENT = 0,
+			FINISHED_ENCODING_EVENT = 1,
+			NUM_ENCODING_EVENTS = 2
+		};
 
-		FFrame& Frame = BufferedFrames[BufferId];
+		int NumEvents = 1;
+		HANDLE Handles[NUM_ENCODING_EVENTS];
+		Handles[START_ENCODING_EVENT] = EncodeQueue.EncodeEvent;
 
+		if (!CurrentlyEncodingQueue.IsEmpty())
 		{
-			SCOPE_CYCLE_COUNTER(STAT_NvEnc_WaitForEncodeEvent);
-			DWORD Result = WaitForSingleObject(Frame.OutputFrame.EventHandle, INFINITE);
-			checkf(Result == WAIT_OBJECT_0, TEXT("Error waiting for frame event: %d"), Result);
-			if (bExitEncoderThread)
-				return;
+			NumEvents++;
+			Handles[FINISHED_ENCODING_EVENT] = (**CurrentlyEncodingQueue.Peek()).OutputFrame.EventHandle;
 		}
 
-		ResetEvent(Frame.OutputFrame.EventHandle);
-
-		// When resolution changes, render thread is stopped and later restarted from game thread.
-		// We can't enqueue render commands when render thread is stopped, so pause until render thread is restarted.
-		while (bWaitForRenderThreadToResume)
-		{
-			FPlatformProcess::Sleep(0);
-		}
+		DWORD Result = WaitForMultipleObjects(NumEvents, Handles, false, INFINITE);
 
-		FPixelStreamingNvVideoEncoder* This = this;
-		FFrame* InFrame = &Frame;
-		int32 CurrImplCounter = ImplCounter.GetValue();
-		ENQUEUE_RENDER_COMMAND(NvEncProcessFrame)(
-		[This, InFrame, BufferId, CurrImplCounter](FRHICommandListImmediate& RHICmdList)
+		if (!bExitEncoderThread)
 		{
-			if (CurrImplCounter != ImplCounter.GetValue()) // Check if the "this" we captured is still valid
+			if (Result == WAIT_OBJECT_0 + START_ENCODING_EVENT)
 			{
-				return;
+				// Get the list of all frames we want to encode.
+				FFrame* Frames[NumBufferedFrames];
+				int NumFrames;
+				EncodeQueue.PopAll(Frames, NumFrames);
+				for (int Idx = 0; Idx < NumFrames; Idx++)
+				{
+					FFrame& Frame = *Frames[Idx];
+					UpdateNvEncConfig(Frame.InputFrame, Frame.BitrateRequested);
+					SubmitFrameToEncoder(Frame);
+					CurrentlyEncodingQueue.Enqueue(&Frame);
+				}
 			}
+			else if (Result == WAIT_OBJECT_0 + FINISHED_ENCODING_EVENT)
+			{
+				FFrame* Frame = nullptr;
+				verify(CurrentlyEncodingQueue.Dequeue(Frame));
+				ResetEvent(Frame->OutputFrame.EventHandle);
+				UE_LOG(LogVideoEncoder, Verbose, TEXT("Buffer #%d (%d) encoded"), Frame->FrameIdx, Frame->Id);
+				ProcessFrame(*Frame);
+			}
+		}
+		else
+		{
+			break;
+		}
 
-			This->ProcessFrame(*InFrame);
-
-			UE_LOG(LogVideoEncoder, Verbose, TEXT("Buffer #%d (%d) encoded"), InFrame->FrameIdx, BufferId);
-			This->FreeBuffers.Enqueue(BufferId);
-		});
 	}
 }
 
 void FPixelStreamingNvVideoEncoder::ProcessFrame(FFrame& Frame)
 {
+	check(Frame.State==EFrameState::Encoding);
+
 	FOutputFrame& OutputFrame = Frame.OutputFrame;
 
 	FHUDStats& Stats = FHUDStats::Get();
@@ -865,13 +871,15 @@ void FPixelStreamingNvVideoEncoder::ProcessFrame(FFrame& Frame)
 		Stats.EncoderQP.Update(OutputFrame.EncodedFrame.qp_);
 	}
 
-	UE_LOG(LogVideoEncoder, VeryVerbose, TEXT("encoded %s ts %u, capture ts %lld, QP %d/%.0f,  latency %.0f/%.0f ms, bitrate %.3f/%.3f/%.3f Mbps, %zu bytes"), ToString(OutputFrame.EncodedFrame._frameType), OutputFrame.EncodedFrame.Timestamp(), Frame.InputFrame.CaptureTs.GetTicks(), OutputFrame.EncodedFrame.qp_, Stats.EncoderQP.Get(), LatencyMs, Stats.EncoderLatencyMs.Get(), RequestedBitrateMbps.Load(), BitrateMbps, Stats.EncoderBitrateMbps.Get(), OutputFrame.EncodedFrame._length);
+	UE_LOG(LogVideoEncoder, VeryVerbose, TEXT("encoded %s ts %u, capture ts %lld, QP %d/%.0f,  latency %.0f/%.0f ms, bitrate %.3f/%.3f/%.3f Mbps, %zu bytes"), ToString(OutputFrame.EncodedFrame._frameType), OutputFrame.EncodedFrame.Timestamp(), Frame.InputFrame.CaptureTs.GetTicks(), OutputFrame.EncodedFrame.qp_, Stats.EncoderQP.Get(), LatencyMs, Stats.EncoderLatencyMs.Get(), RequestedBitrateMbps, BitrateMbps, Stats.EncoderBitrateMbps.Get(), OutputFrame.EncodedFrame._length);
 
 	// Stream the encoded frame
 	{
 		SCOPE_CYCLE_COUNTER(STAT_NvEnc_StreamEncodedFrame);
 		OnEncodedFrame(OutputFrame.EncodedFrame);
 	}
+
+	Frame.State = EFrameState::Free;
 }
 
 void FPixelStreamingNvVideoEncoder::InitFrameInputBuffer(FInputFrame& InputFrame, uint32 Width, uint32 Height)
@@ -880,22 +888,36 @@ void FPixelStreamingNvVideoEncoder::InitFrameInputBuffer(FInputFrame& InputFrame
 	{
 		// Make sure format used here is compatible with NV_ENC_BUFFER_FORMAT specified later in NV_ENC_REGISTER_RESOURCE bufferFormat
 		FRHIResourceCreateInfo CreateInfo;
-		InputFrame.BackBuffer = RHICreateTexture2D(Width, Height, EPixelFormat::PF_A2B10G10R10, 1, 1, TexCreate_RenderTargetable, CreateInfo);
+		InputFrame.BackBuffer = RHICreateTexture2D(Width, Height, EPixelFormat::PF_B8G8R8A8, 1, 1, TexCreate_RenderTargetable | TexCreate_Shared, CreateInfo);
+	}
+
+	// Share this texture with the encoder device.
+	{
+		ID3D11Texture2D* ResolvedBackBuffer = (ID3D11Texture2D*)InputFrame.BackBuffer->GetTexture2D()->GetNativeResource();
+
+		IDXGIResource* DXGIResource = NULL;
+		CHECK_HR_DX9_VOID(ResolvedBackBuffer->QueryInterface(__uuidof(IDXGIResource), (LPVOID*)&DXGIResource));
+
+		HANDLE SharedHandle;
+		CHECK_HR_DX9_VOID(DXGIResource->GetSharedHandle(&SharedHandle));
+		CHECK_HR_DX9_VOID(DXGIResource->Release());
+
+		CHECK_HR_DX9_VOID(EncoderDevice->Device->OpenSharedResource(SharedHandle, __uuidof(ID3D11Texture2D), (LPVOID*)&InputFrame.SharedBackBuffer));
 	}
 
+
 	// Register input back buffer
 	{
 		NV_ENC_REGISTER_RESOURCE RegisterResource;
-		ID3D11Texture2D* ResolvedBackBufferDX11 = (ID3D11Texture2D*)(GetD3D11TextureFromRHITexture(InputFrame.BackBuffer)->GetResource());
+		FMemory::Memzero(RegisterResource);
 		EPixelFormat PixelFormat = InputFrame.BackBuffer->GetFormat();
 
-		FMemory::Memzero(RegisterResource);
 		RegisterResource.version = NV_ENC_REGISTER_RESOURCE_VER;
 		RegisterResource.resourceType = NV_ENC_INPUT_RESOURCE_TYPE_DIRECTX;
-		RegisterResource.resourceToRegister = (void*)ResolvedBackBufferDX11;
+		RegisterResource.resourceToRegister = (void*)InputFrame.SharedBackBuffer;
 		RegisterResource.width = Width;
 		RegisterResource.height = Height;
-		RegisterResource.bufferFormat = NV_ENC_BUFFER_FORMAT_ABGR10;	// Make sure ResolvedBackBuffer is created with a compatible format
+		RegisterResource.bufferFormat = NV_ENC_BUFFER_FORMAT_ABGR;	// Make sure ResolvedBackBuffer is created with a compatible format
 		_NVENCSTATUS Result = NvEncodeAPI->nvEncRegisterResource(EncoderInterface, &RegisterResource);
 		checkf(NV_RESULT(Result), TEXT("Failed to register input back buffer (status: %d)"), Result);
 
@@ -934,25 +956,33 @@ void FPixelStreamingNvVideoEncoder::InitializeResources()
 			checkf(NV_RESULT(Result), TEXT("Failed to create NvEnc bitstream buffer (status: %d)"), Result);
 			Frame.OutputFrame.BitstreamBuffer = CreateBitstreamBuffer.bitstreamBuffer;
 		}
-		// Register event handles
-		if (NvEncInitializeParams.enableEncodeAsync)
-		{
-			RegisterAsyncEvent(&Frame.OutputFrame.EventHandle);
-		}
+
+		RegisterAsyncEvent(&Frame.OutputFrame.EventHandle);
 	}
 }
 
 void FPixelStreamingNvVideoEncoder::ReleaseFrameInputBuffer(FInputFrame& InputFrame)
 {
-	_NVENCSTATUS Result = NvEncodeAPI->nvEncUnmapInputResource(EncoderInterface, InputFrame.MappedResource);
-	checkf(NV_RESULT(Result), TEXT("Failed to unmap input resource (status: %d)"), Result);
-	InputFrame.MappedResource = nullptr;
+	if (InputFrame.MappedResource)
+	{
+		_NVENCSTATUS Result = NvEncodeAPI->nvEncUnmapInputResource(EncoderInterface, InputFrame.MappedResource);
+		checkf(NV_RESULT(Result), TEXT("Failed to unmap input resource (status: %d)"), Result);
+		InputFrame.MappedResource = nullptr;
+	}
 
-	Result = NvEncodeAPI->nvEncUnregisterResource(EncoderInterface, InputFrame.RegisteredResource);
-	checkf(NV_RESULT(Result), TEXT("Failed to unregister input buffer resource (status: %d)"), Result);
-	InputFrame.RegisteredResource = nullptr;
+	if (InputFrame.RegisteredResource)
+	{
+		_NVENCSTATUS Result = NvEncodeAPI->nvEncUnregisterResource(EncoderInterface, InputFrame.RegisteredResource);
+		checkf(NV_RESULT(Result), TEXT("Failed to unregister input buffer resource (status: %d)"), Result);
+		InputFrame.RegisteredResource = nullptr;
+	}
 
 	InputFrame.BackBuffer.SafeRelease();
+	if (InputFrame.SharedBackBuffer)
+	{
+		InputFrame.SharedBackBuffer->Release();
+		InputFrame.SharedBackBuffer = nullptr;
+	}
 }
 
 void FPixelStreamingNvVideoEncoder::ReleaseResources()
@@ -963,9 +993,12 @@ void FPixelStreamingNvVideoEncoder::ReleaseResources()
 
 		ReleaseFrameInputBuffer(Frame.InputFrame);
 
-		_NVENCSTATUS Result = NvEncodeAPI->nvEncDestroyBitstreamBuffer(EncoderInterface, Frame.OutputFrame.BitstreamBuffer);
-		checkf(NV_RESULT(Result), TEXT("Failed to destroy output buffer bitstream (status: %d)"), Result);
-		Frame.OutputFrame.BitstreamBuffer = nullptr;
+		if (Frame.OutputFrame.BitstreamBuffer)
+		{
+			_NVENCSTATUS Result = NvEncodeAPI->nvEncDestroyBitstreamBuffer(EncoderInterface, Frame.OutputFrame.BitstreamBuffer);
+			checkf(NV_RESULT(Result), TEXT("Failed to destroy output buffer bitstream (status: %d)"), Result);
+			Frame.OutputFrame.BitstreamBuffer = nullptr;
+		}
 
 		if (Frame.OutputFrame.EventHandle)
 		{
diff --git a/Engine/Plugins/Experimental/PixelStreaming/Source/PixelStreaming/Private/Codecs/PixelStreamingNvVideoEncoder.h b/Engine/Plugins/Experimental/PixelStreaming/Source/PixelStreaming/Private/Codecs/PixelStreamingNvVideoEncoder.h
--- a/Engine/Plugins/Experimental/PixelStreaming/Source/PixelStreaming/Private/Codecs/PixelStreamingNvVideoEncoder.h
+++ b/Engine/Plugins/Experimental/PixelStreaming/Source/PixelStreaming/Private/Codecs/PixelStreamingNvVideoEncoder.h
@@ -38,10 +38,7 @@ public:
 	 */
 	static bool CheckPlatformCompatibility();
 
-	/**
-	* Note bEnableAsyncMode flag is for debugging purpose, it should be set to true normally unless user wants to test in synchronous mode.
-	*/
-	explicit FPixelStreamingNvVideoEncoder(bool bEnableAsyncMode = true);
+	FPixelStreamingNvVideoEncoder();
 	~FPixelStreamingNvVideoEncoder();
 
 	bool CopyBackBuffer(const FTexture2DRHIRef& BackBuffer, FTimespan Timestamp, FBufferId& BufferId) override;
@@ -64,6 +61,7 @@ private:
 		NV_ENC_INPUT_PTR MappedResource = nullptr;
 		NV_ENC_BUFFER_FORMAT BufferFormat;
 		FTexture2DRHIRef BackBuffer;
+		ID3D11Texture2D* SharedBackBuffer = nullptr;
 		FTimespan CaptureTs;
 	};
 
@@ -74,14 +72,26 @@ private:
 		webrtc::EncodedImage EncodedFrame;
 	};
 
+	enum class EFrameState
+	{
+		Free,
+		Captured,
+		Encoding
+	};
+
 	struct FFrame
 	{
+		const FBufferId Id = 0;
+		EFrameState State = EFrameState::Free;
+		// Bitrate requested at the time the video encoder asked us to encode this frame
+		// We save this, because we can't use it at the moment we receive it.
+		uint32 BitrateRequested = 0;
 		FInputFrame InputFrame;
 		FOutputFrame OutputFrame;
 		uint64 FrameIdx = 0;
 	};
 
-	void Init_RenderThread(bool bEnableAsyncMode);
+	void Init();
 	void InitFrameInputBuffer(FInputFrame& InputFrame, uint32 Width, uint32 Heigh);
 	void InitializeResources();
 	void ReleaseFrameInputBuffer(FInputFrame& InputFrame);
@@ -94,42 +104,15 @@ private:
 	void UpdateRes(const FTexture2DRHIRef& BackBuffer, FInputFrame& InputFrame);
 	void CopyBackBuffer(const FTexture2DRHIRef& BackBuffer, FInputFrame& InputFrame);
 
-	void EncodeFrameInRenderingThread(FFrame& Frame, uint32 Bitrate);
-
 	void EncoderCheckLoop();
 
 	void ProcessFrame(FFrame& Frame);
 
 	void UpdateSettings(FInputFrame& InputFrame, uint32 Bitrate);
-	void TransferRenderTargetToHWEncoder(FFrame& Frame);
-
-	void PostRenderingThreadCreated()
-	{
-		bWaitForRenderThreadToResume = false;
-	}
-
-	void PreRenderingThreadDestroyed()
-	{
-		bWaitForRenderThreadToResume = true;
-	}
+	void SubmitFrameToEncoder(FFrame& Frame);
 
 	void OnEncodedFrame(const webrtc::EncodedImage& EncodedImage);
 
-	FRHICOMMAND_MACRO(FRHITransferRenderTargetToNvEnc)
-	{
-		FPixelStreamingNvVideoEncoder* Encoder;
-		FFrame* Frame;
-
-		FRHITransferRenderTargetToNvEnc(FPixelStreamingNvVideoEncoder* InEncoder, FFrame* InFrame)
-			: Encoder(InEncoder), Frame(InFrame)
-		{}
-
-		void Execute(FRHICommandListBase& CmdList)
-		{
-			Encoder->TransferRenderTargetToHWEncoder(*Frame);
-		}
-	};
-
 	void* DllHandle = nullptr;
 
 	TUniquePtr<NV_ENCODE_API_FUNCTION_LIST> NvEncodeAPI;
@@ -137,9 +120,6 @@ private:
 	NV_ENC_INITIALIZE_PARAMS NvEncInitializeParams;
 	NV_ENC_CONFIG NvEncConfig;
 	FThreadSafeBool bWaitForRenderThreadToResume = false;
-	// Used to make sure we don't have a race condition trying to access a deleted "this" captured
-	// in the render command lambda sent to the render thread from EncoderCheckLoop
-	static FThreadSafeCounter ImplCounter;
 	uint32 CapturedFrameCount = 0; // of captured, not encoded frames
 	static const uint32 NumBufferedFrames = 3;
 	FFrame BufferedFrames[NumBufferedFrames];
@@ -149,19 +129,79 @@ private:
 	// the memory
 	TArray<uint8> EncodedFrameBuffer;
 
-	// collaboration with WebRTC is quite convoluted. We capture frame and pass it to WebRTC
-	// TODO(andriy): implement proper waiting on empty queues instead of spin-lock
-	TQueue<FBufferId> FreeBuffers;
-	TQueue<FBufferId> CapturedBuffers;
-	TQueue<FBufferId> BuffersBeingEncoded;
-
 	float InitialMaxFPS;
 
-	// #AMF(Andriy) : This is only used from one thread, I think, so the comment below is wrong, and it doesn't need to be TAtomic. Need to confirm.
-	//		  cached value because it's used from another thread.
-	//		  Update : If running with RHITRhead (e.g: use r.rhithread.enable 1 ), this is indeed used from two threads (RenderThread and RHIThread, I think)
-	TAtomic<double> RequestedBitrateMbps{ 0 };
+	class FEncoderDevice
+	{
+	public:
+		FEncoderDevice();
+
+		TRefCountPtr<ID3D11Device> Device;
+		TRefCountPtr<ID3D11DeviceContext> DeviceContext;
+	};
+
+	TUniquePtr<FEncoderDevice> EncoderDevice;
+
+	// After a back buffer is processed and copied then we will want to send it
+	// to the encoder. This happens on a different thread so we use a queue of
+	// frame pointers to tell the thread which frames should be encoded.
+	struct FEncodeQueue
+	{
+		// The frames which we should encode.
+		// We can never be encoding more frames that can be buffered.
+		FFrame* Frames[NumBufferedFrames] = { nullptr };
+
+		// The start position of elements in this FIFO ring buffer queue.
+		int Start = 0;
+
+		// The number of elements in this FIFO ring buffer queue.
+		int Length = 0;
+
+		// Allow access by the Render Thread and the Pixel Streaming Encoder
+		// thread.
+		FCriticalSection CriticalSection;
+
+		// An event to signal the Pixel Streaming Encoder thread that it can
+		// encode some frames.
+		HANDLE EncodeEvent = CreateEvent(nullptr, false, false, nullptr);
+
+		virtual ~FEncodeQueue()
+		{
+			CloseHandle(EncodeEvent);
+		}
+
+		// Add another frame to be encoded.
+		void Push(FFrame* Frame)
+		{
+			FScopeLock ScopedLock(&CriticalSection);
+			check(Length < NumBufferedFrames);
+			bool bWasEmpty = Length == 0;
+			int Position = (Start + Length) % NumBufferedFrames;
+			Frames[Position] = Frame;
+			Length++;
+			if (bWasEmpty)
+			{
+				SetEvent(EncodeEvent);
+			}
+		}
+
+		// Get the list of all frame which we should encode.
+		void PopAll(FFrame* OutFrames[NumBufferedFrames], int& OutNumFrames)
+		{
+			FScopeLock ScopedLock(&CriticalSection);
+			OutNumFrames = Length;
+			for (int Position = 0; Position < Length; Position++)
+			{
+				OutFrames[Position] = Frames[Start];
+				Start = (Start + 1) % NumBufferedFrames;
+			}
+			Length = 0;
+			ResetEvent(EncodeEvent);
+		}
+	};
 
+	FEncodeQueue EncodeQueue;
+	double RequestedBitrateMbps = 0;
 	FCriticalSection SubscribersMutex;
 	TSet<FVideoEncoder*> Subscribers;
 };
diff --git a/Engine/Plugins/Experimental/PixelStreaming/Source/PixelStreaming/Private/Codecs/WmfIncludes.h b/Engine/Plugins/Experimental/PixelStreaming/Source/PixelStreaming/Private/Codecs/WmfIncludes.h
--- a/Engine/Plugins/Experimental/PixelStreaming/Source/PixelStreaming/Private/Codecs/WmfIncludes.h
+++ b/Engine/Plugins/Experimental/PixelStreaming/Source/PixelStreaming/Private/Codecs/WmfIncludes.h
@@ -80,6 +80,7 @@ inline const FString GetComErrorDescription(HRESULT Res)
 		}\
 	}
 
+// #AMF(Andriy) : Rename these to DX , instead of DX9 ?
 #define CHECK_HR_DX9(DX9_call)\
 	{\
 		HRESULT Res = DX9_call;\
@@ -89,3 +90,14 @@ inline const FString GetComErrorDescription(HRESULT Res)
 			return false;\
 		}\
 	}
+
+#define CHECK_HR_DX9_VOID(DX9_call)\
+	{\
+		HRESULT Res = DX9_call;\
+		if (FAILED(Res))\
+		{\
+			UE_LOG(PixelPlayer, Error, TEXT("`" #DX9_call "` failed 0x%X: %s - %s"), Res, DXGetErrorString(Res), DXGetErrorDescription(Res));\
+			return;\
+		}\
+	}
+
