Fix up ELLMTag::VideoStreaming to MediaStreaming in light of CL# 9739273
#rnx

diff --git a/Engine/Plugins/Experimental/PixelStreaming/Source/PixelStreaming/Private/Codecs/VideoDecoder.cpp b/Engine/Plugins/Experimental/PixelStreaming/Source/PixelStreaming/Private/Codecs/VideoDecoder.cpp
--- a/Engine/Plugins/Experimental/PixelStreaming/Source/PixelStreaming/Private/Codecs/VideoDecoder.cpp
+++ b/Engine/Plugins/Experimental/PixelStreaming/Source/PixelStreaming/Private/Codecs/VideoDecoder.cpp
@@ -358,7 +358,7 @@ bool FVideoDecoder::QueueBuffer(const webrtc::EncodedImage& InputImage, bool Mis
 
 void FVideoDecoder::DecodeThreadFunc()
 {
-	LLM_SCOPE(ELLMTag::VideoStreaming);
+	LLM_SCOPE(ELLMTag::MediaStreaming);
 	// first checks if decoder has output and only if it asks for more input the input is provided
 	// this way we work around decoder hanging if all samples from its internal pool are in use (h/w decoder)
 
