Updating PixelStreaming version number to 0.2 to indicate the larger change of moving the WebRTC Proxy to be part of the plugin in the engine

diff --git a/Engine/Plugins/Experimental/PixelStreaming/PixelStreaming.uplugin b/Engine/Plugins/Experimental/PixelStreaming/PixelStreaming.uplugin
--- a/Engine/Plugins/Experimental/PixelStreaming/PixelStreaming.uplugin
+++ b/Engine/Plugins/Experimental/PixelStreaming/PixelStreaming.uplugin
@@ -1,7 +1,7 @@
 {
 	"FileVersion": 3,
-	"Version": 0.1,
-	"VersionName": "0.1",
+	"Version": 0.2,
+	"VersionName": "0.2",
 	"FriendlyName": "Pixel Streaming",
 	"Description": "Streaming of UE4 render target to WebRTC-compatible media player and playback of such streaming",
 	"Category": "Graphics",
