Merging CL 9768360 //UE4/Dev-CharacterTech -> //UE4/Release-4.24

checking for pixel streaming explicitly enabled in editor and showing error message box that it's not supported. fixed AudioCapturer to not depend on `UGameEngine`

diff --git a/Engine/Plugins/Experimental/PixelStreaming/Source/PixelStreaming/Private/AudioCapturer.cpp b/Engine/Plugins/Experimental/PixelStreaming/Source/PixelStreaming/Private/AudioCapturer.cpp
--- a/Engine/Plugins/Experimental/PixelStreaming/Source/PixelStreaming/Private/AudioCapturer.cpp
+++ b/Engine/Plugins/Experimental/PixelStreaming/Source/PixelStreaming/Private/AudioCapturer.cpp
@@ -109,13 +109,12 @@ int32 FAudioCapturer::Init()
 	}
 
 	// subscribe to audio data
-	UGameEngine* GameEngine = Cast<UGameEngine>(GEngine);
-	if (!GameEngine)
+	if (!GEngine)
 	{
 		return -1;
 	}
 
-	FAudioDevice* AudioDevice = GameEngine->GetMainAudioDevice();
+	FAudioDevice* AudioDevice = GEngine->GetMainAudioDevice();
 	if (!AudioDevice)
 	{
 		return -1;
@@ -135,13 +134,12 @@ int32 FAudioCapturer::Terminate()
 		return 0;
 
 	// unsubscribe from audio data
-	UGameEngine* GameEngine = Cast<UGameEngine>(GEngine);
-	if (!GameEngine)
+	if (!GEngine)
 	{
 		return -1;
 	}
 
-	FAudioDevice* AudioDevice = GameEngine->GetMainAudioDevice();
+	FAudioDevice* AudioDevice = GEngine->GetMainAudioDevice();
 	if (!AudioDevice)
 	{
 		return -1;
diff --git a/Engine/Plugins/Experimental/PixelStreaming/Source/PixelStreaming/Private/InputDevice.cpp b/Engine/Plugins/Experimental/PixelStreaming/Source/PixelStreaming/Private/InputDevice.cpp
--- a/Engine/Plugins/Experimental/PixelStreaming/Source/PixelStreaming/Private/InputDevice.cpp
+++ b/Engine/Plugins/Experimental/PixelStreaming/Source/PixelStreaming/Private/InputDevice.cpp
@@ -466,6 +466,7 @@ namespace
 	// viewport to have a different size.
 	void UnquantizeAndDenormalize(uint16& InOutX, uint16& InOutY)
 	{
+		checkf(GEngine->GameViewport != nullptr, TEXT("PixelStreaming is not supported in editor, and should have been disabled earlier"));
 		FIntPoint SizeXY = GEngine->GameViewport->Viewport->GetSizeXY();
 		InOutX = InOutX / 65536.0f * SizeXY.X;
 		InOutY = InOutY / 65536.0f * SizeXY.Y;
diff --git a/Engine/Plugins/Experimental/PixelStreaming/Source/PixelStreaming/Private/PixelStreamingModule.cpp b/Engine/Plugins/Experimental/PixelStreaming/Source/PixelStreaming/Private/PixelStreamingModule.cpp
--- a/Engine/Plugins/Experimental/PixelStreaming/Source/PixelStreaming/Private/PixelStreamingModule.cpp
+++ b/Engine/Plugins/Experimental/PixelStreaming/Source/PixelStreaming/Private/PixelStreamingModule.cpp
@@ -15,8 +15,6 @@
 #include "CoreMinimal.h"
 #include "Modules/ModuleManager.h"
 #include "UObject/UObjectIterator.h"
-#include "Engine/GameEngine.h"
-#include "Engine/GameViewportClient.h"
 #include "Engine/Texture2D.h"
 #include "Slate/SceneViewport.h"
 #include "Windows/WindowsHWrapper.h"
@@ -34,6 +32,7 @@
 #include "IImageWrapper.h"
 #include "IImageWrapperModule.h"
 #include "Async/Async.h"
+#include "Engine/Engine.h"
 
 #if !UE_BUILD_SHIPPING
 #	include "DrawDebugHelpers.h"
@@ -84,22 +83,34 @@ void FPixelStreamingModule::InitStreamer()
 		return;
 	}
 
+	if (GIsEditor)
+	{
+		FText TitleText = FText::FromString(TEXT("Pixel Streaming Plugin"));
+		FString ErrorString = TEXT("Pixel Streaming Plugin is not supported in editor, but it was explicitly enabled by command-line arguments. Please remove `PixelStreamingIP` and `PixelStreamingPort` args from editor command line.");
+		FText ErrorText = FText::FromString(ErrorString);
+		FMessageDialog::Open(EAppMsgType::Ok, ErrorText, &TitleText);
+		UE_LOG(PixelStreamer, Error, TEXT("%s"), *ErrorString);
+		return;
+	}
+
 	// Check to see if we can use the Pixel Streaming plugin on this platform.
 	// If not then we avoid setting up our delegates to prevent access to the
 	// plugin. Note that Pixel Streaming is not currently performed in the
 	// Editor.
-	if (!GIsEditor && !CheckPlatformCompatibility())
+	if (!CheckPlatformCompatibility())
+	{
+		return;
+	}
+
+	if (!ensure(GEngine != nullptr))
 	{
 		return;
 	}
 
 	// subscribe to engine delegates here for init / framebuffer creation / whatever
-	if (UGameEngine* GameEngine = Cast<UGameEngine>(GEngine))
+	if (FSlateApplication::IsInitialized())
 	{
-		if (FSlateApplication::IsInitialized())
-		{
-			FSlateApplication::Get().GetRenderer()->OnBackBufferReadyToPresent().AddRaw(this, &FPixelStreamingModule::OnBackBufferReady_RenderThread);
-		}
+		FSlateApplication::Get().GetRenderer()->OnBackBufferReadyToPresent().AddRaw(this, &FPixelStreamingModule::OnBackBufferReady_RenderThread);
 	}
 
 	FGameModeEvents::GameModePostLoginEvent.AddRaw(this, &FPixelStreamingModule::OnGameModePostLogin);
