PixelStreaming when looking at the Unreal App through RDP causes glitchy stream

diff --git a/Engine/Plugins/Experimental/PixelStreaming/Source/PixelStreaming/Private/Codecs/PixelStreamingNvVideoEncoder.cpp b/Engine/Plugins/Experimental/PixelStreaming/Source/PixelStreaming/Private/Codecs/PixelStreamingNvVideoEncoder.cpp
--- a/Engine/Plugins/Experimental/PixelStreaming/Source/PixelStreaming/Private/Codecs/PixelStreamingNvVideoEncoder.cpp
+++ b/Engine/Plugins/Experimental/PixelStreaming/Source/PixelStreaming/Private/Codecs/PixelStreamingNvVideoEncoder.cpp
@@ -31,6 +31,9 @@ THIRD_PARTY_INCLUDES_START
 	#include "Windows/HideWindowsPlatformTypes.h"
 THIRD_PARTY_INCLUDES_END
 
+#if NVENC_VIDEO_ENCODER_DEBUG
+	#include "ClearQuad.h"
+#endif
 
 TAutoConsoleVariable<float> CVarEncoderMaxBitrate(
 	TEXT("Encoder.MaxBitrate"),
@@ -191,6 +194,7 @@ namespace
 		return bAllowAsyncResourceCreation;
 	}
 
+
 }
 
 FPixelStreamingNvVideoEncoder::FEncoderDevice::FEncoderDevice()
@@ -440,8 +444,7 @@ void FPixelStreamingNvVideoEncoder::Init()
 	FBufferId Id = 0;
 	for(FFrame& Frame : BufferedFrames)
 	{
-		// We keep `Id` as const, because it not supposed to change at all once initialized
-		*(const_cast<FBufferId*>(&Frame.Id)) = Id++;
+		Frame.Id = Id++;
 	}
 
 	InitializeResources();
@@ -491,12 +494,13 @@ FPixelStreamingNvVideoEncoder::~FPixelStreamingNvVideoEncoder()
 
 bool FPixelStreamingNvVideoEncoder::CopyBackBuffer(const FTexture2DRHIRef& BackBuffer, FTimespan Timestamp, FBufferId& BufferId)
 {
+	check(IsInRenderingThread());
 
 	// Find a free slot we can use
 	FFrame* Frame = nullptr;
 	for (FFrame& Slot : BufferedFrames)
 	{
-		if (Slot.State == EFrameState::Free)
+		if (Slot.State.Load() == EFrameState::Free)
 		{
 			Frame = &Slot;
 			BufferId = Slot.Id;
@@ -510,25 +514,36 @@ bool FPixelStreamingNvVideoEncoder::CopyBackBuffer(const FTexture2DRHIRef& BackB
 		return false;
 	}
 
+
 	Frame->FrameIdx = CapturedFrameCount++;
 	Frame->InputFrame.CaptureTs = Timestamp;
-	CopyBackBuffer(BackBuffer, Frame->InputFrame);
+
+#if NVENC_VIDEO_ENCODER_DEBUG
+	Frame->CopyBufferStartTs = FTimespan::FromSeconds(FPlatformTime::Seconds());
+	// By clearing the frame at this point, we can catch the occasional glimpse of a solid color
+	// frame in PixelStreaming if there are any bugs detecting when the copy finished
+	ClearFrame(*Frame);
+#endif
+
+	CopyBackBuffer(BackBuffer, *Frame);
 
 	UE_LOG(LogVideoEncoder, Verbose, TEXT("Buffer #%d (%d) captured"), Frame->FrameIdx, BufferId);
-	Frame->State = EFrameState::Captured;
+	Frame->State = EFrameState::Capturing;
 
 	return true;
 }
 
-void FPixelStreamingNvVideoEncoder::CopyBackBuffer(const FTexture2DRHIRef& BackBuffer, FInputFrame& InputFrame)
+void FPixelStreamingNvVideoEncoder::CopyBackBuffer(const FTexture2DRHIRef& BackBuffer, FFrame& Frame)
 {
 	SCOPE_CYCLE_COUNTER(STAT_NvEnc_CopyBackBuffer);
+	FInputFrame& InputFrame = Frame.InputFrame;
 
-	UpdateRes(BackBuffer, InputFrame);
+	UpdateRes(BackBuffer, Frame);
 
 	IRendererModule* RendererModule = &FModuleManager::GetModuleChecked<IRendererModule>("Renderer");
 	FRHICommandListImmediate& RHICmdList = FRHICommandListExecutor::GetImmediateCommandList();
 
+	InputFrame.CopyFence->Clear();
 	if (BackBuffer->GetFormat() == InputFrame.BackBuffer->GetFormat() &&
 		BackBuffer->GetSizeXY() == InputFrame.BackBuffer->GetSizeXY())
 	{
@@ -582,29 +597,68 @@ void FPixelStreamingNvVideoEncoder::CopyBackBuffer(const FTexture2DRHIRef& BackB
 		}
 		RHICmdList.EndRenderPass();
 	}
+
+	RHICmdList.WriteGPUFence(Frame.InputFrame.CopyFence);
 }
 
+
 void FPixelStreamingNvVideoEncoder::EncodeFrame(FBufferId BufferId, const webrtc::EncodedImage& EncodedFrame, uint32 Bitrate)
 {
 	FFrame& Frame = BufferedFrames[BufferId];
-	checkf(Frame.State==EFrameState::Captured, TEXT("Buffer %d : Expected state %d, but found %d"), BufferId, (int)EFrameState::Captured, (int)Frame.State);
 
+	{
+		EFrameState State = Frame.State.Load();
+		checkf(State == EFrameState::Capturing, TEXT("Buffer %d : Expected state %d, but found %d"), BufferId, (int)EFrameState::Captured, (int)State);
+	}
+
+	// Loop and sleep until the fence is signaled.
+	// Also, at the moment of writing there is not proper GPU fence for D3D11 in UE4. It uses FGenericRHIGPUFence.
+	// Due to this, if this thread doesn't make progress because of the fence not being signaled, it can end up stalling the RenderThread
+	// when restarting a PixelStreaming client session (e.g: Refreshing the browser page).
+	// This causes a deadlock, because the RenderThread is the one signaling the fence.
+	// So, if we are waiting for too long, just ignore the fence.
+	{
+		double StartTime = FPlatformTime::Seconds();
+		while (!Frame.InputFrame.CopyFence->Poll())
+		{
+			FPlatformProcess::Sleep(2.0f/1000);
+			if ((FPlatformTime::Seconds() - StartTime) > 0.250f)
+			{
+				UE_LOG(LogVideoEncoder, Warning, TEXT("Buffer #%d taking too long to reach the GPU fence. Ignoring fence."), (int)BufferId);
+				break;
+			}
+		}
+	}
+
+#if NVENC_VIDEO_ENCODER_DEBUG
+	Frame.CopyBufferFinishTs = FTimespan::FromSeconds(FPlatformTime::Seconds());
+#endif
+
+	Frame.State = EFrameState::Captured;
 	Frame.OutputFrame.EncodedFrame = EncodedFrame;
 	Frame.OutputFrame.EncodedFrame._encodedWidth = NvEncInitializeParams.encodeWidth;
 	Frame.OutputFrame.EncodedFrame._encodedHeight = NvEncInitializeParams.encodeHeight;
 
-	Frame.State = EFrameState::Encoding;
 	// Save the requested bitrate, so we can reconfigure the encoder later
 	Frame.BitrateRequested = Bitrate;
 	EncodeQueue.Push(&Frame);
-	UE_LOG(LogVideoEncoder, VeryVerbose, TEXT("Buffer #%d (%d), ts %u started encoding"), Frame.FrameIdx, BufferId, Frame.OutputFrame.EncodedFrame.Timestamp());
+
+	UE_LOG(LogVideoEncoder, VeryVerbose, TEXT("Buffer #%d (%d), ts %u sent to the encoder thread"), Frame.FrameIdx, BufferId, Frame.OutputFrame.EncodedFrame.Timestamp());
 
 }
 
 void FPixelStreamingNvVideoEncoder::OnFrameDropped(FBufferId BufferId)
 {
-	checkf(BufferedFrames[BufferId].State==EFrameState::Captured, TEXT("Buffer %d: Expected state %d, found %d"), BufferId, (int)EFrameState::Captured, (int)BufferedFrames[BufferId].State);
-	BufferedFrames[BufferId].State = EFrameState::Free;
+	FFrame& Frame = BufferedFrames[BufferId];
+
+	{
+		EFrameState State = Frame.State.Load();
+		checkf(State == EFrameState::Capturing, TEXT("Buffer %d: Expected state %d, found %d")
+			, BufferId, (int)EFrameState::Capturing, (int)State);
+	}
+
+	Frame.State = EFrameState::Free;
+
 	UE_LOG(LogVideoEncoder, Log, TEXT("Buffer #%d (%d) dropped"), BufferedFrames[BufferId].FrameIdx, BufferId);
 }
 
@@ -625,8 +679,6 @@ void FPixelStreamingNvVideoEncoder::UnsubscribeFromFrameEncodedEvent(FVideoEncod
 
 void FPixelStreamingNvVideoEncoder::UpdateNvEncConfig(const FInputFrame& InputFrame, uint32 Bitrate)
 {
-	check(!IsInRenderingThread());
-
 	bool bSettingsChanged = false;
 	bool bResolutionChanged = false;
 
@@ -733,10 +785,12 @@ bool FPixelStreamingNvVideoEncoder::UpdateFramerate()
 }
 
 // checks if resolution changed, either the game res changed or new streaming resolution was specified by the console var
-void FPixelStreamingNvVideoEncoder::UpdateRes(const FTexture2DRHIRef& BackBuffer, FInputFrame& InputFrame)
+void FPixelStreamingNvVideoEncoder::UpdateRes(const FTexture2DRHIRef& BackBuffer, FFrame& Frame)
 {
 	check(IsInRenderingThread());
 
+	FInputFrame& InputFrame = Frame.InputFrame;
+
 	// find out what resolution we'd like to stream, it's either "native" (BackBuffer) resolution or something configured specially
 	bool bUseBackBufferSize = CVarEncoderUseBackBufferSize.GetValueOnRenderThread() > 0;
 	uint32 Width;
@@ -771,15 +825,21 @@ void FPixelStreamingNvVideoEncoder::UpdateRes(const FTexture2DRHIRef& BackBuffer
 	}
 
 	// reallocate and re-register InputFrame with NvEnc
-	ReleaseFrameInputBuffer(InputFrame);
-	InitFrameInputBuffer(InputFrame, Width, Height);
+	ReleaseFrameInputBuffer(Frame);
+	InitFrameInputBuffer(Frame, Width, Height);
 }
 
 void FPixelStreamingNvVideoEncoder::SubmitFrameToEncoder(FFrame& Frame)
 {
+	check(Frame.State.Load() == EFrameState::Captured);
+
+#if NVENC_VIDEO_ENCODER_DEBUG
+	Frame.EncodingStartTs = FTimespan::FromSeconds(FPlatformTime::Seconds());
+#endif
+
 	SCOPE_CYCLE_COUNTER(STAT_NvEnc_SendBackBufferToEncoder);
-	check(Frame.State==EFrameState::Encoding);
 
+	Frame.State = EFrameState::Encoding;
 	Frame.OutputFrame.EncodedFrame.timing_.encode_start_ms = rtc::TimeMicros() / 1000;
 
 	NV_ENC_PIC_PARAMS PicParams;
@@ -867,7 +927,7 @@ void FPixelStreamingNvVideoEncoder::EncoderCheckLoop()
 
 void FPixelStreamingNvVideoEncoder::ProcessFrame(FFrame& Frame)
 {
-	check(Frame.State==EFrameState::Encoding);
+	check(Frame.State.Load() == EFrameState::Encoding);
 
 	FOutputFrame& OutputFrame = Frame.OutputFrame;
 
@@ -921,6 +981,26 @@ void FPixelStreamingNvVideoEncoder::ProcessFrame(FFrame& Frame)
 		Stats.EncoderQP.Update(OutputFrame.EncodedFrame.qp_);
 	}
 
+#if NVENC_VIDEO_ENCODER_DEBUG
+	Frame.EncodingFinishTs = FTimespan::FromSeconds(FPlatformTime::Seconds());
+	{
+		FFrameTiming Timing;
+		Timing.Total[0] = (Frame.CopyBufferFinishTs - Frame.CopyBufferStartTs).GetTotalMilliseconds();
+		Timing.Total[1] = (Frame.EncodingStartTs - Frame.CopyBufferStartTs).GetTotalMilliseconds();
+		Timing.Total[2] = (Frame.EncodingFinishTs - Frame.CopyBufferStartTs).GetTotalMilliseconds();
+
+		Timing.Steps[0] = (Frame.CopyBufferFinishTs - Frame.CopyBufferStartTs).GetTotalMilliseconds();
+		Timing.Steps[1] = (Frame.EncodingStartTs - Frame.CopyBufferFinishTs).GetTotalMilliseconds();
+		Timing.Steps[2] = (Frame.EncodingFinishTs - Frame.EncodingStartTs).GetTotalMilliseconds();
+		Timings.Add(Timing);
+		// Limit the array size
+		if (Timings.Num()>1000)
+		{
+			Timings.RemoveAt(0);
+		}
+	}
+#endif
+
 	UE_LOG(LogVideoEncoder, VeryVerbose, TEXT("encoded %s ts %u, capture ts %lld, QP %d/%.0f,  latency %.0f/%.0f ms, bitrate %.3f/%.3f/%.3f Mbps, %zu bytes"), ToString(OutputFrame.EncodedFrame._frameType), OutputFrame.EncodedFrame.Timestamp(), Frame.InputFrame.CaptureTs.GetTicks(), OutputFrame.EncodedFrame.qp_, Stats.EncoderQP.Get(), LatencyMs, Stats.EncoderLatencyMs.Get(), RequestedBitrateMbps, BitrateMbps, Stats.EncoderBitrateMbps.Get(), OutputFrame.EncodedFrame._length);
 
 	// Stream the encoded frame
@@ -932,8 +1012,17 @@ void FPixelStreamingNvVideoEncoder::ProcessFrame(FFrame& Frame)
 	Frame.State = EFrameState::Free;
 }
 
-void FPixelStreamingNvVideoEncoder::InitFrameInputBuffer(FInputFrame& InputFrame, uint32 Width, uint32 Height)
+void FPixelStreamingNvVideoEncoder::InitFrameInputBuffer(FFrame& Frame, uint32 Width, uint32 Height)
 {
+	FInputFrame& InputFrame = Frame.InputFrame;
+
+	// Create (if necessary) and clear the GPU Fence so we can detect when the copy finished
+	if (!InputFrame.CopyFence)
+	{
+		FRHICommandListImmediate& RHICmdList = FRHICommandListExecutor::GetImmediateCommandList();
+		InputFrame.CopyFence = RHICmdList.CreateGPUFence(*FString::Printf(TEXT("PixelStreamingCopy_%d"), Frame.Id));
+	}
+
 	// Create resolved back buffer texture
 	{
 		// Make sure format used here is compatible with NV_ENC_BUFFER_FORMAT specified later in NV_ENC_REGISTER_RESOURCE bufferFormat
@@ -1013,7 +1102,7 @@ void FPixelStreamingNvVideoEncoder::InitializeResources()
 	{
 		FFrame& Frame = BufferedFrames[i];
 
-		InitFrameInputBuffer(Frame.InputFrame, NvEncInitializeParams.encodeWidth, NvEncInitializeParams.encodeHeight);
+		InitFrameInputBuffer(Frame, NvEncInitializeParams.encodeWidth, NvEncInitializeParams.encodeHeight);
 
 		FMemory::Memzero(Frame.OutputFrame);
 		// Create output bitstream buffer
@@ -1032,8 +1121,10 @@ void FPixelStreamingNvVideoEncoder::InitializeResources()
 	}
 }
 
-void FPixelStreamingNvVideoEncoder::ReleaseFrameInputBuffer(FInputFrame& InputFrame)
+void FPixelStreamingNvVideoEncoder::ReleaseFrameInputBuffer(FFrame& Frame)
 {
+	FInputFrame& InputFrame = Frame.InputFrame;
+
 	if (InputFrame.MappedResource)
 	{
 		_NVENCSTATUS Result = NvEncodeAPI->nvEncUnmapInputResource(EncoderInterface, InputFrame.MappedResource);
@@ -1054,6 +1145,11 @@ void FPixelStreamingNvVideoEncoder::ReleaseFrameInputBuffer(FInputFrame& InputFr
 		InputFrame.SharedBackBuffer->Release();
 		InputFrame.SharedBackBuffer = nullptr;
 	}
+
+	if (InputFrame.CopyFence)
+	{
+		InputFrame.CopyFence.SafeRelease();
+	}
 }
 
 void FPixelStreamingNvVideoEncoder::ReleaseResources()
@@ -1062,7 +1158,7 @@ void FPixelStreamingNvVideoEncoder::ReleaseResources()
 	{
 		FFrame& Frame = BufferedFrames[i];
 
-		ReleaseFrameInputBuffer(Frame.InputFrame);
+		ReleaseFrameInputBuffer(Frame);
 
 		if (Frame.OutputFrame.BitstreamBuffer)
 		{
@@ -1115,3 +1211,31 @@ void FPixelStreamingNvVideoEncoder::OnEncodedFrame(const webrtc::EncodedImage& E
 	}
 }
 
+#if NVENC_VIDEO_ENCODER_DEBUG
+// Fills with a solid colour
+void FPixelStreamingNvVideoEncoder::ClearFrame(FFrame& Frame)
+{
+	check(IsInRenderingThread());
+
+	FRHICommandListImmediate& RHICmdList = FRHICommandListExecutor::GetImmediateCommandList();
+
+	static_assert(NumBufferedFrames == 3, "Unexpected number of slots. Please update the array to match.");
+	FLinearColor Colors[NumBufferedFrames] =
+	{
+		FLinearColor(1,0,0),
+		FLinearColor(0,1,0),
+		FLinearColor(0,0,1)
+	};
+
+	FRHIRenderPassInfo RPInfo(Frame.InputFrame.BackBuffer, ERenderTargetActions::Load_Store);
+	TransitionRenderPassTargets(RHICmdList, RPInfo);
+	RHICmdList.BeginRenderPass(RPInfo, TEXT("ClearCanvas"));
+	RHICmdList.SetViewport(0, 0, 0.0f
+		, Frame.InputFrame.BackBuffer->GetSizeXY().X
+		, Frame.InputFrame.BackBuffer->GetSizeXY().Y, 1.0f);
+
+	DrawClearQuad(RHICmdList, Colors[Frame.Id]);
+	RHICmdList.EndRenderPass();
+}
+#endif
+
diff --git a/Engine/Plugins/Experimental/PixelStreaming/Source/PixelStreaming/Private/Codecs/PixelStreamingNvVideoEncoder.h b/Engine/Plugins/Experimental/PixelStreaming/Source/PixelStreaming/Private/Codecs/PixelStreamingNvVideoEncoder.h
--- a/Engine/Plugins/Experimental/PixelStreaming/Source/PixelStreaming/Private/Codecs/PixelStreamingNvVideoEncoder.h
+++ b/Engine/Plugins/Experimental/PixelStreaming/Source/PixelStreaming/Private/Codecs/PixelStreamingNvVideoEncoder.h
@@ -17,6 +17,11 @@ THIRD_PARTY_INCLUDES_START
 	#include "NvEncoder/nvEncodeAPI.h"
 THIRD_PARTY_INCLUDES_END
 
+// This is mostly to use internally at Epic.
+// Setting this to 1 will collect detailed timings in the `Timings` member array.
+// It will also clear every frame with a solid colour before copying the backbuffer into it.
+#define NVENC_VIDEO_ENCODER_DEBUG 0
+
 class FVideoEncoder;
 class FThread;
 
@@ -57,16 +62,21 @@ private:
 
 	struct FInputFrame
 	{
+		FInputFrame() {}
+		UE_NONCOPYABLE(FInputFrame);
 		void* RegisteredResource = nullptr;
 		NV_ENC_INPUT_PTR MappedResource = nullptr;
 		NV_ENC_BUFFER_FORMAT BufferFormat;
 		FTexture2DRHIRef BackBuffer;
 		ID3D11Texture2D* SharedBackBuffer = nullptr;
 		FTimespan CaptureTs;
+		FGPUFenceRHIRef CopyFence;
 	};
 
 	struct FOutputFrame
 	{
+		FOutputFrame() {}
+		UE_NONCOPYABLE(FOutputFrame);
 		NV_ENC_OUTPUT_PTR BitstreamBuffer = nullptr;
 		HANDLE EventHandle = nullptr;
 		webrtc::EncodedImage EncodedFrame;
@@ -75,34 +85,48 @@ private:
 	enum class EFrameState
 	{
 		Free,
+		Capturing,
 		Captured,
 		Encoding
 	};
 
 	struct FFrame
 	{
-		const FBufferId Id = 0;
-		EFrameState State = EFrameState::Free;
+		FFrame() {}
+		UE_NONCOPYABLE(FFrame);
+
+		// Array index of this FFrame. This is set at startup, and should never be changed
+		FBufferId Id = 0;
+
+		TAtomic<EFrameState> State = EFrameState::Free;
 		// Bitrate requested at the time the video encoder asked us to encode this frame
 		// We save this, because we can't use it at the moment we receive it.
 		uint32 BitrateRequested = 0;
 		FInputFrame InputFrame;
 		FOutputFrame OutputFrame;
 		uint64 FrameIdx = 0;
+
+		// Some timestamps for debugging
+#if NVENC_VIDEO_ENCODER_DEBUG
+		FTimespan CopyBufferStartTs;
+		FTimespan CopyBufferFinishTs;
+		FTimespan EncodingStartTs;
+		FTimespan EncodingFinishTs;
+#endif
 	};
 
 	void Init();
-	void InitFrameInputBuffer(FInputFrame& InputFrame, uint32 Width, uint32 Heigh);
+	void InitFrameInputBuffer(FFrame& Frame, uint32 Width, uint32 Heigh);
 	void InitializeResources();
-	void ReleaseFrameInputBuffer(FInputFrame& InputFrame);
+	void ReleaseFrameInputBuffer(FFrame& Frame);
 	void ReleaseResources();
 	void RegisterAsyncEvent(void** OutEvent);
 	void UnregisterAsyncEvent(void* Event);
 
 	bool UpdateFramerate();
 	void UpdateNvEncConfig(const FInputFrame& InputFrame, uint32 Bitrate);
-	void UpdateRes(const FTexture2DRHIRef& BackBuffer, FInputFrame& InputFrame);
-	void CopyBackBuffer(const FTexture2DRHIRef& BackBuffer, FInputFrame& InputFrame);
+	void UpdateRes(const FTexture2DRHIRef& BackBuffer, FFrame& Frame);
+	void CopyBackBuffer(const FTexture2DRHIRef& BackBuffer, FFrame& Frame);
 
 	void EncoderCheckLoop();
 
@@ -121,7 +145,7 @@ private:
 	NV_ENC_CONFIG NvEncConfig;
 	FThreadSafeBool bWaitForRenderThreadToResume = false;
 	uint32 CapturedFrameCount = 0; // of captured, not encoded frames
-	static const uint32 NumBufferedFrames = 3;
+	static constexpr uint32 NumBufferedFrames = 3;
 	FFrame BufferedFrames[NumBufferedFrames];
 	TUniquePtr<FThread> EncoderThread;
 	FThreadSafeBool bExitEncoderThread = false;
@@ -200,6 +224,24 @@ private:
 		}
 	};
 
+#if NVENC_VIDEO_ENCODER_DEBUG
+	// This is just for debugging
+	void ClearFrame(FFrame& Frame);
+	// Timings in milliseconds. Just for debugging
+	struct FFrameTiming
+	{
+		// 0 : CopyBufferStart -> CopyBufferFinish
+		// 1 : CopyBufferStart -> EncodingStart
+		// 2 : CopyBufferStart -> EncodingFinish
+		double Total[3];
+		// 0 : CopyBufferStart -> CopyBufferFinish
+		// 1 : CopyBufferFinish -> EncodingStart
+		// 2 : EncodingStart -> EncodingFinish
+		double Steps[3];
+	};
+	TArray<FFrameTiming> Timings;
+#endif
+
 	FEncodeQueue EncodeQueue;
 	double RequestedBitrateMbps = 0;
 	FCriticalSection SubscribersMutex;
