Fix error: copying member subobject of type 'TAtomic<FPixelStreamingNvVideoEncoder::EFrameState>' invokes deleted constructor when compiling with Clang Win64
#rnx

diff --git a/Engine/Plugins/Experimental/PixelStreaming/Source/PixelStreaming/Private/Codecs/PixelStreamingNvVideoEncoder.h b/Engine/Plugins/Experimental/PixelStreaming/Source/PixelStreaming/Private/Codecs/PixelStreamingNvVideoEncoder.h
--- a/Engine/Plugins/Experimental/PixelStreaming/Source/PixelStreaming/Private/Codecs/PixelStreamingNvVideoEncoder.h
+++ b/Engine/Plugins/Experimental/PixelStreaming/Source/PixelStreaming/Private/Codecs/PixelStreamingNvVideoEncoder.h
@@ -98,7 +98,7 @@ private:
 		// Array index of this FFrame. This is set at startup, and should never be changed
 		FBufferId Id = 0;
 
-		TAtomic<EFrameState> State = EFrameState::Free;
+		TAtomic<EFrameState> State = { EFrameState::Free };
 		// Bitrate requested at the time the video encoder asked us to encode this frame
 		// We save this, because we can't use it at the moment we receive it.
 		uint32 BitrateRequested = 0;
