Fix a potential licensee build issue for a UE4 app which is using Pixel Streaming. Remove a Private include from Public header file.

#rnx

diff --git a/Engine/Plugins/Media/PixelStreaming/Source/PixelStreaming/Public/PixelStreamerInputComponent.h b/Engine/Plugins/Media/PixelStreaming/Source/PixelStreaming/Public/PixelStreamerInputComponent.h
--- a/Engine/Plugins/Media/PixelStreaming/Source/PixelStreaming/Public/PixelStreamerInputComponent.h
+++ b/Engine/Plugins/Media/PixelStreaming/Source/PixelStreaming/Public/PixelStreamerInputComponent.h
@@ -2,7 +2,6 @@
 
 #pragma once
 
-#include "PixelStreamingPrivate.h"
 #include "Components/ActorComponent.h"
 #include "PixelStreamerInputComponent.generated.h"
 
