Compile fix (missing header)

diff --git a/Engine/Plugins/Media/PixelStreaming/Source/PixelStreaming/Private/Codecs/VideoEncoder.cpp b/Engine/Plugins/Media/PixelStreaming/Source/PixelStreaming/Private/Codecs/VideoEncoder.cpp
--- a/Engine/Plugins/Media/PixelStreaming/Source/PixelStreaming/Private/Codecs/VideoEncoder.cpp
+++ b/Engine/Plugins/Media/PixelStreaming/Source/PixelStreaming/Private/Codecs/VideoEncoder.cpp
@@ -5,6 +5,7 @@
 #include "PlayerSession.h"
 #include "Utils.h"
 #include "HUDStats.h"
+#include "Async/Async.h"
 
 TAutoConsoleVariable<int32> CVarPixelStreamingEncoderMaxBitrate(
 	TEXT("PixelStreaming.Encoder.MaxBitrate"),
