Renamed CVarPixelStreamingEncoderPrioritiseQuality to CVarPixelStreamingEncoderPrioritizeQuality

diff --git a/Engine/Plugins/Media/PixelStreaming/Source/PixelStreaming/Private/Codecs/VideoEncoder.cpp b/Engine/Plugins/Media/PixelStreaming/Source/PixelStreaming/Private/Codecs/VideoEncoder.cpp
--- a/Engine/Plugins/Media/PixelStreaming/Source/PixelStreaming/Private/Codecs/VideoEncoder.cpp
+++ b/Engine/Plugins/Media/PixelStreaming/Source/PixelStreaming/Private/Codecs/VideoEncoder.cpp
@@ -25,8 +25,8 @@ TAutoConsoleVariable<FString> CVarPixelStreamingEncoderTargetSize(
 	TEXT("Encoder target size in format widthxheight"),
 	ECVF_Cheat);
 
-TAutoConsoleVariable<int32> CVarPixelStreamingEncoderPrioritiseQuality(
-	TEXT("PixelStreaming.Encoder.PrioritiseQuality"),
+TAutoConsoleVariable<int32> CVarPixelStreamingEncoderPrioritizeQuality(
+	TEXT("PixelStreaming.Encoder.PrioritizeQuality"),
 	0,
 	TEXT("Reduces framerate automatically on bitrate reduction to trade FPS/latency for video quality"),
 	ECVF_Cheat);
@@ -187,7 +187,7 @@ int32 FVideoEncoder::Encode(const webrtc::VideoFrame& Frame, const webrtc::Codec
 	// Adjust framerate if required, so we can keep quality without increasing bitrate
 	{
 		float Fps;
-		if (!CVarPixelStreamingEncoderPrioritiseQuality.GetValueOnAnyThread())
+		if (!CVarPixelStreamingEncoderPrioritizeQuality.GetValueOnAnyThread())
 		{
 			Fps = HWEncoderDetails.InitialMaxFPS;
 		}
