Integrating 11932716

PixelStreaming: added MinBitrate console variable, that makes PixelStreaming to ignore video encoder bitrate recommended by WebRTC if it drops below MinBitrate value. Can be useful to sacrifice interactivity for better video quality on low-bitrate connections. Disabled by default

diff --git a/Engine/Plugins/Media/PixelStreaming/Source/PixelStreaming/Private/Codecs/VideoEncoder.cpp b/Engine/Plugins/Media/PixelStreaming/Source/PixelStreaming/Private/Codecs/VideoEncoder.cpp
--- a/Engine/Plugins/Media/PixelStreaming/Source/PixelStreaming/Private/Codecs/VideoEncoder.cpp
+++ b/Engine/Plugins/Media/PixelStreaming/Source/PixelStreaming/Private/Codecs/VideoEncoder.cpp
@@ -7,10 +7,16 @@
 #include "HUDStats.h"
 #include "Async/Async.h"
 
+TAutoConsoleVariable<int32> CVarPixelStreamingEncoderMinBitrate(
+	TEXT("PixelStreaming.Encoder.MinBitrate"),
+	0,
+	TEXT("Min bitrate no matter what WebRTC says, in bps"),
+	ECVF_RenderThreadSafe);
+
 TAutoConsoleVariable<int32> CVarPixelStreamingEncoderMaxBitrate(
 	TEXT("PixelStreaming.Encoder.MaxBitrate"),
 	50000000,
-	TEXT("Max bitrate no matter what WebRTC says, in Bps"),
+	TEXT("Max bitrate no matter what WebRTC says, in bps"),
 	ECVF_RenderThreadSafe);
 
 TAutoConsoleVariable<int32> CVarPixelStreamingEncoderUseBackBufferSize(
@@ -251,7 +257,10 @@ int32 FVideoEncoder::Encode(const webrtc::VideoFrame& Frame, const webrtc::Codec
 
 	// TODO(andriy): `LastBitrate.get_sum_kbps()` most probably includes audio bitrate too,
 	// check if this causes any packet drops
-	HWEncoderDetails.LastBitrate = FMath::Min((uint32)LastBitrate.get_sum_bps(), (uint32)CVarPixelStreamingEncoderMaxBitrate.GetValueOnAnyThread());
+	uint32 Bitrate = (uint32)LastBitrate.get_sum_bps();
+	uint32 MinBitrate = (uint32)CVarPixelStreamingEncoderMinBitrate.GetValueOnAnyThread();
+	uint32 MaxBitrate = (uint32)CVarPixelStreamingEncoderMaxBitrate.GetValueOnAnyThread();
+	HWEncoderDetails.LastBitrate = FMath::Clamp(Bitrate, MinBitrate, MaxBitrate);
 	HWEncoderDetails.Encoder->Encode(BufferId, bForceKeyFrame, HWEncoderDetails.LastBitrate, MoveTemp(EncoderCookie));
 
 	return WEBRTC_VIDEO_CODEC_OK;
