Removing mentions of UE4 in commends and code.

diff --git a/Engine/Plugins/Media/PixelStreaming/Source/PixelStreaming/PixelStreaming.Build.cs b/Engine/Plugins/Media/PixelStreaming/Source/PixelStreaming/PixelStreaming.Build.cs
--- a/Engine/Plugins/Media/PixelStreaming/Source/PixelStreaming/PixelStreaming.Build.cs
+++ b/Engine/Plugins/Media/PixelStreaming/Source/PixelStreaming/PixelStreaming.Build.cs
@@ -131,7 +131,7 @@ namespace UnrealBuildTool.Rules
                 PrivateDependencyModuleNames.Add("HeadMountedDisplay");
             }
 
-            // required for casting UE4 BackBuffer to D3D11 Texture2D for NvEnc
+            // required for casting UnrealEngine BackBuffer to D3D11 Texture2D for NvEnc
             PrivateDependencyModuleNames.AddRange(new string[] { "D3D11RHI" });
             PrivateIncludePaths.AddRange(new string[]
             {
diff --git a/Engine/Plugins/Media/PixelStreaming/Source/PixelStreaming/Private/AudioCapturer.cpp b/Engine/Plugins/Media/PixelStreaming/Source/PixelStreaming/Private/AudioCapturer.cpp
--- a/Engine/Plugins/Media/PixelStreaming/Source/PixelStreaming/Private/AudioCapturer.cpp
+++ b/Engine/Plugins/Media/PixelStreaming/Source/PixelStreaming/Private/AudioCapturer.cpp
@@ -246,7 +246,7 @@ int32 FAudioCapturer::InitRecording()
 	{
 		FScopeLock Lock(&DeviceBufferCS);
 		// #Audio : Allow dynamic values for samplerate and/or channels ,
-		// or receive those from UE4 ?
+		// or receive those from UnrealEngine ?
 		DeviceBuffer->SetRecordingSampleRate(SampleRate);
 		DeviceBuffer->SetRecordingChannels(NumChannels);
 	}
diff --git a/Engine/Plugins/Media/PixelStreaming/Source/PixelStreaming/Private/Codecs/VideoDecoder.cpp b/Engine/Plugins/Media/PixelStreaming/Source/PixelStreaming/Private/Codecs/VideoDecoder.cpp
--- a/Engine/Plugins/Media/PixelStreaming/Source/PixelStreaming/Private/Codecs/VideoDecoder.cpp
+++ b/Engine/Plugins/Media/PixelStreaming/Source/PixelStreaming/Private/Codecs/VideoDecoder.cpp
@@ -22,7 +22,7 @@ DEFINE_LOG_CATEGORY(LogVideoDecoder);
 
 const D3DFORMAT DX9_NV12_FORMAT = (D3DFORMAT)MAKEFOURCC('N', 'V', '1', '2');
 
-inline ID3D11Device* GetUE4DxDevice()
+inline ID3D11Device* GetUEDxDevice()
 {
 	return static_cast<ID3D11Device*>(GDynamicRHI->RHIGetNativeDevice());
 }
@@ -265,15 +265,15 @@ bool FVideoDecoder::FallbackToSwDecoding(FString Reason)
 	if (IsWindows8Plus())
 	{
 		// NOTE: the following doesn't apply to Windows 7 as it doesn't use DX11 device in decoding thread
-		// as we don't use a dedicated DirextX device for s/w decoding, UE4's rendering device will be used from inside the decoder
+		// as we don't use a dedicated DirextX device for s/w decoding, UnrealEngine's rendering device will be used from inside the decoder
 		// to produce output samples, which means access from render and decoding threads. We need to enable multithread protection
 		// for the device. Multithread protection can have performance impact, though its affect is expected to be negligible in most cases.
 		// WARNING:
-		// Once multithread protection is enabled we don't disable it, so UE4's rendering device stays protected for the rest of its lifetime.
+		// Once multithread protection is enabled we don't disable it, so UnrealEngine's rendering device stays protected for the rest of its lifetime.
 		// Some other system could enable multithread protection after we did it, we have no means to know about this, and so disabling it
 		// at the end of playback can cause GPU driver crash
 		TComPtr<ID3D10Multithread> DxMultithread;
-		CHECK_HR(DxMultithread.FromQueryInterface(__uuidof(ID3D10Multithread), GetUE4DxDevice()));
+		CHECK_HR(DxMultithread.FromQueryInterface(__uuidof(ID3D10Multithread), GetUEDxDevice()));
 		DxMultithread->SetMultithreadProtected(1);
 	}
 
@@ -502,7 +502,7 @@ void FVideoDecoder::StopDecoding()
 		// blocked. 
 		if (!ExitingDecodingThreadEvent->Wait(FTimespan::FromSeconds(1)))
 		{
-			// the following (commented out) trick doesn't work because it's not possible to terminate an UE4 thread
+			// the following (commented out) trick doesn't work because it's not possible to terminate an UnrealEngine thread
 
 			//// Timeout. Kill the thread. This means possibly leaving internal decoding threads running, crashing in WMF decoder destructor,
 			//// potential inability to start the next streaming session and other horrible things. Still better than doing nothing as getting stuck on 
@@ -634,7 +634,7 @@ bool FVideoDecoder::CopyTexture(IMFSample* Sample, const TSharedPtr<FWmfMediaHar
 	LONGLONG SampleDuration = 0;
 	CHECK_HR(Sample->GetSampleDuration(&SampleDuration));
 
-	ID3D11Device* UE4DxDevice = static_cast<ID3D11Device*>(GDynamicRHI->RHIGetNativeDevice());
+	ID3D11Device* UEDxDevice = static_cast<ID3D11Device*>(GDynamicRHI->RHIGetNativeDevice());
 
 	TComPtr<IMFMediaBuffer> Buffer;
 	CHECK_HR(Sample->GetBufferByIndex(0, &Buffer));
@@ -759,10 +759,10 @@ bool FVideoDecoder::CreateDXGIManagerAndDevice()
 	}
 
 	// Create device from same adapter as already existing device
-	ID3D11Device* UE4DxDevice = static_cast<ID3D11Device*>(GDynamicRHI->RHIGetNativeDevice());
+	ID3D11Device* UEDxDevice = static_cast<ID3D11Device*>(GDynamicRHI->RHIGetNativeDevice());
 
 	TComPtr<IDXGIDevice> DXGIDevice;
-	UE4DxDevice->QueryInterface(__uuidof(IDXGIDevice), (void**)&DXGIDevice);
+	UEDxDevice->QueryInterface(__uuidof(IDXGIDevice), (void**)&DXGIDevice);
 
 	TComPtr<IDXGIAdapter> DXGIAdapter;
 	DXGIDevice->GetAdapter((IDXGIAdapter**)&DXGIAdapter);
@@ -771,8 +771,8 @@ bool FVideoDecoder::CreateDXGIManagerAndDevice()
 
 	uint32 DeviceCreationFlags = 0;
 
-	uint32 UE4DxDeviceCreationFlags = UE4DxDevice->GetCreationFlags();
-	if ((UE4DxDeviceCreationFlags & D3D11_CREATE_DEVICE_DEBUG) != 0)
+	uint32 UEDxDeviceCreationFlags = UEDxDevice->GetCreationFlags();
+	if ((UEDxDeviceCreationFlags & D3D11_CREATE_DEVICE_DEBUG) != 0)
 	{
 		DeviceCreationFlags |= D3D11_CREATE_DEVICE_DEBUG;
 	}
diff --git a/Engine/Plugins/Media/PixelStreaming/Source/PixelStreaming/Private/Codecs/VideoDecoder.h b/Engine/Plugins/Media/PixelStreaming/Source/PixelStreaming/Private/Codecs/VideoDecoder.h
--- a/Engine/Plugins/Media/PixelStreaming/Source/PixelStreaming/Private/Codecs/VideoDecoder.h
+++ b/Engine/Plugins/Media/PixelStreaming/Source/PixelStreaming/Private/Codecs/VideoDecoder.h
@@ -138,7 +138,7 @@ protected:
 	FWmfMediaHardwareVideoDecodingTextureSamplePool HwTextureSamplePool;
 	FSwTextureSamplePool SwTextureSamplePool;
 
-	// D3D11 device dedicated for h/w decoding to avoid mutithreading collisions with UE4 rendering device
+	// D3D11 device dedicated for h/w decoding to avoid mutithreading collisions with UnrealEngine rendering device
 	static TUniquePtr<FDeviceInfo> s_DeviceInfo;
 };
 
diff --git a/Engine/Plugins/Media/PixelStreaming/Source/PixelStreaming/Private/Codecs/WmfIncludes.h b/Engine/Plugins/Media/PixelStreaming/Source/PixelStreaming/Private/Codecs/WmfIncludes.h
--- a/Engine/Plugins/Media/PixelStreaming/Source/PixelStreaming/Private/Codecs/WmfIncludes.h
+++ b/Engine/Plugins/Media/PixelStreaming/Source/PixelStreaming/Private/Codecs/WmfIncludes.h
@@ -30,7 +30,7 @@ THIRD_PARTY_INCLUDES_END
 
 #include "D3D11RHIPrivate.h"
 
-// following commented include causes name clash between UE4 and Windows `IMediaEventSink`,
+// following commented include causes name clash between UnrealEngine and Windows `IMediaEventSink`,
 // we just need a couple of GUIDs from there so the solution is to duplicate them below
 //#include "wmcodecdsp.h"
 
diff --git a/Engine/Plugins/Media/PixelStreaming/Source/PixelStreaming/Private/Player.h b/Engine/Plugins/Media/PixelStreaming/Source/PixelStreaming/Private/Player.h
--- a/Engine/Plugins/Media/PixelStreaming/Source/PixelStreaming/Private/Player.h
+++ b/Engine/Plugins/Media/PixelStreaming/Source/PixelStreaming/Private/Player.h
@@ -28,7 +28,7 @@ class FThread;
 class FStreamerConnection;
 
 // A media player for PixelStreaming.
-// This class is a bridge between UE4 Media Framework and underlying implementation of WebRTC player
+// This class is a bridge between UnrealEngine Media Framework and underlying implementation of WebRTC player
 class FPlayer
 	: public IMediaPlayer
 	, private IMediaCache
diff --git a/Engine/Plugins/Media/PixelStreaming/Source/PixelStreaming/Private/Streamer.cpp b/Engine/Plugins/Media/PixelStreaming/Source/PixelStreaming/Private/Streamer.cpp
--- a/Engine/Plugins/Media/PixelStreaming/Source/PixelStreaming/Private/Streamer.cpp
+++ b/Engine/Plugins/Media/PixelStreaming/Source/PixelStreaming/Private/Streamer.cpp
@@ -25,7 +25,7 @@ bool FStreamer::CheckPlatformCompatibility()
 FStreamer::FStreamer(const FString& InSignallingServerUrl):
 	SignallingServerUrl(InSignallingServerUrl)
 {
-	RedirectWebRtcLogsToUE4(rtc::LoggingSeverity::LS_VERBOSE);
+	RedirectWebRtcLogsToUnreal(rtc::LoggingSeverity::LS_VERBOSE);
 
 	FModuleManager::LoadModuleChecked<IModuleInterface>(TEXT("AVEncoder"));
 
diff --git a/Engine/Plugins/Media/PixelStreaming/Source/PixelStreaming/Private/StreamerConnection.cpp b/Engine/Plugins/Media/PixelStreaming/Source/PixelStreaming/Private/StreamerConnection.cpp
--- a/Engine/Plugins/Media/PixelStreaming/Source/PixelStreaming/Private/StreamerConnection.cpp
+++ b/Engine/Plugins/Media/PixelStreaming/Source/PixelStreaming/Private/StreamerConnection.cpp
@@ -22,7 +22,7 @@ FStreamerConnection::FStreamerConnection(
 	: SignallingServerAddress(InSignallingServerAddress)
 	, OnDisconnection(MoveTemp(InOnDisconnection))
 {
-	RedirectWebRtcLogsToUE4(rtc::LoggingSeverity::LS_VERBOSE);
+	RedirectWebRtcLogsToUnreal(rtc::LoggingSeverity::LS_VERBOSE);
 
 	// required for communication with Signalling Server and must be called in the game thread, while it's used in signalling thread
 	FModuleManager::LoadModuleChecked<FWebSocketsModule>("WebSockets");
diff --git a/Engine/Plugins/Media/PixelStreaming/Source/PixelStreaming/Private/WebRTCIncludes.h b/Engine/Plugins/Media/PixelStreaming/Source/PixelStreaming/Private/WebRTCIncludes.h
--- a/Engine/Plugins/Media/PixelStreaming/Source/PixelStreaming/Private/WebRTCIncludes.h
+++ b/Engine/Plugins/Media/PixelStreaming/Source/PixelStreaming/Private/WebRTCIncludes.h
@@ -17,8 +17,8 @@ THIRD_PARTY_INCLUDES_START
 #	define INCL_EXTRA_HTON_FUNCTIONS
 #endif
 
-// C4582/3: constructor/desctructor is not implicitly called in "api/rtcerror.h", treated as an error by UE4
-// for some unknown reasons we have to disable it inside those UE4's windows-related includes
+// C4582/3: constructor/desctructor is not implicitly called in "api/rtcerror.h", treated as an error by UnrealEngine
+// for some unknown reasons we have to disable it inside those UnrealEngine's windows-related includes
 // C6323: Use of arithmetic operator on Boolean type(s).
 #pragma warning(push)
 #pragma warning(disable: 4582 4583 6323)
@@ -44,7 +44,7 @@ THIRD_PARTY_INCLUDES_END
 
 THIRD_PARTY_INCLUDES_START
 
-// C4582: constructor is not implicitly called in "api/rtcerror.h", treated as an error by UE4
+// C4582: constructor is not implicitly called in "api/rtcerror.h", treated as an error by UnrealEngine
 // C6319: Use of the comma-operator in a tested expression causes the left argument to be ignored when it has no side-effects.
 // C6323: Use of arithmetic operator on Boolean type(s).
 #pragma warning(push)
diff --git a/Engine/Plugins/Media/PixelStreaming/Source/PixelStreaming/Private/WebRTCLogging.cpp b/Engine/Plugins/Media/PixelStreaming/Source/PixelStreaming/Private/WebRTCLogging.cpp
--- a/Engine/Plugins/Media/PixelStreaming/Source/PixelStreaming/Private/WebRTCLogging.cpp
+++ b/Engine/Plugins/Media/PixelStreaming/Source/PixelStreaming/Private/WebRTCLogging.cpp
@@ -41,7 +41,7 @@ private:
 		const char* tag) override
 	{
 #if !NO_LOGGING
-		static const ELogVerbosity::Type RtcToUE4LogCategoryMap[] =
+		static const ELogVerbosity::Type RtcToUnrealLogCategoryMap[] =
 		{
 			ELogVerbosity::VeryVerbose,
 			ELogVerbosity::Verbose,
@@ -50,7 +50,7 @@ private:
 			ELogVerbosity::Error
 		};
 
-		if (PixelStreamingWebRTC.IsSuppressed(RtcToUE4LogCategoryMap[severity]))
+		if (PixelStreamingWebRTC.IsSuppressed(RtcToUnrealLogCategoryMap[severity]))
 		{
 			return;
 		}
@@ -95,7 +95,7 @@ private:
 	}
 };
 
-void RedirectWebRtcLogsToUE4(rtc::LoggingSeverity Verbosity)
+void RedirectWebRtcLogsToUnreal(rtc::LoggingSeverity Verbosity)
 {
 #if !UE_BUILD_SHIPPING
 	static FWebRtcLogsRedirector Redirector{ Verbosity };
diff --git a/Engine/Plugins/Media/PixelStreaming/Source/PixelStreaming/Private/WebRTCLogging.h b/Engine/Plugins/Media/PixelStreaming/Source/PixelStreaming/Private/WebRTCLogging.h
--- a/Engine/Plugins/Media/PixelStreaming/Source/PixelStreaming/Private/WebRTCLogging.h
+++ b/Engine/Plugins/Media/PixelStreaming/Source/PixelStreaming/Private/WebRTCLogging.h
@@ -4,4 +4,4 @@
 
 #include "rtc_base/logging.h"
 
-void RedirectWebRtcLogsToUE4(rtc::LoggingSeverity Verbosity);
+void RedirectWebRtcLogsToUnreal(rtc::LoggingSeverity Verbosity);
diff --git a/Engine/Plugins/Media/PixelStreaming/Source/PixelStreaming/Public/IPixelStreamingModule.h b/Engine/Plugins/Media/PixelStreaming/Source/PixelStreaming/Public/IPixelStreamingModule.h
--- a/Engine/Plugins/Media/PixelStreaming/Source/PixelStreaming/Public/IPixelStreamingModule.h
+++ b/Engine/Plugins/Media/PixelStreaming/Source/PixelStreaming/Public/IPixelStreamingModule.h
@@ -63,7 +63,7 @@ public:
 
 	/**
 	 * Send a data command back to the browser where we are sending video. This
-	 * is different to a response as a command is low-level and coming from UE4
+	 * is different to a response as a command is low-level and coming from UnrealEngine
 	 * rather than the pixel streamed application.
 	 * @param Descriptor - A generic descriptor string.
 	 */
diff --git a/Engine/Source/Runtime/AVEncoder/Private/Microsoft/AVEncoderMicrosoftCommon.cpp b/Engine/Source/Runtime/AVEncoder/Private/Microsoft/AVEncoderMicrosoftCommon.cpp
--- a/Engine/Source/Runtime/AVEncoder/Private/Microsoft/AVEncoderMicrosoftCommon.cpp
+++ b/Engine/Source/Runtime/AVEncoder/Private/Microsoft/AVEncoderMicrosoftCommon.cpp
@@ -17,10 +17,10 @@ namespace AVEncoder
 // Windows only code
 //
 #if PLATFORM_WINDOWS
-ID3D11Device* GetUE4DxDevice()
+ID3D11Device* GetUEDxDevice()
 {
 	auto Device = static_cast<ID3D11Device*>(GDynamicRHI->RHIGetNativeDevice());
-	checkf(Device != nullptr, TEXT("Failed to get UE4's ID3D11Device"));
+	checkf(Device != nullptr, TEXT("Failed to get UnrealEngine's ID3D11Device"));
 	return Device;
 }
 #endif
@@ -29,10 +29,10 @@ ID3D11Device* GetUE4DxDevice()
 // XboxOne only code
 // 
 #if PLATFORM_XBOXONE 
-ID3D12Device* GetUE4DxDevice()
+ID3D12Device* GetUEDxDevice()
 {
 	auto Device = static_cast<ID3D12Device*>(GDynamicRHI->RHIGetNativeDevice());
-	checkf(Device != nullptr, TEXT("Failed to get UE4's ID3D12Device"));
+	checkf(Device != nullptr, TEXT("Failed to get UnrealEngine's ID3D12Device"));
 	return Device;
 }
 
diff --git a/Engine/Source/Runtime/AVEncoder/Private/Microsoft/AVEncoderMicrosoftCommon.h b/Engine/Source/Runtime/AVEncoder/Private/Microsoft/AVEncoderMicrosoftCommon.h
--- a/Engine/Source/Runtime/AVEncoder/Private/Microsoft/AVEncoderMicrosoftCommon.h
+++ b/Engine/Source/Runtime/AVEncoder/Private/Microsoft/AVEncoderMicrosoftCommon.h
@@ -135,7 +135,7 @@ inline const FString GetComErrorDescription(HRESULT Res)
 	#include "XboxCommonHidePlatformTypes.h"
 #endif
 
-// following commented include causes name clash between UE4 and Windows `IMediaEventSink`,
+// following commented include causes name clash between UnrealEngine and Windows `IMediaEventSink`,
 // we just need a couple of GUIDs from there so the solution is to duplicate them below
 //#include "wmcodecdsp.h"
 
@@ -145,8 +145,8 @@ const GUID CLSID_CMSH264EncoderMFT = { 0x6ca50344, 0x051a, 0x4ded, { 0x97, 0x79,
 const GUID CLSID_VideoProcessorMFT = { 0x88753b26, 0x5b24, 0x49bd, { 0xb2, 0xe7, 0xc, 0x44, 0x5c, 0x78, 0xc9, 0x82 } };
 
 // `MF_LOW_LATENCY` is defined in "mfapi.h" for >= WIN8
-// UE4 supports lower Windows versions at the moment and so `WINVER` is < `_WIN32_WINNT_WIN8`
-// to be able to use `MF_LOW_LATENCY` with default UE4 build we define it ourselves and check actual
+// UnrealEngine supports lower Windows versions at the moment and so `WINVER` is < `_WIN32_WINNT_WIN8`
+// to be able to use `MF_LOW_LATENCY` with default UnrealEngine build we define it ourselves and check actual
 // Windows version in runtime
 #if (WINVER < _WIN32_WINNT_WIN8)
 	const GUID MF_LOW_LATENCY = { 0x9c27891a, 0xed7a, 0x40e1,{ 0x88, 0xe8, 0xb2, 0x27, 0x27, 0xa0, 0x24, 0xee } };
@@ -155,11 +155,11 @@ const GUID CLSID_VideoProcessorMFT = { 0x88753b26, 0x5b24, 0x49bd, { 0xb2, 0xe7,
 
 #if PLATFORM_WINDOWS
 
-ID3D11Device* GetUE4DxDevice();
+ID3D11Device* GetUEDxDevice();
 
 #elif PLATFORM_XBOXONE
 
-ID3D12Device* GetUE4DxDevice();
+ID3D12Device* GetUEDxDevice();
 
 #endif
 
@@ -173,7 +173,7 @@ public:
 	FScopeDisabledDxDebugErrors(TArray<D3D11_MESSAGE_ID>&& ErrorsToDisable)
 	{
 		TRefCountPtr<ID3D11Debug> Debug;
-		HRESULT HRes = GetUE4DxDevice()->QueryInterface(__uuidof(ID3D11Debug), reinterpret_cast<void**>(Debug.GetInitReference()));
+		HRESULT HRes = GetUEDxDevice()->QueryInterface(__uuidof(ID3D11Debug), reinterpret_cast<void**>(Debug.GetInitReference()));
 
 		if (HRes == E_NOINTERFACE)
 		{
diff --git a/Engine/Source/Runtime/AVEncoder/Private/Microsoft/Windows/NvVideoEncoder.cpp b/Engine/Source/Runtime/AVEncoder/Private/Microsoft/Windows/NvVideoEncoder.cpp
--- a/Engine/Source/Runtime/AVEncoder/Private/Microsoft/Windows/NvVideoEncoder.cpp
+++ b/Engine/Source/Runtime/AVEncoder/Private/Microsoft/Windows/NvVideoEncoder.cpp
@@ -323,17 +323,17 @@ FNvVideoEncoder::FEncoderDevice::FEncoderDevice()
 
 		if (RHIName == TEXT("D3D11"))
 		{
-			auto UE4D3DDevice = static_cast<ID3D11Device*>(GDynamicRHI->RHIGetNativeDevice());
-			checkf(UE4D3DDevice != nullptr, TEXT("Cannot initialize NvEnc with invalid device"));
-			CHECK_HR_VOID(UE4D3DDevice->QueryInterface(__uuidof(IDXGIDevice), (void**)DXGIDevice.GetInitReference()));
+			auto UED3DDevice = static_cast<ID3D11Device*>(GDynamicRHI->RHIGetNativeDevice());
+			checkf(UED3DDevice != nullptr, TEXT("Cannot initialize NvEnc with invalid device"));
+			CHECK_HR_VOID(UED3DDevice->QueryInterface(__uuidof(IDXGIDevice), (void**)DXGIDevice.GetInitReference()));
 			CHECK_HR_VOID(DXGIDevice->GetAdapter(Adapter.GetInitReference()));
 			FeatureLevel = D3D_FEATURE_LEVEL_11_0;
 		}
 		else if (RHIName == TEXT("D3D12"))
 		{
-			auto UE4D3DDevice = static_cast<ID3D12Device*>(GDynamicRHI->RHIGetNativeDevice());
-			checkf(UE4D3DDevice != nullptr, TEXT("Cannot initialize NvEnc with invalid device"));
-			LUID AdapterLuid = UE4D3DDevice->GetAdapterLuid();
+			auto UED3DDevice = static_cast<ID3D12Device*>(GDynamicRHI->RHIGetNativeDevice());
+			checkf(UED3DDevice != nullptr, TEXT("Cannot initialize NvEnc with invalid device"));
+			LUID AdapterLuid = UED3DDevice->GetAdapterLuid();
 			TRefCountPtr<IDXGIFactory4> DXGIFactory;
 			CHECK_HR_VOID(CreateDXGIFactory(IID_PPV_ARGS(DXGIFactory.GetInitReference())));
 			// To use a shared texture from D3D12, we need to use a D3D 11.1 device, because we need the
@@ -1075,7 +1075,7 @@ bool FNvVideoEncoder::InitFrameInputBuffer(FFrame& Frame, uint32 Width, uint32 H
 	}
 	else if (RHIName == TEXT("D3D12"))
 	{
-		auto UE4D3DDevice = static_cast<ID3D12Device*>(GDynamicRHI->RHIGetNativeDevice());
+		auto UED3DDevice = static_cast<ID3D12Device*>(GDynamicRHI->RHIGetNativeDevice());
 		static uint32 NamingIdx = 0;
 		ID3D12Resource* ResolvedTexture = (ID3D12Resource*)InputFrame.Texture->GetTexture2D()->GetNativeResource();
 
@@ -1083,7 +1083,7 @@ bool FNvVideoEncoder::InitFrameInputBuffer(FFrame& Frame, uint32 Width, uint32 H
 		// NOTE: ID3D12Device::CreateSharedHandle gives as an NT Handle, and so we need to call CloseHandle on it
 		//
 		HANDLE SharedHandle;
-		HRESULT Res1 = UE4D3DDevice->CreateSharedHandle(ResolvedTexture, NULL, GENERIC_ALL, *FString::Printf(TEXT("PixelStreaming_NvEnc_%u"), NamingIdx++), &SharedHandle);
+		HRESULT Res1 = UED3DDevice->CreateSharedHandle(ResolvedTexture, NULL, GENERIC_ALL, *FString::Printf(TEXT("PixelStreaming_NvEnc_%u"), NamingIdx++), &SharedHandle);
 		CHECK_HR_DEFAULT(Res1);
 
 		TRefCountPtr <ID3D11Device1> Device1;
diff --git a/Engine/Source/Runtime/AVEncoder/Public/AVEncoder.h b/Engine/Source/Runtime/AVEncoder/Public/AVEncoder.h
--- a/Engine/Source/Runtime/AVEncoder/Public/AVEncoder.h
+++ b/Engine/Source/Runtime/AVEncoder/Public/AVEncoder.h
@@ -54,7 +54,7 @@
  * 		// Register ourselves to receive the encoded frames
  * 		VideoEncoder->RegisterListener(this);
  * 
- * 		// Register back buffer capture with UE4
+ * 		// Register back buffer capture with UnrealEngine
  * 		FSlateApplication::Get().GetRenderer()->OnBackBufferReadyToPresent().AddRaw(this, &FGameplayMediaEncoder::OnBackBufferReady);
  * 
  * 		StartTime = FTimespan::FromSeconds(FPlatformTime::Seconds());
@@ -63,7 +63,7 @@
  * 	}
  * 
  * 
- * 	// Back buffer capture callback from UE4
+ * 	// Back buffer capture callback from UnrealEngine
  * 	void OnBackBufferReady(SWindow& SlateWindow, const FTexture2DRHIRef& BackBuffer)
  * 	{
  * 		FTimespan Now = FTimespan::FromSeconds(FPlatformTime::Seconds()) - StartTime;
