remove PLATFORM_XBOXONE from media code.

#rnx

diff --git a/Engine/Plugins/Media/PixelStreaming/Source/PixelStreaming/Private/PCH.h b/Engine/Plugins/Media/PixelStreaming/Source/PixelStreaming/Private/PCH.h
--- a/Engine/Plugins/Media/PixelStreaming/Source/PixelStreaming/Private/PCH.h
+++ b/Engine/Plugins/Media/PixelStreaming/Source/PixelStreaming/Private/PCH.h
@@ -72,7 +72,7 @@
 #include "Engine/Engine.h"
 #include "Framework/Application/SlateUser.h"
 
-#if PLATFORM_WINDOWS || PLATFORM_XBOXONE
+#if PLATFORM_WINDOWS
 #include "Windows/WindowsPlatformMisc.h"
 #elif PLATFORM_LINUX
 #include "Linux/LinuxPlatformMisc.h"
diff --git a/Engine/Plugins/Media/PixelStreaming/Source/PixelStreaming/Private/PixelStreamingModule.cpp b/Engine/Plugins/Media/PixelStreaming/Source/PixelStreaming/Private/PixelStreamingModule.cpp
--- a/Engine/Plugins/Media/PixelStreaming/Source/PixelStreaming/Private/PixelStreamingModule.cpp
+++ b/Engine/Plugins/Media/PixelStreaming/Source/PixelStreaming/Private/PixelStreamingModule.cpp
@@ -17,7 +17,7 @@
 #include "Engine/Texture2D.h"
 #include "Slate/SceneViewport.h"
 
-#if PLATFORM_WINDOWS || PLATFORM_XBOXONE
+#if PLATFORM_WINDOWS
 #include "Windows/WindowsHWrapper.h"
 #elif PLATFORM_LINUX
 #include "CudaModule.h"
@@ -54,7 +54,7 @@ namespace
 		ECVF_Default
 	);
 
-	#if PLATFORM_WINDOWS || PLATFORM_XBOXONE
+	#if PLATFORM_WINDOWS
 	// required for WMF video decoding
 	// some Windows versions don't have Media Foundation preinstalled. We configure MF DLLs as delay-loaded and load them manually here
 	// checking the result and avoiding error message box if failed
@@ -144,7 +144,7 @@ void FPixelStreamingModule::InitPlayer()
 {
 	check(!bPlayerInitialized);
 
-#if PLATFORM_WINDOWS || PLATFORM_XBOXONE
+#if PLATFORM_WINDOWS
 	// Win7+ only
 	if (!IsWindows7Plus())
 	{
@@ -224,7 +224,7 @@ bool FPixelStreamingModule::CheckPlatformCompatibility() const
 {
 	bool bCompatible = true;
 
-	#if PLATFORM_WINDOWS || PLATFORM_XBOXONE
+	#if PLATFORM_WINDOWS
 	bool bWin8OrHigher = FPlatformMisc::VerifyWindowsVersion(6, 2);
 	if (!bWin8OrHigher)
 	{
diff --git a/Engine/Plugins/Media/PixelStreaming/Source/PixelStreaming/Private/PixelStreamingPrivate.h b/Engine/Plugins/Media/PixelStreaming/Source/PixelStreaming/Private/PixelStreamingPrivate.h
--- a/Engine/Plugins/Media/PixelStreaming/Source/PixelStreaming/Private/PixelStreamingPrivate.h
+++ b/Engine/Plugins/Media/PixelStreaming/Source/PixelStreaming/Private/PixelStreamingPrivate.h
@@ -5,7 +5,7 @@
 #include "WebRTCIncludes.h"
 
 // TODO: Add back in when we implement decoder.
-// #if PLATFORM_WINDOWS || PLATFORM_XBOXONE
+// #if PLATFORM_WINDOWS
 // #include "WmfIncludes.h"
 // #endif
 
diff --git a/Engine/Plugins/Media/PixelStreaming/Source/PixelStreaming/Private/Player.cpp b/Engine/Plugins/Media/PixelStreaming/Source/PixelStreaming/Private/Player.cpp
--- a/Engine/Plugins/Media/PixelStreaming/Source/PixelStreaming/Private/Player.cpp
+++ b/Engine/Plugins/Media/PixelStreaming/Source/PixelStreaming/Private/Player.cpp
@@ -115,7 +115,7 @@ void FPlayer::TickInput(FTimespan DeltaTime, FTimespan Timecode)
 
 bool FPlayer::CreateManagerAndDevice()
 {
-#if PLATFORM_WINDOWS || PLATFORM_XBOXONE
+#if PLATFORM_WINDOWS
 	// TODO (M84FIX) - Implement new Decoder 
 	return false;
 	//return FVideoDecoder::CreateDXManagerAndDevice();
@@ -127,7 +127,7 @@ bool FPlayer::CreateManagerAndDevice()
 
 bool FPlayer::DestroyManagerAndDevice()
 {
-#if PLATFORM_WINDOWS || PLATFORM_XBOXONE
+#if PLATFORM_WINDOWS
 	// TODO (M84FIX) - Implement new Decoder 
 	return false;
 	//return FVideoDecoder::DestroyDXManagerAndDevice();
diff --git a/Engine/Plugins/Media/PixelStreaming/Source/PixelStreaming/Private/Streamer.cpp b/Engine/Plugins/Media/PixelStreaming/Source/PixelStreaming/Private/Streamer.cpp
--- a/Engine/Plugins/Media/PixelStreaming/Source/PixelStreaming/Private/Streamer.cpp
+++ b/Engine/Plugins/Media/PixelStreaming/Source/PixelStreaming/Private/Streamer.cpp
@@ -61,7 +61,7 @@ FStreamer::FStreamer(const FString& InSignallingServerUrl)
 FStreamer::~FStreamer()
 {
 	// stop WebRtc WndProc thread
-#if PLATFORM_WINDOWS || PLATFORM_XBOXONE
+#if PLATFORM_WINDOWS
 	PostThreadMessage(WebRtcSignallingThreadId, WM_QUIT, 0, 0);
 
 	UE_LOG(PixelStreamer, Log, TEXT("Exiting WebRTC WndProc thread"));
@@ -84,7 +84,7 @@ void FStreamer::WebRtcSignallingThreadFunc()
 {
 	// initialisation of WebRTC stuff and things that depends on it should happen in WebRTC signalling thread
 
-#if PLATFORM_WINDOWS || PLATFORM_XBOXONE
+#if PLATFORM_WINDOWS
 
 	// WebRTC assumes threads within which PeerConnectionFactory is created is the signalling thread
 	WebRtcSignallingThreadId = FPlatformTLS::GetCurrentThreadId();
diff --git a/Engine/Plugins/Media/PixelStreaming/Source/PixelStreaming/Private/Streamer.h b/Engine/Plugins/Media/PixelStreaming/Source/PixelStreaming/Private/Streamer.h
--- a/Engine/Plugins/Media/PixelStreaming/Source/PixelStreaming/Private/Streamer.h
+++ b/Engine/Plugins/Media/PixelStreaming/Source/PixelStreaming/Private/Streamer.h
@@ -54,7 +54,7 @@ private:
 private:
 	FString SignallingServerUrl;
 
-	#if PLATFORM_WINDOWS || PLATFORM_XBOXONE
+	#if PLATFORM_WINDOWS
 	DWORD WebRtcSignallingThreadId = 0;
 	TUniquePtr<FThread> WebRtcSignallingThread;
 	#elif PLATFORM_LINUX
diff --git a/Engine/Plugins/Media/PixelStreaming/Source/PixelStreaming/Private/StreamerConnection.cpp b/Engine/Plugins/Media/PixelStreaming/Source/PixelStreaming/Private/StreamerConnection.cpp
--- a/Engine/Plugins/Media/PixelStreaming/Source/PixelStreaming/Private/StreamerConnection.cpp
+++ b/Engine/Plugins/Media/PixelStreaming/Source/PixelStreaming/Private/StreamerConnection.cpp
@@ -45,7 +45,7 @@ FStreamerConnection::~FStreamerConnection()
 	DataChannel = nullptr;
 	PeerConnection = nullptr;
 
-#if PLATFORM_WINDOWS || PLATFORM_XBOXONE
+#if PLATFORM_WINDOWS
 	// stop Signalling Thread
 	check(SignallingThreadId != 0);
 	PostThreadMessage(SignallingThreadId, WM_QUIT, 0, 0);
@@ -63,7 +63,7 @@ void FStreamerConnection::SignallingThreadFunc()
 {
 	SignallingThreadId = FPlatformTLS::GetCurrentThreadId();
 
-#if PLATFORM_WINDOWS || PLATFORM_XBOXONE
+#if PLATFORM_WINDOWS
 
 	// init WebRTC networking and inter-thread communication
 	rtc::WinsockInitializer WSInitialiser;
@@ -109,7 +109,7 @@ void FStreamerConnection::SignallingThreadFunc()
 	// now that everything is ready connect to SignallingServer
 	SignallingServerConnection = MakeUnique<FSignallingServerConnection>(SignallingServerAddress, *this);
 
-#if PLATFORM_WINDOWS || PLATFORM_XBOXONE
+#if PLATFORM_WINDOWS
 	// WebRTC window messaging loop
 	MSG Msg;
 	BOOL Gm;
diff --git a/Engine/Plugins/Media/PixelStreaming/Source/PixelStreaming/Private/StreamerConnection.h b/Engine/Plugins/Media/PixelStreaming/Source/PixelStreaming/Private/StreamerConnection.h
--- a/Engine/Plugins/Media/PixelStreaming/Source/PixelStreaming/Private/StreamerConnection.h
+++ b/Engine/Plugins/Media/PixelStreaming/Source/PixelStreaming/Private/StreamerConnection.h
@@ -59,7 +59,7 @@ private:
 private:
 	TUniquePtr<FThread> SignallingThread;
 	
-	#if PLATFORM_WINDOWS || PLATFORM_XBOXONE
+	#if PLATFORM_WINDOWS
 	DWORD SignallingThreadId = 0;
 	#elif PLATFORM_LINUX
 	FUnixPlatformTypes::DWORD SignallingThreadId = 0;
diff --git a/Engine/Plugins/Media/PixelStreaming/Source/PixelStreaming/Private/Utils.h b/Engine/Plugins/Media/PixelStreaming/Source/PixelStreaming/Private/Utils.h
--- a/Engine/Plugins/Media/PixelStreaming/Source/PixelStreaming/Private/Utils.h
+++ b/Engine/Plugins/Media/PixelStreaming/Source/PixelStreaming/Private/Utils.h
@@ -4,7 +4,7 @@
 
 #include "PixelStreamingPrivate.h"
 
-#if PLATFORM_WINDOWS || PLATFORM_XBOXONE
+#if PLATFORM_WINDOWS
 #include "Windows/WindowsPlatformMisc.h"
 #elif PLATFORM_LINUX
 #include "Linux/LinuxPlatformMisc.h"
@@ -21,7 +21,7 @@
 
 #include <string>
 
-#if PLATFORM_WINDOWS || PLATFORM_XBOXONE
+#if PLATFORM_WINDOWS
 inline bool IsWindows7Plus()
 {
 	return FPlatformMisc::VerifyWindowsVersion(6, 1);
diff --git a/Engine/Plugins/Media/PixelStreaming/Source/PixelStreaming/Private/WebRTCIncludes.h b/Engine/Plugins/Media/PixelStreaming/Source/PixelStreaming/Private/WebRTCIncludes.h
--- a/Engine/Plugins/Media/PixelStreaming/Source/PixelStreaming/Private/WebRTCIncludes.h
+++ b/Engine/Plugins/Media/PixelStreaming/Source/PixelStreaming/Private/WebRTCIncludes.h
@@ -2,7 +2,7 @@
 
 #pragma once
 
-#if PLATFORM_WINDOWS || PLATFORM_XBOXONE
+#if PLATFORM_WINDOWS
 #include "Windows/AllowWindowsPlatformTypes.h"
 #include "Windows/AllowWindowsPlatformAtomics.h"
 
@@ -43,7 +43,7 @@ THIRD_PARTY_INCLUDES_END
 #include "Windows/HideWindowsPlatformAtomics.h"
 #include "Windows/HideWindowsPlatformTypes.h"
 
-#endif // PLATFORM_WINDOWS || PLATFORM_XBOXONE
+#endif // PLATFORM_WINDOWS
 
 THIRD_PARTY_INCLUDES_START
 
diff --git a/Engine/Source/Runtime/AVEncoder/Private/MicrosoftCommon.h b/Engine/Source/Runtime/AVEncoder/Private/MicrosoftCommon.h
--- a/Engine/Source/Runtime/AVEncoder/Private/MicrosoftCommon.h
+++ b/Engine/Source/Runtime/AVEncoder/Private/MicrosoftCommon.h
@@ -77,11 +77,7 @@ inline const FString GetComErrorDescription(HRESULT Res)
 		}\
 	}
 
-#if PLATFORM_WINDOWS
-	#include "Windows/HideWindowsPlatformTypes.h"
-#elif PLATFORM_XBOXONE
-	#include "XboxCommonHidePlatformTypes.h"
-#endif
+#include "Windows/HideWindowsPlatformTypes.h"
 
 // following commented include causes name clash between UE4 and Windows `IMediaEventSink`,
 // we just need a couple of GUIDs from there so the solution is to duplicate them below
diff --git a/Engine/Source/Runtime/GameplayMediaEncoder/Private/GameplayMediaEncoderCommon.h b/Engine/Source/Runtime/GameplayMediaEncoder/Private/GameplayMediaEncoderCommon.h
--- a/Engine/Source/Runtime/GameplayMediaEncoder/Private/GameplayMediaEncoderCommon.h
+++ b/Engine/Source/Runtime/GameplayMediaEncoder/Private/GameplayMediaEncoderCommon.h
@@ -24,7 +24,7 @@
 //
 // Includes common to Windows and XboxOne
 //
-#if PLATFORM_WINDOWS || PLATFORM_XBOXONE
+#if PLATFORM_WINDOWS
 	#include "Templates/RefCounting.h"
 #endif
 
