Fixing no unity build errors
#rnx

diff --git a/Engine/Plugins/Media/PixelStreaming/Source/PixelStreaming/Private/SignallingServerConnection.cpp b/Engine/Plugins/Media/PixelStreaming/Source/PixelStreaming/Private/SignallingServerConnection.cpp
--- a/Engine/Plugins/Media/PixelStreaming/Source/PixelStreaming/Private/SignallingServerConnection.cpp
+++ b/Engine/Plugins/Media/PixelStreaming/Source/PixelStreaming/Private/SignallingServerConnection.cpp
@@ -6,6 +6,7 @@
 #include "WebSocketsModule.h"
 #include "IWebSocket.h"
 #include "Dom/JsonObject.h"
+#include "Engine/World.h"
 #include "Serialization/JsonSerializer.h"
 #include "Policies/CondensedJsonPrintPolicy.h"
 #include "Misc/AssertionMacros.h"
diff --git a/Engine/Plugins/Media/PixelStreaming/Source/PixelStreaming/Private/SignallingServerConnection.h b/Engine/Plugins/Media/PixelStreaming/Source/PixelStreaming/Private/SignallingServerConnection.h
--- a/Engine/Plugins/Media/PixelStreaming/Source/PixelStreaming/Private/SignallingServerConnection.h
+++ b/Engine/Plugins/Media/PixelStreaming/Source/PixelStreaming/Private/SignallingServerConnection.h
@@ -10,6 +10,7 @@
 #include "Containers/Map.h"
 #include "Dom/JsonObject.h"
 #include "Delegates/IDelegateInstance.h"
+#include "Engine/EngineTypes.h"
 
 class IWebSocket;
 
diff --git a/Engine/Plugins/Media/PixelStreaming/Source/PixelStreaming/Private/VideoCapturer.cpp b/Engine/Plugins/Media/PixelStreaming/Source/PixelStreaming/Private/VideoCapturer.cpp
--- a/Engine/Plugins/Media/PixelStreaming/Source/PixelStreaming/Private/VideoCapturer.cpp
+++ b/Engine/Plugins/Media/PixelStreaming/Source/PixelStreaming/Private/VideoCapturer.cpp
@@ -3,12 +3,17 @@
 #include "VideoCapturer.h"
 #include "Utils.h"
 
+#include "CommonRenderResources.h"
+#include "GlobalShader.h"
 #include "HAL/PlatformTime.h"
 #include "Misc/Timespan.h"
+#include "Modules/ModuleManager.h"
 #include "PixelStreamingFrameBuffer.h"
 #include "PixelStreamingSettings.h"
 #include "ClearQuad.h"
 #include "LatencyTester.h"
+#include "RHIStaticStates.h"
+#include "ScreenRendering.h"
 
 #if PLATFORM_LINUX
 #include "CudaModule.h"
diff --git a/Engine/Plugins/Media/PixelStreaming/Source/PixelStreaming/Private/VideoCapturer.h b/Engine/Plugins/Media/PixelStreaming/Source/PixelStreaming/Private/VideoCapturer.h
--- a/Engine/Plugins/Media/PixelStreaming/Source/PixelStreaming/Private/VideoCapturer.h
+++ b/Engine/Plugins/Media/PixelStreaming/Source/PixelStreaming/Private/VideoCapturer.h
@@ -3,6 +3,7 @@
 #pragma once
 
 #include "PixelStreamingPrivate.h"
+#include "RHI.h"
 #include "VideoEncoderInput.h"
 
 using SourceState = webrtc::MediaSourceInterface::SourceState;
