Fix all NonUnity issues related to AVEncoder and PixelStreaming, and some minor code clean up

diff --git a/Engine/Plugins/Media/PixelStreaming/Source/PixelStreaming/Private/PixelStreamingEncoderFactory.cpp b/Engine/Plugins/Media/PixelStreaming/Source/PixelStreaming/Private/PixelStreamingEncoderFactory.cpp
--- a/Engine/Plugins/Media/PixelStreaming/Source/PixelStreaming/Private/PixelStreamingEncoderFactory.cpp
+++ b/Engine/Plugins/Media/PixelStreaming/Source/PixelStreaming/Private/PixelStreamingEncoderFactory.cpp
@@ -4,11 +4,10 @@
 #include "PlayerSession.h"
 #include "PixelStreamingVideoEncoder.h"
 #include "PixelStreamingSettings.h"
+#include "Utils.h"
 #include "absl/strings/match.h"
 #include "modules/video_coding/codecs/vp8/include/vp8.h"
 
-using namespace webrtc;
-
 FPixelStreamingVideoEncoderFactory::FPixelStreamingVideoEncoderFactory()
 {
 	EncoderContext.Factory = this;
@@ -23,19 +22,19 @@ void FPixelStreamingVideoEncoderFactory::AddSession(FPlayerSession& PlayerSessio
 	PendingPlayerSessions.Enqueue(&PlayerSession);
 }
 
-std::vector<SdpVideoFormat> FPixelStreamingVideoEncoderFactory::GetSupportedFormats() const
+std::vector<webrtc::SdpVideoFormat> FPixelStreamingVideoEncoderFactory::GetSupportedFormats() const
 {
 	const bool bForceVP8 = PixelStreamingSettings::IsForceVP8();
 
 	std::vector<webrtc::SdpVideoFormat> video_formats;
 	if (bForceVP8)
-		video_formats.push_back(SdpVideoFormat(cricket::kVp8CodecName));
+		video_formats.push_back(webrtc::SdpVideoFormat(cricket::kVp8CodecName));
 	else
 		video_formats.push_back(CreateH264Format(webrtc::H264::kProfileConstrainedBaseline, webrtc::H264::kLevel3_1));
 	return video_formats;
 }
 
-FPixelStreamingVideoEncoderFactory::CodecInfo FPixelStreamingVideoEncoderFactory::QueryVideoEncoder(const SdpVideoFormat& format) const
+FPixelStreamingVideoEncoderFactory::CodecInfo FPixelStreamingVideoEncoderFactory::QueryVideoEncoder(const webrtc::SdpVideoFormat& format) const
 {
 	CodecInfo codec_info = { false, false };
 	codec_info.is_hardware_accelerated = true;
@@ -43,10 +42,10 @@ FPixelStreamingVideoEncoderFactory::CodecInfo FPixelStreamingVideoEncoderFactory
 	return codec_info;
 }
 
-std::unique_ptr<VideoEncoder> FPixelStreamingVideoEncoderFactory::CreateVideoEncoder(const SdpVideoFormat& format)
+std::unique_ptr<webrtc::VideoEncoder> FPixelStreamingVideoEncoderFactory::CreateVideoEncoder(const webrtc::SdpVideoFormat& format)
 {
 	if (absl::EqualsIgnoreCase(format.name, cricket::kVp8CodecName))
-		return VP8Encoder::Create();
+		return webrtc::VP8Encoder::Create();
 	else
 	{
 		FPlayerSession* Session;
diff --git a/Engine/Plugins/Media/PixelStreaming/Source/PixelStreaming/Private/PixelStreamingEncoderFactory.h b/Engine/Plugins/Media/PixelStreaming/Source/PixelStreaming/Private/PixelStreamingEncoderFactory.h
--- a/Engine/Plugins/Media/PixelStreaming/Source/PixelStreaming/Private/PixelStreamingEncoderFactory.h
+++ b/Engine/Plugins/Media/PixelStreaming/Source/PixelStreaming/Private/PixelStreamingEncoderFactory.h
@@ -2,6 +2,7 @@
 #pragma once
 
 #include "WebRTCIncludes.h"
+#include "VideoEncoder.h"
 #include "Containers/Queue.h"
 
 class FPlayerSession;
@@ -11,7 +12,7 @@ class FPixelStreamingVideoEncoder;
 struct FEncoderContext
 {
 	FPixelStreamingVideoEncoderFactory* Factory;
-	TUniquePtr<FVideoEncoder> Encoder;
+	TUniquePtr<AVEncoder::FVideoEncoder> Encoder;
 };
 
 class FPixelStreamingVideoEncoderFactory : public webrtc::VideoEncoderFactory
diff --git a/Engine/Plugins/Media/PixelStreaming/Source/PixelStreaming/Private/Streamer.cpp b/Engine/Plugins/Media/PixelStreaming/Source/PixelStreaming/Private/Streamer.cpp
--- a/Engine/Plugins/Media/PixelStreaming/Source/PixelStreaming/Private/Streamer.cpp
+++ b/Engine/Plugins/Media/PixelStreaming/Source/PixelStreaming/Private/Streamer.cpp
@@ -8,6 +8,10 @@
 #include "PixelStreamerDelegates.h"
 #include "PixelStreamingEncoderFactory.h"
 #include "PixelStreamingSettings.h"
+#include "WebRTCLogging.h"
+
+#include "WebSocketsModule.h"
+
 
 DEFINE_LOG_CATEGORY(PixelStreamer);
 
diff --git a/Engine/Plugins/Media/PixelStreaming/Source/PixelStreaming/Private/Streamer.h b/Engine/Plugins/Media/PixelStreaming/Source/PixelStreaming/Private/Streamer.h
--- a/Engine/Plugins/Media/PixelStreaming/Source/PixelStreaming/Private/Streamer.h
+++ b/Engine/Plugins/Media/PixelStreaming/Source/PixelStreaming/Private/Streamer.h
@@ -5,6 +5,10 @@
 #include "SignallingServerConnection.h"
 #include "ProtocolDefs.h"
 
+#include "RHI.h"
+#include "HAL/Thread.h"
+
+
 class FVideoCapturer;
 class FPlayerSession;
 class FPixelStreamingVideoEncoderFactory;
diff --git a/Engine/Source/Runtime/AVEncoder/Private/Encoders/Amf/Amf_Common.cpp b/Engine/Source/Runtime/AVEncoder/Private/Encoders/Amf/Amf_Common.cpp
--- a/Engine/Source/Runtime/AVEncoder/Private/Encoders/Amf/Amf_Common.cpp
+++ b/Engine/Source/Runtime/AVEncoder/Private/Encoders/Amf/Amf_Common.cpp
@@ -1,6 +1,8 @@
 // Copyright Epic Games, Inc. All Rights Reserved.
 
 #include "Amf_Common.h"
+#include "HAL/Platform.h"
+
 #include "AVEncoder.h"
 
 #include "RHI.h"
diff --git a/Engine/Source/Runtime/AVEncoder/Private/Encoders/Amf/Amf_EncoderH264.cpp b/Engine/Source/Runtime/AVEncoder/Private/Encoders/Amf/Amf_EncoderH264.cpp
--- a/Engine/Source/Runtime/AVEncoder/Private/Encoders/Amf/Amf_EncoderH264.cpp
+++ b/Engine/Source/Runtime/AVEncoder/Private/Encoders/Amf/Amf_EncoderH264.cpp
@@ -1,6 +1,8 @@
 // Copyright Epic Games, Inc. All Rights Reserved.
 
 #include "Amf_EncoderH264.h"
+#include "HAL/Platform.h"
+
 #include "VideoEncoderCommon.h"
 #include "CodecPacket.h"
 #include "AVEncoderDebug.h"
diff --git a/Engine/Source/Runtime/AVEncoder/Private/Encoders/NVENC/NVENC_Common.cpp b/Engine/Source/Runtime/AVEncoder/Private/Encoders/NVENC/NVENC_Common.cpp
--- a/Engine/Source/Runtime/AVEncoder/Private/Encoders/NVENC/NVENC_Common.cpp
+++ b/Engine/Source/Runtime/AVEncoder/Private/Encoders/NVENC/NVENC_Common.cpp
@@ -1,6 +1,8 @@
 // Copyright Epic Games, Inc. All Rights Reserved.
 
 #include "NVENC_Common.h"
+#include "HAL/Platform.h"
+
 #include "CoreMinimal.h"
 #include "VideoEncoderCommon.h"
 #include "HAL/PlatformProcess.h"
diff --git a/Engine/Source/Runtime/AVEncoder/Private/Encoders/NVENC/NVENC_EncoderH264.cpp b/Engine/Source/Runtime/AVEncoder/Private/Encoders/NVENC/NVENC_EncoderH264.cpp
--- a/Engine/Source/Runtime/AVEncoder/Private/Encoders/NVENC/NVENC_EncoderH264.cpp
+++ b/Engine/Source/Runtime/AVEncoder/Private/Encoders/NVENC/NVENC_EncoderH264.cpp
@@ -1,6 +1,8 @@
 // Copyright Epic Games, Inc. All Rights Reserved.
 
 #include "NVENC_EncoderH264.h"
+#include "HAL/Platform.h"
+
 #include "VideoEncoderCommon.h"
 #include "CodecPacket.h"
 #include "AVEncoderDebug.h"
