- Remove folder missing warning for Pixel Streaming Servers during build time

diff --git a/Engine/Plugins/Media/PixelStreaming/Source/PixelStreaming/PixelStreaming.Build.cs b/Engine/Plugins/Media/PixelStreaming/Source/PixelStreaming/PixelStreaming.Build.cs
--- a/Engine/Plugins/Media/PixelStreaming/Source/PixelStreaming/PixelStreaming.Build.cs
+++ b/Engine/Plugins/Media/PixelStreaming/Source/PixelStreaming/PixelStreaming.Build.cs
@@ -18,7 +18,6 @@ namespace UnrealBuildTool.Rules
 
 			if (!Directory.Exists(SignallingServerDir))
 			{
-				Log.TraceInformation(string.Format("Signalling Server path '{0}' does not exist", SignallingServerDir));
 				return;
 			}
 
@@ -43,7 +42,6 @@ namespace UnrealBuildTool.Rules
 
 			if (!Directory.Exists(MatchmakingServerDir))
 			{
-				Log.TraceInformation(string.Format("Matchmaking Server path '{0}' does not exist", MatchmakingServerDir));
 				return;
 			}
 
@@ -62,24 +60,6 @@ namespace UnrealBuildTool.Rules
 			}
 		}
 
-		private void AddWebRTCServers()
-		{
-			string webRTCRevision = "31262";
-			string webRTCRevisionDirectory = "./ThirdParty/WebRTC/rev." + webRTCRevision;
-			string webRTCProgramsDirectory = Path.Combine(webRTCRevisionDirectory, "programs/Win64/VS2015/Release");
-
-			List<string> DependenciesToAdd = new List<string>();
-			DependenciesToAdd.AddRange(Directory.GetFiles(webRTCProgramsDirectory, "*.exe"));
-			DependenciesToAdd.AddRange(Directory.GetFiles(webRTCProgramsDirectory, "*.pdb"));
-			DependenciesToAdd.AddRange(Directory.GetFiles(webRTCProgramsDirectory, "*.bat"));
-			DependenciesToAdd.AddRange(Directory.GetFiles(webRTCProgramsDirectory, "*.ps1"));
-
-			foreach (string Dependency in DependenciesToAdd)
-			{
-				RuntimeDependencies.Add(Dependency, StagedFileType.NonUFS);
-			}
-		}
-
 		public PixelStreaming(ReadOnlyTargetRules Target) : base(Target)
 		{
 			// use private PCH to include lots of WebRTC headers
@@ -156,14 +136,9 @@ namespace UnrealBuildTool.Rules
 
 			}
 			
+			// Add servers to packaged projects if they are avaliable
 			AddSignallingServer();
 			AddMatchmakingServer();
-			
-			// We have not been able to build these yet for M84 (these are the WebRTC bundled sample STUN/TURN servers).
-			// For future I think we should advise use of COTURN instead of the WebRTC sample servers.
-			// if (Target.Platform == UnrealTargetPlatform.Win64) {
-			//     AddWebRTCServers();
-			// }
 		}
 	}
 }
