Pixel Streaming orchestration compatibility enhancements

diff --git a/Engine/Plugins/Media/PixelStreaming/Source/PixelStreaming/Private/PixelStreamingEncoderFactory.cpp b/Engine/Plugins/Media/PixelStreaming/Source/PixelStreaming/Private/PixelStreamingEncoderFactory.cpp
--- a/Engine/Plugins/Media/PixelStreaming/Source/PixelStreaming/Private/PixelStreamingEncoderFactory.cpp
+++ b/Engine/Plugins/Media/PixelStreaming/Source/PixelStreaming/Private/PixelStreamingEncoderFactory.cpp
@@ -47,7 +47,7 @@ std::unique_ptr<VideoEncoder> FPixelStreamingVideoEncoderFactory::CreateVideoEnc
 	bool res = PendingPlayerSessions.Dequeue(Session);
 	checkf(res, TEXT("no player session associated with encoder instance"));
 
-	UE_LOG(PixelStreamer, Log, TEXT("Encoder factory addded encoder for PlayerId=%d"), Session->GetPlayerId());
+	UE_LOG(PixelStreamer, Log, TEXT("Encoder factory addded encoder for PlayerId=%s"), *Session->GetPlayerId());
 
 	auto VideoEncoder = std::make_unique<FPixelStreamingVideoEncoder>(Session, &EncoderContext);
 	this->ActiveEncoders.Add(VideoEncoder.get());
@@ -64,7 +64,7 @@ void FPixelStreamingVideoEncoderFactory::OnEncodedImage(const webrtc::EncodedIma
 		if(!Encoder->IsRegisteredWithWebRTC())
 		{
 			ActiveEncoders.RemoveAt(Index);
-			UE_LOG(PixelStreamer, Log, TEXT("Encoder factory cleaned up stale encoder associated with PlayerId=%d"), Encoder->GetPlayerId());
+			UE_LOG(PixelStreamer, Log, TEXT("Encoder factory cleaned up stale encoder associated with PlayerId=%s"), *Encoder->GetPlayerId());
 		}
 		else
 		{
@@ -77,8 +77,8 @@ void FPixelStreamingVideoEncoderFactory::OnEncodedImage(const webrtc::EncodedIma
 void FPixelStreamingVideoEncoderFactory::ReleaseVideoEncoder(FPixelStreamingVideoEncoder* encoder)
 {
 	ActiveEncoders.Remove(encoder);
-	int PlayerId = encoder->GetPlayerId();
-	UE_LOG(PixelStreamer, Log, TEXT("Encoder factory asked to remove encoder for PlayerId=%d"), PlayerId);
+	FPlayerId PlayerId = encoder->GetPlayerId();
+	UE_LOG(PixelStreamer, Log, TEXT("Encoder factory asked to remove encoder for PlayerId=%s"), *PlayerId);
 }
 
 //TODO real encoder class in here
diff --git a/Engine/Plugins/Media/PixelStreaming/Source/PixelStreaming/Private/PixelStreamingModule.cpp b/Engine/Plugins/Media/PixelStreaming/Source/PixelStreaming/Private/PixelStreamingModule.cpp
--- a/Engine/Plugins/Media/PixelStreaming/Source/PixelStreaming/Private/PixelStreamingModule.cpp
+++ b/Engine/Plugins/Media/PixelStreaming/Source/PixelStreaming/Private/PixelStreamingModule.cpp
@@ -81,8 +81,10 @@ namespace
 
 void FPixelStreamingModule::InitStreamer()
 {
+	FString StreamerId;
 	FString SignallingServerIP;
 	uint16 SignallingServerPort = 8888;
+	FParse::Value(FCommandLine::Get(), TEXT("PixelStreamingID="), StreamerId);
 	if (!FParse::Value(FCommandLine::Get(), TEXT("PixelStreamingIP="), SignallingServerIP) ||
 		!FParse::Value(FCommandLine::Get(), TEXT("PixelStreamingPort="), SignallingServerPort))
 	{
@@ -90,6 +92,8 @@ void FPixelStreamingModule::InitStreamer()
 		return;
 	}
 
+	UE_LOG(PixelStreamer, Log, TEXT("PixelStreaming endpoint ID: %s"), *StreamerId);
+
 	if (GIsEditor)
 	{
 		FText TitleText = FText::FromString(TEXT("Pixel Streaming Plugin"));
@@ -137,7 +141,7 @@ void FPixelStreamingModule::InitStreamer()
 	UFreezeFrame::CreateInstance();
 	verify(FModuleManager::Get().LoadModule(FName("ImageWrapper")));
 
-	Streamer = MakeUnique<FStreamer>(FString::Printf(TEXT("ws://%s:%d"), *SignallingServerIP, SignallingServerPort));
+	Streamer = MakeUnique<FStreamer>(FString::Printf(TEXT("ws://%s:%d"), *SignallingServerIP, SignallingServerPort), StreamerId);
 }
 
 void FPixelStreamingModule::InitPlayer()
diff --git a/Engine/Plugins/Media/PixelStreaming/Source/PixelStreaming/Private/PixelStreamingPrivate.h b/Engine/Plugins/Media/PixelStreaming/Source/PixelStreaming/Private/PixelStreamingPrivate.h
--- a/Engine/Plugins/Media/PixelStreaming/Source/PixelStreaming/Private/PixelStreamingPrivate.h
+++ b/Engine/Plugins/Media/PixelStreaming/Source/PixelStreaming/Private/PixelStreamingPrivate.h
@@ -20,4 +20,4 @@ DECLARE_LOG_CATEGORY_EXTERN(PixelStreaming, Log, All);
 DECLARE_LOG_CATEGORY_EXTERN(PixelStreamer, Log, All);
 DECLARE_LOG_CATEGORY_EXTERN(PixelPlayer, Log, All);
 
-using FPlayerId = uint32;
\ No newline at end of file
+using FPlayerId = FString;
diff --git a/Engine/Plugins/Media/PixelStreaming/Source/PixelStreaming/Private/PixelStreamingVideoEncoder.cpp b/Engine/Plugins/Media/PixelStreaming/Source/PixelStreaming/Private/PixelStreamingVideoEncoder.cpp
--- a/Engine/Plugins/Media/PixelStreaming/Source/PixelStreaming/Private/PixelStreamingVideoEncoder.cpp
+++ b/Engine/Plugins/Media/PixelStreaming/Source/PixelStreaming/Private/PixelStreamingVideoEncoder.cpp
@@ -155,14 +155,14 @@ void FPixelStreamingVideoEncoder::SetQualityController(bool bControlsQualityNow)
 {
 	if (bControlsQuality != bControlsQualityNow)
 	{
-		UE_LOG(PixelStreamer, Log, TEXT("%s : PlayerId=%d, controls quality %d"), TEXT("FVideoEncoder::SetQualityController"), this->PlayerId, bControlsQualityNow);
+		UE_LOG(PixelStreamer, Log, TEXT("%s : PlayerId=%s, controls quality %d"), TEXT("FVideoEncoder::SetQualityController"), *this->PlayerId, bControlsQualityNow);
 		bControlsQuality = bControlsQualityNow;
 	}
 }
 
 int FPixelStreamingVideoEncoder::InitEncode(const VideoCodec* codec_settings, const VideoEncoder::Settings& settings)
 {
-	UE_LOG(PixelStreamer, Log, TEXT("PixelStreaming video encoder initialise for PlayerId=%d"), this->GetPlayerId());
+	UE_LOG(PixelStreamer, Log, TEXT("PixelStreaming video encoder initialise for PlayerId=%s"), *this->GetPlayerId());
 
 	if (!bControlsQuality)
 		return WEBRTC_VIDEO_CODEC_OK;
@@ -182,7 +182,7 @@ int FPixelStreamingVideoEncoder::InitEncode(const VideoCodec* codec_settings, co
 
 int32 FPixelStreamingVideoEncoder::RegisterEncodeCompleteCallback(EncodedImageCallback* callback)
 {
-	UE_LOG(PixelStreamer, Log, TEXT("PixelStreaming video encoder callback registered for PlayerId=%d"), this->GetPlayerId());
+	UE_LOG(PixelStreamer, Log, TEXT("PixelStreaming video encoder callback registered for PlayerId=%s"), *this->GetPlayerId());
 	OnEncodedImageCallback = callback;
 	// Encoder is initialised, add it as active encoder to factory.
 	return WEBRTC_VIDEO_CODEC_OK;
@@ -190,7 +190,7 @@ int32 FPixelStreamingVideoEncoder::RegisterEncodeCompleteCallback(EncodedImageCa
 
 int32 FPixelStreamingVideoEncoder::Release()
 {
-	UE_LOG(PixelStreamer, Log, TEXT("PixelStreaming video encoder released for PlayerId=%d"), this->GetPlayerId());
+	UE_LOG(PixelStreamer, Log, TEXT("PixelStreaming video encoder released for PlayerId=%s"), *this->GetPlayerId());
 	OnEncodedImageCallback = nullptr;
 	return WEBRTC_VIDEO_CODEC_OK;
 }
diff --git a/Engine/Plugins/Media/PixelStreaming/Source/PixelStreaming/Private/PlayerSession.cpp b/Engine/Plugins/Media/PixelStreaming/Source/PixelStreaming/Private/PlayerSession.cpp
--- a/Engine/Plugins/Media/PixelStreaming/Source/PixelStreaming/Private/PlayerSession.cpp
+++ b/Engine/Plugins/Media/PixelStreaming/Source/PixelStreaming/Private/PlayerSession.cpp
@@ -19,12 +19,12 @@ FPlayerSession::FPlayerSession(FStreamer& InStreamer, FPlayerId InPlayerId, bool
 	, bOriginalQualityController(bInOriginalQualityController)
 	, InputDevice(FModuleManager::Get().GetModuleChecked<IPixelStreamingModule>("PixelStreaming").GetInputDevice())
 {
-	UE_LOG(PixelStreamer, Log, TEXT("%s: PlayerId=%u, quality controller: %d"), TEXT("FPlayerSession::FPlayerSession"), PlayerId, bOriginalQualityController);
+	UE_LOG(PixelStreamer, Log, TEXT("%s: PlayerId=%s, quality controller: %d"), TEXT("FPlayerSession::FPlayerSession"), *PlayerId, bOriginalQualityController);
 }
 
 FPlayerSession::~FPlayerSession()
 {
-	UE_LOG(PixelStreamer, Log, TEXT("%s: PlayerId=%u"), TEXT("FPlayerSession::~FPlayerSession"), PlayerId);
+	UE_LOG(PixelStreamer, Log, TEXT("%s: PlayerId=%s"), TEXT("FPlayerSession::~FPlayerSession"), *PlayerId);
 	if (DataChannel)
 		DataChannel->UnregisterObserver();
 }
@@ -104,7 +104,7 @@ void FPlayerSession::OnRemoteIceCandidate(TUniquePtr<webrtc::IceCandidateInterfa
 {
 	if (!PeerConnection->AddIceCandidate(Candidate.Get()))
 	{
-		UE_LOG(PixelStreamer, Error, TEXT("Failed to apply remote ICE Candidate from Player %u"), PlayerId);
+		UE_LOG(PixelStreamer, Error, TEXT("Failed to apply remote ICE Candidate from Player %s"), *PlayerId);
 		DisconnectPlayer(TEXT("Failed to apply remote ICE Candidate"));
 		return;
 	}
@@ -125,67 +125,67 @@ void FPlayerSession::DisconnectPlayer(const FString& Reason)
 
 void FPlayerSession::OnSignalingChange(webrtc::PeerConnectionInterface::SignalingState NewState)
 {
-	UE_LOG(PixelStreamer, Log, TEXT("%s : PlayerId=%u, NewState=%s"), TEXT("FPlayerSession::OnSignalingChange"), PlayerId, ToString(NewState));
+	UE_LOG(PixelStreamer, Log, TEXT("%s : PlayerId=%s, NewState=%s"), TEXT("FPlayerSession::OnSignalingChange"), *PlayerId, ToString(NewState));
 }
 
 // Called when a remote stream is added
 void FPlayerSession::OnAddStream(rtc::scoped_refptr<webrtc::MediaStreamInterface> Stream)
 {
-	UE_LOG(PixelStreamer, Log, TEXT("%s : PlayerId=%u, Stream=%s"), TEXT("FPlayerSession::OnAddStream"), PlayerId, *ToString(Stream->id()));
+	UE_LOG(PixelStreamer, Log, TEXT("%s : PlayerId=%s, Stream=%s"), TEXT("FPlayerSession::OnAddStream"), *PlayerId, *ToString(Stream->id()));
 }
 
 void FPlayerSession::OnRemoveStream(rtc::scoped_refptr<webrtc::MediaStreamInterface> Stream)
 {
-	UE_LOG(PixelStreamer, Log, TEXT("%s : PlayerId=%u, Stream=%s"), TEXT("FPlayerSession::OnRemoveStream"), PlayerId, *ToString(Stream->id()));
+	UE_LOG(PixelStreamer, Log, TEXT("%s : PlayerId=%s, Stream=%s"), TEXT("FPlayerSession::OnRemoveStream"), *PlayerId, *ToString(Stream->id()));
 }
 
 void FPlayerSession::OnDataChannel(rtc::scoped_refptr<webrtc::DataChannelInterface> InDataChannel)
 {
-	UE_LOG(PixelStreamer, Log, TEXT("%s : PlayerId=%u"), TEXT("FPlayerSession::OnDataChannel"), PlayerId);
+	UE_LOG(PixelStreamer, Log, TEXT("%s : PlayerId=%s"), TEXT("FPlayerSession::OnDataChannel"), *PlayerId);
 	DataChannel = InDataChannel;
 	DataChannel->RegisterObserver(this);
 }
 
 void FPlayerSession::OnRenegotiationNeeded()
 {
-	UE_LOG(PixelStreamer, Log, TEXT("%s : PlayerId=%u"), TEXT("FPlayerSession::OnRenegotiationNeeded"), PlayerId);
+	UE_LOG(PixelStreamer, Log, TEXT("%s : PlayerId=%s"), TEXT("FPlayerSession::OnRenegotiationNeeded"), *PlayerId);
 }
 
 void FPlayerSession::OnIceConnectionChange(webrtc::PeerConnectionInterface::IceConnectionState NewState)
 {
-	UE_LOG(PixelStreamer, Log, TEXT("%s : PlayerId=%u, NewState=%s"), TEXT("FPlayerSession::OnIceConnectionChange"), PlayerId, ToString(NewState));
+	UE_LOG(PixelStreamer, Log, TEXT("%s : PlayerId=%s, NewState=%s"), TEXT("FPlayerSession::OnIceConnectionChange"), *PlayerId, ToString(NewState));
 }
 
 void FPlayerSession::OnIceGatheringChange(webrtc::PeerConnectionInterface::IceGatheringState NewState)
 {
-	UE_LOG(PixelStreamer, Log, TEXT("%s : PlayerId=%u, NewState=%s"), TEXT("FPlayerSession::OnIceGatheringChange"), PlayerId, ToString(NewState));
+	UE_LOG(PixelStreamer, Log, TEXT("%s : PlayerId=%s, NewState=%s"), TEXT("FPlayerSession::OnIceGatheringChange"), *PlayerId, ToString(NewState));
 }
 
 void FPlayerSession::OnIceCandidate(const webrtc::IceCandidateInterface* Candidate)
 {
-	UE_LOG(PixelStreamer, Log, TEXT("%s : PlayerId=%u"), TEXT("FPlayerSession::OnIceCandidate"), PlayerId);
+	UE_LOG(PixelStreamer, Log, TEXT("%s : PlayerId=%s"), TEXT("FPlayerSession::OnIceCandidate"), *PlayerId);
 
 	Streamer.GetSignallingServerConnection()->SendIceCandidate(PlayerId, *Candidate);
 }
 
 void FPlayerSession::OnIceCandidatesRemoved(const std::vector<cricket::Candidate>& candidates)
 {
-	UE_LOG(PixelStreamer, Log, TEXT("%s : PlayerId=%u"), TEXT("FPlayerSession::OnIceCandidatesRemoved"), PlayerId);
+	UE_LOG(PixelStreamer, Log, TEXT("%s : PlayerId=%s"), TEXT("FPlayerSession::OnIceCandidatesRemoved"), *PlayerId);
 }
 
 void FPlayerSession::OnIceConnectionReceivingChange(bool Receiving)
 {
-	UE_LOG(PixelStreamer, Log, TEXT("%s : PlayerId=%u, Receiving=%d"), TEXT("FPlayerSession::OnIceConnectionReceivingChange"), PlayerId, *reinterpret_cast<int8*>(&Receiving));
+	UE_LOG(PixelStreamer, Log, TEXT("%s : PlayerId=%s, Receiving=%d"), TEXT("FPlayerSession::OnIceConnectionReceivingChange"), *PlayerId, *reinterpret_cast<int8*>(&Receiving));
 }
 
 void FPlayerSession::OnTrack(rtc::scoped_refptr<webrtc::RtpTransceiverInterface> transceiver)
 {
-	UE_LOG(PixelStreamer, Log, TEXT("%s : PlayerId=%u"), TEXT("FPlayerSession::OnTrack"), PlayerId);
+	UE_LOG(PixelStreamer, Log, TEXT("%s : PlayerId=%s"), TEXT("FPlayerSession::OnTrack"), *PlayerId);
 }
 
 void FPlayerSession::OnRemoveTrack(rtc::scoped_refptr<webrtc::RtpReceiverInterface> receiver)
 {
-	UE_LOG(PixelStreamer, Log, TEXT("%s : PlayerId=%u"), TEXT("FPlayerSession::OnRemoveTrack"), PlayerId);
+	UE_LOG(PixelStreamer, Log, TEXT("%s : PlayerId=%s"), TEXT("FPlayerSession::OnRemoveTrack"), *PlayerId);
 }
 
 //
@@ -265,7 +265,7 @@ void FPlayerSession::SendFreezeFrame(const TArray64<uint8>& JpegBytes)
 		return;
 	}
 
-	UE_LOG(PixelStreamer, VeryVerbose, TEXT("player %d: %s"), PlayerId, TEXT("FPlayerSession::SendFreezeFrame"));
+	UE_LOG(PixelStreamer, VeryVerbose, TEXT("player %s: %s"), *PlayerId, TEXT("FPlayerSession::SendFreezeFrame"));
 
 	// just a sanity check. WebRTC buffer size is 16MB, which translates to 3840Mbps for 30fps video. It's not expected
 	// that freeze frame size will come close to this limit and it's not expected we'll send freeze frames too often.
@@ -295,7 +295,7 @@ void FPlayerSession::SendUnfreezeFrame()
 		return;
 	}
 
-	UE_LOG(PixelStreamer, VeryVerbose, TEXT("player %d: %s"), PlayerId, TEXT("FPlayerSession::SendUnfreezeFrame"));
+	UE_LOG(PixelStreamer, VeryVerbose, TEXT("player %s: %s"), *PlayerId, TEXT("FPlayerSession::SendUnfreezeFrame"));
 
 	rtc::CopyOnWriteBuffer Buffer(2);
 	Buffer[0] = static_cast<uint8_t>(PixelStreamingProtocol::EToPlayerMsg::UnfreezeFrame);
@@ -310,7 +310,7 @@ void FPlayerSession::OnNewSecondarySession()
 
 void FPlayerSession::OnBufferedAmountChange(uint64_t PreviousAmount)
 {
-	UE_LOG(PixelStreamer, VeryVerbose, TEXT("player %d: OnBufferedAmountChanged: prev %d, cur %d"), PlayerId, PreviousAmount, DataChannel->buffered_amount());
+	UE_LOG(PixelStreamer, VeryVerbose, TEXT("player %s: OnBufferedAmountChanged: prev %d, cur %d"), *PlayerId, PreviousAmount, DataChannel->buffered_amount());
 }
 
 void FPlayerSession::SetVideoEncoder(FPixelStreamingVideoEncoder* InVideoEncoder)
diff --git a/Engine/Plugins/Media/PixelStreaming/Source/PixelStreaming/Private/SignallingServerConnection.cpp b/Engine/Plugins/Media/PixelStreaming/Source/PixelStreaming/Private/SignallingServerConnection.cpp
--- a/Engine/Plugins/Media/PixelStreaming/Source/PixelStreaming/Private/SignallingServerConnection.cpp
+++ b/Engine/Plugins/Media/PixelStreaming/Source/PixelStreaming/Private/SignallingServerConnection.cpp
@@ -29,14 +29,14 @@ DEFINE_LOG_CATEGORY(LogPixelStreamingSS);
 #define HANDLE_PLAYER_SS_ERROR(PlayerId, ErrorMsg, ...)\
 	do\
 	{\
-		UE_LOG(LogPixelStreamingSS, Error, TEXT("player %d: ") ErrorMsg, PlayerId, ##__VA_ARGS__);\
+		UE_LOG(LogPixelStreamingSS, Error, TEXT("player %s: ") ErrorMsg, *PlayerId, ##__VA_ARGS__);\
 		SendDisconnectPlayer(PlayerId, FString::Printf(ErrorMsg, ##__VA_ARGS__));\
 		return;\
 	}\
 	while(false);
 
-FSignallingServerConnection::FSignallingServerConnection(const FString& Url, FSignallingServerConnectionObserver& InObserver)
-	: Observer(InObserver)
+FSignallingServerConnection::FSignallingServerConnection(const FString& Url, FSignallingServerConnectionObserver& InObserver, const FString& InStreamerId)
+	: Observer(InObserver), StreamerId(InStreamerId)
 {
 	WS = FWebSocketsModule::Get().CreateWebSocket(Url, TEXT("ws"));
 
@@ -79,11 +79,11 @@ void FSignallingServerConnection::SendOffer(const webrtc::SessionDescriptionInte
 	WS->Send(ToString(OfferJson, false));
 }
 
-void FSignallingServerConnection::SendAnswer(uint32 PlayerId, const webrtc::SessionDescriptionInterface& SDP)
+void FSignallingServerConnection::SendAnswer(FPlayerId PlayerId, const webrtc::SessionDescriptionInterface& SDP)
 {
 	auto AnswerJson = MakeShared<FJsonObject>();
 	AnswerJson->SetStringField(TEXT("type"), TEXT("answer"));
-	AnswerJson->SetNumberField(TEXT("playerId"), PlayerId);
+	AnswerJson->SetStringField(TEXT("playerId"), PlayerId);
 
 	std::string SdpAnsi;
 	verifyf(SDP.ToString(&SdpAnsi), TEXT("Failed to serialise local SDP"));
@@ -117,12 +117,12 @@ void FSignallingServerConnection::SendIceCandidate(const webrtc::IceCandidateInt
 	WS->Send(ToString(IceCandidateJson, false));
 }
 
-void FSignallingServerConnection::SendIceCandidate(uint32 PlayerId, const webrtc::IceCandidateInterface& IceCandidate)
+void FSignallingServerConnection::SendIceCandidate(FPlayerId PlayerId, const webrtc::IceCandidateInterface& IceCandidate)
 {
 	auto IceCandidateJson = MakeShared<FJsonObject>();
 
 	IceCandidateJson->SetStringField(TEXT("type"), TEXT("iceCandidate"));
-	IceCandidateJson->SetNumberField(TEXT("playerId"), PlayerId);
+	IceCandidateJson->SetStringField(TEXT("playerId"), PlayerId);
 
 	auto CandidateJson = MakeShared<FJsonObject>();
 	CandidateJson->SetStringField(TEXT("sdpMid"), IceCandidate.sdp_mid().c_str());
@@ -153,12 +153,12 @@ void FSignallingServerConnection::KeepAlive()
 	}
 }
 
-void FSignallingServerConnection::SendDisconnectPlayer(uint32 PlayerId, const FString& Reason)
+void FSignallingServerConnection::SendDisconnectPlayer(FPlayerId PlayerId, const FString& Reason)
 {
 	auto Json = MakeShared<FJsonObject>();
 
 	Json->SetStringField(TEXT("type"), TEXT("disconnectPlayer"));
-	Json->SetNumberField(TEXT("playerId"), PlayerId);
+	Json->SetStringField(TEXT("playerId"), PlayerId);
 	Json->SetStringField(TEXT("reason"), Reason);
 
 	FString Msg = ToString(Json, false);
@@ -205,7 +205,11 @@ void FSignallingServerConnection::OnMessage(const FString& Msg)
 		HANDLE_SS_ERROR(TEXT("Cannot find `type` field in SS message:\n%s"), *Msg);
 	}
 
-	if (MsgType == TEXT("config"))
+	if (MsgType == TEXT("identify"))
+	{
+		OnIdRequested();
+	}
+	else if (MsgType == TEXT("config"))
 	{
 		OnConfig(JsonMsg);
 	}
@@ -243,6 +247,19 @@ void FSignallingServerConnection::OnMessage(const FString& Msg)
 	}
 }
 
+void FSignallingServerConnection::OnIdRequested()
+{
+	auto Json = MakeShared<FJsonObject>();
+
+	Json->SetStringField(TEXT("type"), TEXT("endpointId"));
+	Json->SetStringField(TEXT("id"), StreamerId);
+
+	FString Msg = ToString(Json, false);
+	UE_LOG(LogPixelStreamingSS, Verbose, TEXT("-> SS: endpointId\n%s"), *Msg);
+
+	WS->Send(Msg);
+}
+
 void FSignallingServerConnection::OnConfig(const FJsonObjectPtr& Json)
 {
 	// SS sends `config` that looks like:
@@ -322,10 +339,10 @@ void FSignallingServerConnection::OnSessionDescription(const FJsonObjectPtr& Jso
 
 	if (Type == webrtc::SdpType::kOffer)
 	{
-		uint32 PlayerId;
-		if (!Json->TryGetNumberField(TEXT("playerId"), PlayerId))
+		FPlayerId PlayerId;
+		if (!Json->TryGetStringField(TEXT("playerId"), PlayerId))
 		{
-			HANDLE_SS_ERROR(TEXT("Failed to get `playerId`from `offer` message\n%s"), *ToString(Json));
+			HANDLE_SS_ERROR(TEXT("Failed to get `playerId` from `offer` message\n%s"), *ToString(Json));
 		}
 
 		Observer.OnOffer(PlayerId, TUniquePtr<webrtc::SessionDescriptionInterface>{SessionDesc.release()});
@@ -374,7 +391,7 @@ void FSignallingServerConnection::OnStreamerIceCandidate(const FJsonObjectPtr& J
 
 void FSignallingServerConnection::OnPlayerIceCandidate(const FJsonObjectPtr& Json)
 {
-	uint32 PlayerId = Json->GetNumberField(TEXT("playerId"));
+	FPlayerId PlayerId = Json->GetStringField(TEXT("playerId"));
 
 	const FJsonObjectPtr* CandidateJson;
 	if (!Json->TryGetObjectField(TEXT("candidate"), CandidateJson))
@@ -423,10 +440,10 @@ void FSignallingServerConnection::OnPlayerCount(const FJsonObjectPtr& Json)
 
 void FSignallingServerConnection::OnPlayerDisconnected(const FJsonObjectPtr& Json)
 {
-	uint32 PlayerId;
-	if (!Json->TryGetNumberField(TEXT("playerId"), PlayerId))
+	FPlayerId PlayerId;
+	if (!Json->TryGetStringField(TEXT("playerId"), PlayerId))
 	{
-		HANDLE_SS_ERROR(TEXT("Failed to get `playerId`from `playerDisconnected` message\n%s"), *ToString(Json));
+		HANDLE_SS_ERROR(TEXT("Failed to get `playerId` from `playerDisconnected` message\n%s"), *ToString(Json));
 	}
 
 	Observer.OnPlayerDisconnected(PlayerId);
diff --git a/Engine/Plugins/Media/PixelStreaming/Source/PixelStreaming/Private/SignallingServerConnection.h b/Engine/Plugins/Media/PixelStreaming/Source/PixelStreaming/Private/SignallingServerConnection.h
--- a/Engine/Plugins/Media/PixelStreaming/Source/PixelStreaming/Private/SignallingServerConnection.h
+++ b/Engine/Plugins/Media/PixelStreaming/Source/PixelStreaming/Private/SignallingServerConnection.h
@@ -24,11 +24,11 @@ public:
 	virtual void OnConfig(const webrtc::PeerConnectionInterface::RTCConfiguration& Config) = 0;
 
 	// Streamer-only
-	virtual void OnOffer(uint32 PlayerId, TUniquePtr<webrtc::SessionDescriptionInterface> Sdp)
+	virtual void OnOffer(FPlayerId PlayerId, TUniquePtr<webrtc::SessionDescriptionInterface> Sdp)
 	{ unimplemented(); }
-	virtual void OnRemoteIceCandidate(uint32 PlayerId, TUniquePtr<webrtc::IceCandidateInterface> Candidate)
+	virtual void OnRemoteIceCandidate(FPlayerId PlayerId, TUniquePtr<webrtc::IceCandidateInterface> Candidate)
 	{ unimplemented(); }
-	virtual void OnPlayerDisconnected(uint32 PlayerId)
+	virtual void OnPlayerDisconnected(FPlayerId PlayerId)
 	{ unimplemented(); }
 
 	// Player-only
@@ -36,21 +36,21 @@ public:
 	{ unimplemented(); }
 	virtual void OnRemoteIceCandidate(TUniquePtr<webrtc::IceCandidateInterface> Candidate)
 	{ unimplemented(); }
-	virtual void OnPlayerCount(uint32 PlayerId)
+	virtual void OnPlayerCount(uint32 Count)
 	{ unimplemented(); }
 };
 
 class FSignallingServerConnection final
 {
 public:
-	explicit FSignallingServerConnection(const FString& Url, FSignallingServerConnectionObserver& Observer);
+	explicit FSignallingServerConnection(const FString& Url, FSignallingServerConnectionObserver& Observer, const FString& StreamerId);
 	~FSignallingServerConnection();
 
 	void SendOffer(const webrtc::SessionDescriptionInterface& SDP);
-	void SendAnswer(uint32 PlayerId, const webrtc::SessionDescriptionInterface& SDP);
+	void SendAnswer(FPlayerId PlayerId, const webrtc::SessionDescriptionInterface& SDP);
 	void SendIceCandidate(const webrtc::IceCandidateInterface& IceCandidate);
-	void SendIceCandidate(uint32 PlayerId, const webrtc::IceCandidateInterface& IceCandidate);
-	void SendDisconnectPlayer(uint32 PlayerId, const FString& Reason);
+	void SendIceCandidate(FPlayerId PlayerId, const webrtc::IceCandidateInterface& IceCandidate);
+	void SendDisconnectPlayer(FPlayerId PlayerId, const FString& Reason);
 
 private:
 	void KeepAlive();
@@ -62,6 +62,7 @@ private:
 
 	using FJsonObjectPtr = TSharedPtr<FJsonObject>;
 
+	void OnIdRequested();
 	void OnConfig(const FJsonObjectPtr& Json);
 	void OnSessionDescription(const FJsonObjectPtr& Json);
 	void OnStreamerIceCandidate(const FJsonObjectPtr& Json);
@@ -71,6 +72,7 @@ private:
 
 private:
 	FSignallingServerConnectionObserver& Observer;
+	FString StreamerId;
 
 	FDelegateHandle OnConnectedHandle;
 	FDelegateHandle OnConnectionErrorHandle;
diff --git a/Engine/Plugins/Media/PixelStreaming/Source/PixelStreaming/Private/Streamer.cpp b/Engine/Plugins/Media/PixelStreaming/Source/PixelStreaming/Private/Streamer.cpp
--- a/Engine/Plugins/Media/PixelStreaming/Source/PixelStreaming/Private/Streamer.cpp
+++ b/Engine/Plugins/Media/PixelStreaming/Source/PixelStreaming/Private/Streamer.cpp
@@ -41,8 +41,8 @@ bool FStreamer::CheckPlatformCompatibility()
 	return AVEncoder::FVideoEncoderFactory::Get().HasEncoderForCodec(AVEncoder::ECodecType::H264);
 }
 
-FStreamer::FStreamer(const FString& InSignallingServerUrl)
-	: SignallingServerUrl(InSignallingServerUrl)
+FStreamer::FStreamer(const FString& InSignallingServerUrl, const FString& InStreamerId)
+	: SignallingServerUrl(InSignallingServerUrl), StreamerId(InStreamerId)
 {
 	RedirectWebRtcLogsToUnreal(rtc::LoggingSeverity::LS_VERBOSE);
 
@@ -171,7 +171,7 @@ void FStreamer::WebRtcSignallingThreadFunc()
 
 void FStreamer::ConnectToSignallingServer()
 {
-	SignallingServerConnection = MakeUnique<FSignallingServerConnection>(SignallingServerUrl, *this);
+	SignallingServerConnection = MakeUnique<FSignallingServerConnection>(SignallingServerUrl, *this, StreamerId);
 }
 
 void FStreamer::OnFrameBufferReady(const FTexture2DRHIRef& FrameBuffer)
@@ -192,7 +192,7 @@ void FStreamer::OnOffer(FPlayerId PlayerId, TUniquePtr<webrtc::SessionDescriptio
 	AddStreams(PlayerId);
 
 	FPlayerSession* Player = GetPlayerSession(PlayerId);
-	checkf(Player, TEXT("just created player %d not found"), PlayerId);
+	checkf(Player, TEXT("just created player %s not found"), *PlayerId);
 
 	Player->OnOffer(MoveTemp(Sdp));
 
@@ -206,13 +206,13 @@ void FStreamer::OnOffer(FPlayerId PlayerId, TUniquePtr<webrtc::SessionDescriptio
 void FStreamer::OnRemoteIceCandidate(FPlayerId PlayerId, TUniquePtr<webrtc::IceCandidateInterface> Candidate)
 {
 	FPlayerSession* Player = GetPlayerSession(PlayerId);
-	checkf(Player, TEXT("player %u not found"), PlayerId);
+	checkf(Player, TEXT("player %s not found"), *PlayerId);
 	Player->OnRemoteIceCandidate(MoveTemp(Candidate));
 }
 
 void FStreamer::OnPlayerDisconnected(FPlayerId PlayerId)
 {
-	UE_LOG(PixelStreamer, Log, TEXT("player %d disconnected"), PlayerId);
+	UE_LOG(PixelStreamer, Log, TEXT("player %s disconnected"), *PlayerId);
 	DeletePlayerSession(PlayerId);
 }
 
@@ -254,7 +254,7 @@ void FStreamer::CreatePlayerSession(FPlayerId PlayerId)
 		}
 	}
 
-	UE_LOG(PixelStreamer, Log, TEXT("Creating player session for PlayerId=%d"), PlayerId);
+	UE_LOG(PixelStreamer, Log, TEXT("Creating player session for PlayerId=%s"), *PlayerId);
 	
 	// this is called from WebRTC signalling thread, the only thread were `Players` map is modified, so no need to lock it
 	bool bOriginalQualityController = Players.Num() == 0; // first player controls quality by default
@@ -274,7 +274,7 @@ void FStreamer::DeletePlayerSession(FPlayerId PlayerId)
 	FPlayerSession* Player = GetPlayerSession(PlayerId);
 	if (!Player)
 	{
-		UE_LOG(PixelStreamer, VeryVerbose, TEXT("failed to delete player %d: not found"), PlayerId);
+		UE_LOG(PixelStreamer, VeryVerbose, TEXT("failed to delete player %s: not found"), *PlayerId);
 		return;
 	}
 
@@ -309,8 +309,8 @@ void FStreamer::DeletePlayerSession(FPlayerId PlayerId)
 void FStreamer::AddStreams(FPlayerId PlayerId)
 {
 	FString const StreamId = TEXT("stream_id");
-	FString const AudioLabel = FString::Printf(TEXT("audio_label_%d"), PlayerId);
-	FString const VideoLabel= FString::Printf(TEXT("video_label_%d"), PlayerId);
+	FString const AudioLabel = FString::Printf(TEXT("audio_label_%s"), *PlayerId);
+	FString const VideoLabel= FString::Printf(TEXT("video_label_%s"), *PlayerId);
 
 	FPlayerSession* Session = GetPlayerSession(PlayerId);
 	check(Session);
@@ -337,8 +337,8 @@ void FStreamer::AddStreams(FPlayerId PlayerId)
 	{
 		UE_LOG(PixelStreamer,
 			   Error,
-			   TEXT("Failed to add AudioTrack to PeerConnection of player %u. Msg=%s"),
-			   Session->GetPlayerId(),
+			   TEXT("Failed to add AudioTrack to PeerConnection of player %s. Msg=%s"),
+			   *Session->GetPlayerId(),
 			   TCHAR_TO_UTF8(addAudioTrackResult.error().message()));
 	}
 
@@ -347,8 +347,8 @@ void FStreamer::AddStreams(FPlayerId PlayerId)
 	{
 		UE_LOG(PixelStreamer,
 			   Error,
-			   TEXT("Failed to add VideoTrack to PeerConnection of player %u. Msg=%s"),
-			   Session->GetPlayerId(),
+			   TEXT("Failed to add VideoTrack to PeerConnection of player %s. Msg=%s"),
+			   *Session->GetPlayerId(),
 				TCHAR_TO_UTF8(addVideoTrackResult.error().message()));
 	}
 	else
@@ -370,7 +370,7 @@ void FStreamer::AddStreams(FPlayerId PlayerId)
 
 void FStreamer::OnQualityOwnership(FPlayerId PlayerId)
 {
-	checkf(GetPlayerSession(PlayerId), TEXT("player %d not found"), PlayerId);
+	checkf(GetPlayerSession(PlayerId), TEXT("player %s not found"), *PlayerId);
 	{
 		FScopeLock PlayersLock(&PlayersCS);
 		for (auto&& PlayerEntry : Players)
@@ -412,7 +412,7 @@ void FStreamer::SendCachedFreezeFrameTo(FPlayerSession& Player)
 {
 	if (CachedJpegBytes.Num() > 0)
 	{
-		UE_LOG(PixelStreamer, Log, TEXT("Sending cached freeze frame to player %d: %d bytes"), Player.GetPlayerId(), CachedJpegBytes.Num());
+		UE_LOG(PixelStreamer, Log, TEXT("Sending cached freeze frame to player %s: %d bytes"), *Player.GetPlayerId(), CachedJpegBytes.Num());
 		Player.SendFreezeFrame(CachedJpegBytes);
 	}
 }
diff --git a/Engine/Plugins/Media/PixelStreaming/Source/PixelStreaming/Private/Streamer.h b/Engine/Plugins/Media/PixelStreaming/Source/PixelStreaming/Private/Streamer.h
--- a/Engine/Plugins/Media/PixelStreaming/Source/PixelStreaming/Private/Streamer.h
+++ b/Engine/Plugins/Media/PixelStreaming/Source/PixelStreaming/Private/Streamer.h
@@ -14,7 +14,7 @@ class FStreamer : public FSignallingServerConnectionObserver
 public:
 	static bool CheckPlatformCompatibility();
 
-	explicit FStreamer(const FString& SignallingServerUrl);
+	explicit FStreamer(const FString& SignallingServerUrl, const FString& StreamerId);
 	virtual ~FStreamer() override;
 
 	void SendPlayerMessage(PixelStreamingProtocol::EToPlayerMsg Type, const FString& Descriptor);
@@ -53,6 +53,7 @@ private:
 
 private:
 	FString SignallingServerUrl;
+	FString StreamerId;
 
 #if PLATFORM_WINDOWS
 	DWORD WebRtcSignallingThreadId = 0;
diff --git a/Engine/Plugins/Media/PixelStreaming/Source/PixelStreaming/Private/StreamerConnection.cpp b/Engine/Plugins/Media/PixelStreaming/Source/PixelStreaming/Private/StreamerConnection.cpp
--- a/Engine/Plugins/Media/PixelStreaming/Source/PixelStreaming/Private/StreamerConnection.cpp
+++ b/Engine/Plugins/Media/PixelStreaming/Source/PixelStreaming/Private/StreamerConnection.cpp
@@ -107,7 +107,7 @@ void FStreamerConnection::SignallingThreadFunc()
 #endif
 
 	// now that everything is ready connect to SignallingServer
-	SignallingServerConnection = MakeUnique<FSignallingServerConnection>(SignallingServerAddress, *this);
+	SignallingServerConnection = MakeUnique<FSignallingServerConnection>(SignallingServerAddress, *this, TEXT(""));
 
 #if PLATFORM_WINDOWS
 	// WebRTC window messaging loop
diff --git a/Engine/Source/Programs/PixelStreaming/WebServers/SignallingWebServer/cirrus.js b/Engine/Source/Programs/PixelStreaming/WebServers/SignallingWebServer/cirrus.js
--- a/Engine/Source/Programs/PixelStreaming/WebServers/SignallingWebServer/cirrus.js
+++ b/Engine/Source/Programs/PixelStreaming/WebServers/SignallingWebServer/cirrus.js
@@ -350,7 +350,7 @@ playerServer.on('connection', function (ws, req) {
 		return;
 	}
 
-	let playerId = ++nextPlayerId;
+	let playerId = (++nextPlayerId).toString();
 	console.log(`player ${playerId} (${req.connection.remoteAddress}) connected`);
 	players.set(playerId, { ws: ws, id: playerId });
 
