fix for deprecation warning in PixelStreamingModule.cpp

diff --git a/Engine/Plugins/Media/PixelStreaming/Source/PixelStreaming/Private/PixelStreamingModule.cpp b/Engine/Plugins/Media/PixelStreaming/Source/PixelStreaming/Private/PixelStreamingModule.cpp
--- a/Engine/Plugins/Media/PixelStreaming/Source/PixelStreaming/Private/PixelStreamingModule.cpp
+++ b/Engine/Plugins/Media/PixelStreaming/Source/PixelStreaming/Private/PixelStreamingModule.cpp
@@ -311,7 +311,8 @@ void FPixelStreamingModule::FreezeFrame(UTexture2D* Texture)
 	if (Texture)
 	{
 		// A frame is supplied so immediately read its data and send as a JPEG.
-		FTexture2DRHIRef Texture2DRHI = Texture->Resource && Texture->Resource->TextureRHI ? Texture->Resource->TextureRHI->GetTexture2D() : nullptr;
+		FTextureResource* TextureResource = Texture->GetResource();
+		FTexture2DRHIRef Texture2DRHI = TextureResource && TextureResource->TextureRHI ? TextureResource->TextureRHI->GetTexture2D() : nullptr;
 		if (!Texture2DRHI)
 		{
 			UE_LOG(PixelStreamer, Error, TEXT("Attempting freeze frame with texture %s with no texture 2D RHI"), *Texture->GetName());
