PR #8041: Correct parse of RTC iceServer configuration in pixel stream (Contributed by [REDACTED_PERSON])

campbellwmorgan[REDACTED_HANDLE]

diff --git a/Engine/Plugins/Media/PixelStreaming/Source/PixelStreaming/Private/SignallingServerConnection.cpp b/Engine/Plugins/Media/PixelStreaming/Source/PixelStreaming/Private/SignallingServerConnection.cpp
--- a/Engine/Plugins/Media/PixelStreaming/Source/PixelStreaming/Private/SignallingServerConnection.cpp
+++ b/Engine/Plugins/Media/PixelStreaming/Source/PixelStreaming/Private/SignallingServerConnection.cpp
@@ -298,6 +298,16 @@ void FSignallingServerConnection::OnConfig(const FJsonObjectPtr& Json)
 					IceServer.urls.push_back(to_string(Url));
 				}
 			}
+			else
+			{
+				// in the RTC Spec, "urls" can be an array or a single string
+				// https://www.w3.org/TR/webrtc/#dictionary-rtciceserver-members
+				FString UrlsSingle;
+				if ((*IceServerJson)->TryGetStringField(TEXT("urls"), UrlsSingle))
+				{
+					IceServer.urls.push_back(to_string(UrlsSingle));
+				}
+			}
 
 			FString Username;
 			if ((*IceServerJson)->TryGetStringField(TEXT("username"), Username))
