Merge in PixelStreaming, AVEncoder, and SignallingServer changes from TensorWorks local repository
Brings in multiple improvements such as parameterisation of many encoder settings, AMF encoder, better pixel streaming stats and some bugfixes

diff --git a/Engine/Plugins/Media/PixelStreaming/PixelStreaming.uplugin b/Engine/Plugins/Media/PixelStreaming/PixelStreaming.uplugin
--- a/Engine/Plugins/Media/PixelStreaming/PixelStreaming.uplugin
+++ b/Engine/Plugins/Media/PixelStreaming/PixelStreaming.uplugin
@@ -19,24 +19,7 @@
 			"Name": "PixelStreaming",
 			"Type": "RuntimeNoCommandlet",
 			"LoadingPhase": "PostDefault",
-			"WhitelistPlatforms": [ "Win64" ],
-			"BlacklistPlatforms": [ "Linux" ],
-			"BlacklistTargets": [ "Server" ]
-		},
-		{
-			"Name": "PixelStreamingPlayerFactory",
-			"Type": "Editor",
-			"LoadingPhase": "PostDefault",
-			"WhitelistPlatforms": [ "Win64" ],
-			"BlacklistPlatforms": [ "Linux" ],
-			"BlacklistTargets": [ "Server" ]
-		},
-		{
-			"Name": "PixelStreamingPlayerFactory",
-			"Type": "RuntimeNoCommandlet",
-			"LoadingPhase": "PostDefault",
-			"WhitelistPlatforms": [ "Win64" ],
-			"BlacklistPlatforms": [ "Linux" ],
+			"WhitelistPlatforms": [ "Win32", "Win64" ],
 			"BlacklistTargets": [ "Server" ]
 		}
 	],
@@ -44,7 +27,6 @@
 		{
 			"Name": "WmfMedia",
 			"Enabled": true,
-			"WhitelistPlatforms": [ "Win64" ],
 			"BlacklistPlatforms": [ "Linux" ]
 		}
 	]
diff --git a/Engine/Plugins/Media/PixelStreaming/Source/PixelStreaming/PixelStreaming.Build.cs b/Engine/Plugins/Media/PixelStreaming/Source/PixelStreaming/PixelStreaming.Build.cs
--- a/Engine/Plugins/Media/PixelStreaming/Source/PixelStreaming/PixelStreaming.Build.cs
+++ b/Engine/Plugins/Media/PixelStreaming/Source/PixelStreaming/PixelStreaming.Build.cs
@@ -8,177 +8,142 @@ using EpicGames.Core;
 
 namespace UnrealBuildTool.Rules
 {
-    public class PixelStreaming : ModuleRules
-    {
-        private void AddSignallingServer()
-        {
-            string PixelStreamingProgramsDirectory = "./Programs/PixelStreaming";
-            string SignallingServerDir = new DirectoryInfo(PixelStreamingProgramsDirectory + "/WebServers/SignallingWebServer").FullName;
-
-            if (!Directory.Exists(SignallingServerDir))
-            {
-                Log.TraceInformation(string.Format("Signalling Server path '{0}' does not exist", SignallingServerDir));
-                return;
-            }
-
-            List<string> DependenciesToAdd = new List<string>();
-            DependenciesToAdd.AddRange(Directory.GetFiles(SignallingServerDir, "*.*", SearchOption.AllDirectories));
-
-            string NodeModulesDirPath = new DirectoryInfo(SignallingServerDir + "/node_modules").FullName;
-            string LogsDirPath = new DirectoryInfo(SignallingServerDir + "/logs").FullName;
-            foreach (string Dependency in DependenciesToAdd)
-            {
-                if (!Dependency.StartsWith(NodeModulesDirPath) &&
-                    !Dependency.StartsWith(LogsDirPath))
-                {
-                    RuntimeDependencies.Add(Dependency, StagedFileType.NonUFS);
-                }
-            }
-        }
-
-        private void AddMatchmakingServer()
-        {
+	public class PixelStreaming : ModuleRules
+	{
+		private void AddSignallingServer()
+		{
 			string PixelStreamingProgramsDirectory = "./Programs/PixelStreaming";
-            string MatchmakingServerDir = new DirectoryInfo(PixelStreamingProgramsDirectory + "/WebServers/Matchmaker").FullName;
-
-            if (!Directory.Exists(MatchmakingServerDir))
-            {
-                Log.TraceInformation(string.Format("Matchmaking Server path '{0}' does not exist", MatchmakingServerDir));
-                return;
-            }
-
-            List<string> DependenciesToAdd = new List<string>();
-            DependenciesToAdd.AddRange(Directory.GetFiles(MatchmakingServerDir, "*.*", SearchOption.AllDirectories));
-
-            string NodeModulesDirPath = new DirectoryInfo(MatchmakingServerDir + "/node_modules").FullName;
-            string LogsDirPath = new DirectoryInfo(MatchmakingServerDir + "/logs").FullName;
-            foreach (string Dependency in DependenciesToAdd)
-            {
-                if (!Dependency.StartsWith(NodeModulesDirPath) &&
-                    !Dependency.StartsWith(LogsDirPath))
-                {
-                    RuntimeDependencies.Add(Dependency, StagedFileType.NonUFS);
-                }
-            }
-        }
+			string SignallingServerDir = new DirectoryInfo(PixelStreamingProgramsDirectory + "/WebServers/SignallingWebServer").FullName;
+
+			if (!Directory.Exists(SignallingServerDir))
+			{
+				Log.TraceInformation(string.Format("Signalling Server path '{0}' does not exist", SignallingServerDir));
+				return;
+			}
+
+			List<string> DependenciesToAdd = new List<string>();
+			DependenciesToAdd.AddRange(Directory.GetFiles(SignallingServerDir, "*.*", SearchOption.AllDirectories));
+
+			string NodeModulesDirPath = new DirectoryInfo(SignallingServerDir + "/node_modules").FullName;
+			string LogsDirPath = new DirectoryInfo(SignallingServerDir + "/logs").FullName;
+			foreach (string Dependency in DependenciesToAdd)
+			{
+				if (!Dependency.StartsWith(NodeModulesDirPath) &&
+					!Dependency.StartsWith(LogsDirPath))
+				{
+					RuntimeDependencies.Add(Dependency, StagedFileType.NonUFS);
+				}
+			}
+		}
+
+		private void AddMatchmakingServer()
+		{
+			string PixelStreamingProgramsDirectory = "./Programs/PixelStreaming";
+			string MatchmakingServerDir = new DirectoryInfo(PixelStreamingProgramsDirectory + "/WebServers/Matchmaker").FullName;
+
+			if (!Directory.Exists(MatchmakingServerDir))
+			{
+				Log.TraceInformation(string.Format("Matchmaking Server path '{0}' does not exist", MatchmakingServerDir));
+				return;
+			}
+
+			List<string> DependenciesToAdd = new List<string>();
+			DependenciesToAdd.AddRange(Directory.GetFiles(MatchmakingServerDir, "*.*", SearchOption.AllDirectories));
+
+			string NodeModulesDirPath = new DirectoryInfo(MatchmakingServerDir + "/node_modules").FullName;
+			string LogsDirPath = new DirectoryInfo(MatchmakingServerDir + "/logs").FullName;
+			foreach (string Dependency in DependenciesToAdd)
+			{
+				if (!Dependency.StartsWith(NodeModulesDirPath) &&
+					!Dependency.StartsWith(LogsDirPath))
+				{
+					RuntimeDependencies.Add(Dependency, StagedFileType.NonUFS);
+				}
+			}
+		}
 
 		private void AddWebRTCServers()
-        {
+		{
 			string webRTCRevision = "31262";
-            string webRTCRevisionDirectory = "./ThirdParty/WebRTC/rev." + webRTCRevision;
+			string webRTCRevisionDirectory = "./ThirdParty/WebRTC/rev." + webRTCRevision;
 			string webRTCProgramsDirectory = Path.Combine(webRTCRevisionDirectory, "programs/Win64/VS2015/Release");
 
-            List<string> DependenciesToAdd = new List<string>();
-            DependenciesToAdd.AddRange(Directory.GetFiles(webRTCProgramsDirectory, "*.exe"));
-            DependenciesToAdd.AddRange(Directory.GetFiles(webRTCProgramsDirectory, "*.pdb"));
-            DependenciesToAdd.AddRange(Directory.GetFiles(webRTCProgramsDirectory, "*.bat"));
-            DependenciesToAdd.AddRange(Directory.GetFiles(webRTCProgramsDirectory, "*.ps1"));
-
-            foreach (string Dependency in DependenciesToAdd)
-            {
-                RuntimeDependencies.Add(Dependency, StagedFileType.NonUFS);
-            }
-        }
-
-        public PixelStreaming(ReadOnlyTargetRules Target) : base(Target)
-        {
-            // use private PCH to include lots of WebRTC headers
-            PCHUsage = PCHUsageMode.UseExplicitOrSharedPCHs;
-            PrivatePCHHeaderFile = "Private/PCH.h";
-            //PCHUsage = PCHUsageMode.NoPCHs;
-
-            var EngineDir = Path.GetFullPath(Target.RelativeEnginePath);
-
-            // NOTE: General rule is not to access the private folder of another module,
-            // but to use the ISubmixBufferListener interface, we  need to include some private headers
-            PrivateIncludePaths.Add(Path.Combine(EngineDir, "Source/Runtime/AudioMixer/Private"));
-
-            PrivateDependencyModuleNames.AddRange(new string[]
-            {
+			List<string> DependenciesToAdd = new List<string>();
+			DependenciesToAdd.AddRange(Directory.GetFiles(webRTCProgramsDirectory, "*.exe"));
+			DependenciesToAdd.AddRange(Directory.GetFiles(webRTCProgramsDirectory, "*.pdb"));
+			DependenciesToAdd.AddRange(Directory.GetFiles(webRTCProgramsDirectory, "*.bat"));
+			DependenciesToAdd.AddRange(Directory.GetFiles(webRTCProgramsDirectory, "*.ps1"));
+
+			foreach (string Dependency in DependenciesToAdd)
+			{
+				RuntimeDependencies.Add(Dependency, StagedFileType.NonUFS);
+			}
+		}
+
+		public PixelStreaming(ReadOnlyTargetRules Target) : base(Target)
+		{
+			// use private PCH to include lots of WebRTC headers
+			PCHUsage = PCHUsageMode.UseExplicitOrSharedPCHs;
+			PrivatePCHHeaderFile = "Private/PCH.h";
+			//PCHUsage = PCHUsageMode.NoPCHs;
+
+			var EngineDir = Path.GetFullPath(Target.RelativeEnginePath);
+
+			// NOTE: General rule is not to access the private folder of another module
+			PrivateIncludePaths.AddRange(new string[]
+				{
+					Path.Combine(EngineDir, "Source/Runtime/AudioMixer/Private"),
+					Path.Combine(EngineDir, "Source/Runtime/VulkanRHI/Private"),
+					Path.Combine(EngineDir, "Source/Runtime/VulkanRHI/Private/Linux"),
+				});
+
+			PrivateDependencyModuleNames.AddRange(new string[]
+			{
 				"ApplicationCore",
 				"Core",
-                "CoreUObject",
-                "Engine",
+				"CoreUObject",
+				"Engine",
 				"InputCore",
-                "InputDevice",
+				"InputDevice",
 				"Json",
 				"RenderCore",
-                "RHI",
+				"RHI",
 				"RHICore",
-                "Slate",
+				"Slate",
 				"SlateCore",
-                "AudioMixer",
+				"AudioMixer",
 				"WebRTC",
 				"WebSockets",
 				"Sockets",
-                "MediaUtils",
+				"MediaUtils",
 				"AVEncoder",
 				"DeveloperSettings"
 			});
+			
+			AddEngineThirdPartyPrivateStaticDependencies(Target, "Vulkan");
 
 			if (Target.Type == TargetType.Editor)
 			{
 				PrivateDependencyModuleNames.Add("UnrealEd");
 			}
 
-            DynamicallyLoadedModuleNames.AddRange(new string[]
-           {
-                "Media",
-           });
+			DynamicallyLoadedModuleNames.AddRange(new string[]
+			{
+				"Media",
+			});
 
-            PrivateIncludePathModuleNames.AddRange(new string[]
-            {
+			PrivateIncludePathModuleNames.AddRange(new string[]
+			{
 				"Media"
-            });
-
-            if (Target.bCompileAgainstEngine)
-            {
-                PrivateDependencyModuleNames.Add("Engine");
-                PrivateDependencyModuleNames.Add("HeadMountedDisplay");
-            }
+			});
 
-			if (Target.Platform == UnrealTargetPlatform.Win64)
+			if (Target.bCompileAgainstEngine)
 			{
-            	// required for casting UnrealEngine BackBuffer to D3D11 Texture2D for NvEnc
-				PrivateDependencyModuleNames.AddRange(new string[] { "D3D11RHI" });
-				PrivateIncludePaths.AddRange(new string[]
-				{
-						Path.Combine(EngineDir, "Source/Runtime/Windows/D3D11RHI/Private"),
-						Path.Combine(EngineDir, "Source/Runtime/Windows/D3D11RHI/Private/Windows"),
-				});
-				// required by D3D11RHI
-				AddEngineThirdPartyPrivateStaticDependencies(Target, "IntelMetricsDiscovery");
-				AddEngineThirdPartyPrivateStaticDependencies(Target, "IntelExtensionsFramework");
-				AddEngineThirdPartyPrivateStaticDependencies(Target, "NVAftermath");
-
-				{   // WebRTC stuff
-					PublicSystemLibraries.Add("Msdmo.lib");
-					PublicSystemLibraries.Add("Dmoguids.lib");
-					PublicSystemLibraries.Add("wmcodecdspuuid.lib");
-					PublicSystemLibraries.Add("winmm.lib");
-				}
-
-				// for `FWmfMediaHardwareVideoDecodingTextureSample`
-				// this really needs to be refactored to break dependency on WmfMedia plugin
-				PrivateDependencyModuleNames.Add("WmfMedia");
-				PrivateIncludePaths.Add("../../../Media/WmfMedia/Source/WmfMedia/Private/");
-
-				AddEngineThirdPartyPrivateStaticDependencies(Target, "DX11");
-				AddEngineThirdPartyPrivateStaticDependencies(Target, "DX9");
-
-				PublicAdditionalLibraries.Add(Target.UEThirdPartySourceDirectory + "Windows/DirectX/Lib/x64/dxerr.lib");
-
-				PublicSystemLibraries.Add("Dxva2.lib");
-				PublicSystemLibraries.Add("strmiids.lib");
-				PublicSystemLibraries.Add("legacy_stdio_definitions.lib");
-
-				//delay - load all MF DLLs to be able to check Windows version for compatibility in `StartupModule` before
-				//  loading them manually
-				PublicSystemLibraries.Add("mfplat.lib");
-				PublicDelayLoadDLLs.Add("mfplat.dll");
-				PublicSystemLibraries.Add("mfuuid.lib");
+				PrivateDependencyModuleNames.Add("Engine");
+				PrivateDependencyModuleNames.Add("HeadMountedDisplay");
 			}
-			else if (Target.IsInPlatformGroup(UnrealPlatformGroup.Linux))
+
+			if (Target.IsInPlatformGroup(UnrealPlatformGroup.Linux))
 			{
 				// required for casting UnrealEngine BackBuffer to Vulkan Texture2D for NvEnc
 				PrivateDependencyModuleNames.AddRange(new string[] { "CUDA", "VulkanRHI", "nvEncode"});
@@ -188,18 +153,17 @@ namespace UnrealBuildTool.Rules
 					Path.Combine(EngineDir, "Source/Runtime/VulkanRHI/Private/Linux"),
 				});
 
-				AddEngineThirdPartyPrivateStaticDependencies(Target, "Vulkan");
 
 			}
-
-            AddSignallingServer();
-    	    AddMatchmakingServer();
-
+			
+			AddSignallingServer();
+			AddMatchmakingServer();
+			
 			// We have not been able to build these yet for M84 (these are the WebRTC bundled sample STUN/TURN servers).
 			// For future I think we should advise use of COTURN instead of the WebRTC sample servers.
 			// if (Target.Platform == UnrealTargetPlatform.Win64) {
 			//     AddWebRTCServers();
 			// }
-        }
-    }
+		}
+	}
 }
diff --git a/Engine/Plugins/Media/PixelStreaming/Source/PixelStreaming/Private/AudioCapturer.cpp b/Engine/Plugins/Media/PixelStreaming/Source/PixelStreaming/Private/AudioCapturer.cpp
--- a/Engine/Plugins/Media/PixelStreaming/Source/PixelStreaming/Private/AudioCapturer.cpp
+++ b/Engine/Plugins/Media/PixelStreaming/Source/PixelStreaming/Private/AudioCapturer.cpp
@@ -25,8 +25,6 @@ DEFINE_LOG_CATEGORY(LogAudioCapturer);
 		};                       \
 	}
 
-// const int FAudioCapturer::SampleRate;
-// const int FAudioCapturer::NumChannels;
 
 void FAudioCapturer::OnNewSubmixBuffer(const USoundSubmix* OwningSubmix, float* AudioData, int32 NumSamples, int32 InNumChannels, const int32 InSampleRate, double AudioClock)
 {
diff --git a/Engine/Plugins/Media/PixelStreaming/Source/PixelStreaming/Private/HUDStats.cpp b/Engine/Plugins/Media/PixelStreaming/Source/PixelStreaming/Private/HUDStats.cpp
--- a/Engine/Plugins/Media/PixelStreaming/Source/PixelStreaming/Private/HUDStats.cpp
+++ b/Engine/Plugins/Media/PixelStreaming/Source/PixelStreaming/Private/HUDStats.cpp
@@ -1,28 +1,8 @@
 // Copyright Epic Games, Inc. All Rights Reserved.
 
 #include "HUDStats.h"
-
 #include "Engine/Engine.h"
 
-namespace
-{
-	bool bHudStatsEnabledCVar = false;
-
-	FAutoConsoleVariableRef CVarHudStats(
-		TEXT("PixelStreaming.HUDStats"),
-		bHudStatsEnabledCVar,
-		TEXT("Whether to show PixelStreaming stats on HUD"),
-		FConsoleVariableDelegate::CreateLambda([](IConsoleVariable* InVariable)
-		{
-			FHUDStats& Stats = FHUDStats::Get();
-			Stats.bEnabled = bHudStatsEnabledCVar;
-			Stats.Reset();
-			if (Stats.bEnabled)
-				GAreScreenMessagesEnabled = true;
-		}),
-		ECVF_Cheat
-	);
-}
 
 FHUDStats& FHUDStats::Get()
 {
@@ -32,25 +12,37 @@ FHUDStats& FHUDStats::Get()
 
 void FHUDStats::Tick()
 {
-	if (!GEngine || !bEnabled)
+	bool bHUDEnabled = PixelStreamingSettings::CVarPixelStreamingHudStats.GetValueOnGameThread();
+
+	if (!GEngine || !bHUDEnabled)
 	{
 		return;
 	}
 
+	if(bHUDEnabled)
+	{
+		GAreScreenMessagesEnabled = true;
+	}
+
 	// display stats on HUD
 
-	double NowSecs = rtc::TimeMillis() % 1000000 / 1000.0;
-	GEngine->AddOnScreenDebugMessage(1, 0, FColor::Green, *FString::Printf(TEXT("time %.3f ms"), NowSecs), false /* newer on top */, FVector2D{ 5, 5 } /* text scale */);
+	double NowMillis = FTimespan::FromSeconds(FPlatformTime::Seconds()).GetTotalMilliseconds();
+	GEngine->AddOnScreenDebugMessage(1, 0, FColor::Green, *FString::Printf(TEXT("Timestamps: %f ms"), NowMillis), false /* newer on top */, FVector2D{ 5, 5 } /* text scale */);
+
+	GEngine->AddOnScreenDebugMessage(2, 0, FColor::Green, *FString::Printf(TEXT("End to end latency: %.0f ms"), EndToEndLatencyMs.Get()), false /* newer on top */);
+
+	GEngine->AddOnScreenDebugMessage(3, 0, FColor::Green, *FString::Printf(TEXT("Encoder latency: %.0f ms"), EncoderLatencyMs.Get()), false /* newer on top */);
+
+	GEngine->AddOnScreenDebugMessage(4, 0, FColor::Green, *FString::Printf(TEXT("Capture latency: %.0f ms"), CaptureLatencyMs.Get()), false /* newer on top */);
 
-	GEngine->AddOnScreenDebugMessage(2, 0, FColor::Green, *FString::Printf(TEXT("latency ms: end-to-end %.0f, encoder %.0f"), EndToEndLatencyMs.Get(), EncoderLatencyMs.Get()), false /* newer on top */);
+	GEngine->AddOnScreenDebugMessage(5, 0, FColor::Green, *FString::Printf(TEXT("Encoder FPS: %.0f"), EncoderFPS.Get()), false /* newer on top */);
 
-	GEngine->AddOnScreenDebugMessage(3, 0, FColor::Green, *FString::Printf(TEXT("bitrate %.3f Mbps"), EncoderBitrateMbps.Get()), false /* newer on top */);
+	GEngine->AddOnScreenDebugMessage(6, 0, FColor::Green, *FString::Printf(TEXT("Render FPS: %.0f"), 1.0f / FApp::GetDeltaTime() ), false /* newer on top */);
 
-	GEngine->AddOnScreenDebugMessage(4, 0, FColor::Green, *FString::Printf(TEXT("QP: %.0f"), EncoderQP.Get()), false /* newer on top */);
+	GEngine->AddOnScreenDebugMessage(7, 0, FColor::Green, *FString::Printf(TEXT("Encoder bitrate: %.3f Mbps"), EncoderBitrateMbps.Get()), false /* newer on top */);
 
-	GEngine->AddOnScreenDebugMessage(5, 0, FColor::Green, *FString::Printf(TEXT("FPS: encoder %.0f, max %.0f"), EncoderFPS.Get(), GEngine->GetMaxFPS()), false /* newer on top */);
+	GEngine->AddOnScreenDebugMessage(8, 0, FColor::Green, *FString::Printf(TEXT("QP: %.0f"), EncoderQP.Get()), false /* newer on top */);
 	
-	GEngine->AddOnScreenDebugMessage(6, 0, FColor::Green, *FString::Printf(TEXT("B/W: %0.3f Mbps"), BandwidthMbps), false);
 }
 
 void FHUDStats::Reset()
@@ -59,4 +51,5 @@ void FHUDStats::Reset()
 	EncoderLatencyMs.Reset();
 	EncoderBitrateMbps.Reset();
 	EncoderQP.Reset();
+	CaptureLatencyMs.Reset();
 }
\ No newline at end of file
diff --git a/Engine/Plugins/Media/PixelStreaming/Source/PixelStreaming/Private/HUDStats.h b/Engine/Plugins/Media/PixelStreaming/Source/PixelStreaming/Private/HUDStats.h
--- a/Engine/Plugins/Media/PixelStreaming/Source/PixelStreaming/Private/HUDStats.h
+++ b/Engine/Plugins/Media/PixelStreaming/Source/PixelStreaming/Private/HUDStats.h
@@ -3,8 +3,7 @@
 #pragma once
 
 #include "Utils.h"
-
-#include "Templates/Atomic.h"
+#include "PixelStreamingSettings.h"
 
 // displays selected stats on HUD
 struct FHUDStats
@@ -14,14 +13,13 @@ struct FHUDStats
 	void Tick();
 	void Reset();
 
-	TAtomic<bool> bEnabled{ false };
-
 	static constexpr uint32 SmoothingPeriod = 3 * 60; // kinda 3 secs for 60FPS
 
+	// Note: FSmoothedValue is thread safe.
 	FSmoothedValue<SmoothingPeriod> EndToEndLatencyMs;
 	FSmoothedValue<SmoothingPeriod> EncoderLatencyMs;
+	FSmoothedValue<SmoothingPeriod> CaptureLatencyMs;
 	FSmoothedValue<SmoothingPeriod> EncoderBitrateMbps;
 	FSmoothedValue<SmoothingPeriod> EncoderQP;
 	FSmoothedValue<SmoothingPeriod> EncoderFPS;
-	double BandwidthMbps = 0;
 };
diff --git a/Engine/Plugins/Media/PixelStreaming/Source/PixelStreaming/Private/LatencyTester.cpp b/Engine/Plugins/Media/PixelStreaming/Source/PixelStreaming/Private/LatencyTester.cpp
new file mode 100644
--- /dev/null
+++ b/Engine/Plugins/Media/PixelStreaming/Source/PixelStreaming/Private/LatencyTester.cpp
@@ -0,0 +1,112 @@
+// Copyright Epic Games, Inc. All Rights Reserved.
+
+#include "LatencyTester.h"
+#include <chrono>
+
+FLatencyTester& FLatencyTester::GetInstance()
+{
+    static FLatencyTester Instance;
+    return Instance;
+}
+
+bool FLatencyTester::IsTestRunning()
+{
+    return FLatencyTester::GetInstance().TestStage.load() != FLatencyTester::ELatencyTestStage::INACTIVE;
+}
+
+FLatencyTester::ELatencyTestStage FLatencyTester::GetTestStage()
+{
+    return FLatencyTester::GetInstance().TestStage.load();
+}
+
+// Start latency tester, this makes the RecordXXX functions do stuff (otherwise they no-op).
+void FLatencyTester::Start()
+{
+    FLatencyTester::GetInstance().TestStage.store(FLatencyTester::ELatencyTestStage::STARTED);
+    FLatencyTester::GetInstance().ReceiptTimeMs = 0;
+    FLatencyTester::GetInstance().PreCaptureTimeMs = 0;
+    FLatencyTester::GetInstance().PostCaptureTimeMs = 0;
+    FLatencyTester::GetInstance().PreEncodeTimeMs = 0;
+    FLatencyTester::GetInstance().PostEncodeTimeMs = 0;
+}
+
+bool FLatencyTester::RecordReceiptTime()
+{
+    if(FLatencyTester::GetInstance().GetTestStage() == FLatencyTester::ELatencyTestStage::STARTED)
+    {
+        FLatencyTester::GetInstance().ReceiptTimeMs = FLatencyTester::GetInstance().EpochMillisNow();
+        FLatencyTester::GetInstance().TestStage.store(FLatencyTester::ELatencyTestStage::PRE_CAPTURE);
+        return true;
+    }
+    return false;
+}
+
+bool FLatencyTester::RecordPreCaptureTime(uint32 FrameId)
+{
+    if(FLatencyTester::GetInstance().GetTestStage() == FLatencyTester::ELatencyTestStage::PRE_CAPTURE)
+    {
+        FLatencyTester::GetInstance().FrameId = FrameId;
+        FLatencyTester::GetInstance().PreCaptureTimeMs = FLatencyTester::GetInstance().EpochMillisNow();
+        FLatencyTester::GetInstance().TestStage.store(FLatencyTester::ELatencyTestStage::POST_CAPTURE);
+        return true;
+    }
+    return false;
+}
+
+bool FLatencyTester::RecordPostCaptureTime(uint32 FrameId)
+{
+    if(FLatencyTester::GetInstance().GetTestStage() == FLatencyTester::ELatencyTestStage::POST_CAPTURE && FrameId == FLatencyTester::GetInstance().FrameId)
+    {
+        FLatencyTester::GetInstance().PostCaptureTimeMs = FLatencyTester::GetInstance().EpochMillisNow();
+        FLatencyTester::GetInstance().TestStage.store(FLatencyTester::ELatencyTestStage::PRE_ENCODE);
+        return true;
+    }
+    return false;
+}
+
+bool FLatencyTester::RecordPreEncodeTime(uint32 FrameId)
+{
+    if(FLatencyTester::GetInstance().GetTestStage() == FLatencyTester::ELatencyTestStage::PRE_ENCODE && FrameId == FLatencyTester::GetInstance().FrameId)
+    {
+        FLatencyTester::GetInstance().PreEncodeTimeMs = FLatencyTester::GetInstance().EpochMillisNow();
+        FLatencyTester::GetInstance().TestStage.store(FLatencyTester::ELatencyTestStage::POST_ENCODE);
+        return true;
+    }
+    return false;
+}
+
+bool FLatencyTester::RecordPostEncodeTime(uint32 FrameId)
+{
+    if(FLatencyTester::GetInstance().GetTestStage() == FLatencyTester::ELatencyTestStage::POST_ENCODE && FrameId == FLatencyTester::GetInstance().FrameId)
+    {
+        FLatencyTester::GetInstance().PostEncodeTimeMs = FLatencyTester::GetInstance().EpochMillisNow();
+        FLatencyTester::GetInstance().TestStage.store(FLatencyTester::ELatencyTestStage::RESULTS_READY);
+        return true;
+    }
+    return false;
+}
+
+// Ends latency testing, resets all the internally recorded timings and makes subsequent calls to RecordXXX functions a no-op.
+bool FLatencyTester::End(FString& OutJSONString)
+{
+    if(FLatencyTester::GetInstance().GetTestStage() == FLatencyTester::ELatencyTestStage::RESULTS_READY)
+    {
+        FLatencyTester::GetInstance().TestStage.store(FLatencyTester::ELatencyTestStage::INACTIVE);
+        OutJSONString = FString::Printf( 
+            TEXT( "{ \"ReceiptTimeMs\": %llu, \"PreCaptureTimeMs\": %llu, \"PostCaptureTimeMs\": %llu, \"PreEncodeTimeMs\": %llu, \"PostEncodeTimeMs\": %llu }" ), 
+            FLatencyTester::GetInstance().ReceiptTimeMs,
+            FLatencyTester::GetInstance().PreCaptureTimeMs,
+            FLatencyTester::GetInstance().PostCaptureTimeMs,
+            FLatencyTester::GetInstance().PreEncodeTimeMs,
+            FLatencyTester::GetInstance().PostEncodeTimeMs
+        );
+        return true;
+    }
+    
+    return false;
+}
+
+unsigned long long FLatencyTester::EpochMillisNow()
+{
+    return std::chrono::duration_cast<std::chrono::milliseconds>(std::chrono::system_clock::now().time_since_epoch()).count();
+}
\ No newline at end of file
diff --git a/Engine/Plugins/Media/PixelStreaming/Source/PixelStreaming/Private/LatencyTester.h b/Engine/Plugins/Media/PixelStreaming/Source/PixelStreaming/Private/LatencyTester.h
new file mode 100644
--- /dev/null
+++ b/Engine/Plugins/Media/PixelStreaming/Source/PixelStreaming/Private/LatencyTester.h
@@ -0,0 +1,50 @@
+// Copyright Epic Games, Inc. All Rights Reserved.
+
+#pragma once
+
+#include "Containers/UnrealString.h"
+
+// Performs a latency test of one frame passing through Pixel Streaming
+// from capture, to encode, to transmit
+class FLatencyTester 
+{
+    public:
+
+        enum ELatencyTestStage 
+        {
+            INACTIVE = 0,
+            STARTED = 1,
+            PRE_CAPTURE = 2,
+            POST_CAPTURE = 3,
+            PRE_ENCODE = 4,
+            POST_ENCODE = 5,
+            RESULTS_READY = 6
+        };
+
+        static ELatencyTestStage GetTestStage();
+        // Start latency tester, this makes the RecordXXX functions do stuff (otherwise they no-op).
+        static void Start();
+        static bool RecordReceiptTime();
+        static bool RecordPreCaptureTime(uint32 FrameId);
+        static bool RecordPostCaptureTime(uint32 FrameId);
+        static bool RecordPreEncodeTime(uint32 FrameId);
+        static bool RecordPostEncodeTime(uint32 FrameId);
+        static bool IsTestRunning();
+        // Ends latency testing, resets all the internally recorded timings and makes subsequent calls to RecordXXX functions a no-op.
+        static bool End(FString& OutJSONString);
+
+    private:
+        static FLatencyTester& GetInstance();
+
+        //Private constructor, only static methods are exposed.
+        FLatencyTester(){}
+        unsigned long long EpochMillisNow();
+
+        unsigned long long ReceiptTimeMs = 0;
+        unsigned long long PreCaptureTimeMs = 0;
+        unsigned long long PostCaptureTimeMs = 0;
+        unsigned long long PreEncodeTimeMs = 0;
+        unsigned long long PostEncodeTimeMs = 0;
+        uint32 FrameId = 0;
+        std::atomic<ELatencyTestStage> TestStage = ELatencyTestStage::INACTIVE;
+};
\ No newline at end of file
diff --git a/Engine/Plugins/Media/PixelStreaming/Source/PixelStreaming/Private/PixelStreamingEncoderFactory.cpp b/Engine/Plugins/Media/PixelStreaming/Source/PixelStreaming/Private/PixelStreamingEncoderFactory.cpp
--- a/Engine/Plugins/Media/PixelStreaming/Source/PixelStreaming/Private/PixelStreamingEncoderFactory.cpp
+++ b/Engine/Plugins/Media/PixelStreaming/Source/PixelStreaming/Private/PixelStreamingEncoderFactory.cpp
@@ -3,6 +3,9 @@
 #include "PixelStreamingEncoderFactory.h"
 #include "PlayerSession.h"
 #include "PixelStreamingVideoEncoder.h"
+#include "PixelStreamingSettings.h"
+#include "absl/strings/match.h"
+#include "modules/video_coding/codecs/vp8/include/vp8.h"
 
 using namespace webrtc;
 
@@ -22,15 +25,14 @@ void FPixelStreamingVideoEncoderFactory::AddSession(FPlayerSession& PlayerSessio
 
 std::vector<SdpVideoFormat> FPixelStreamingVideoEncoderFactory::GetSupportedFormats() const
 {
-	std::vector<webrtc::SdpVideoFormat> video_formats;
-	// for (const webrtc::SdpVideoFormat& h264_format : webrtc::SupportedH264Codecs())
-	// 	video_formats.push_back(h264_format);
+	const bool bForceVP8 = PixelStreamingSettings::IsForceVP8();
 
-	// TODO (M84FIX): Is this actually the only profile we want to support?
-	video_formats.push_back(CreateH264Format(webrtc::H264::kProfileConstrainedBaseline, webrtc::H264::kLevel3_1));
+	std::vector<webrtc::SdpVideoFormat> video_formats;
+	if (bForceVP8)
+		video_formats.push_back(SdpVideoFormat(cricket::kVp8CodecName));
+	else
+		video_formats.push_back(CreateH264Format(webrtc::H264::kProfileConstrainedBaseline, webrtc::H264::kLevel3_1));
 	return video_formats;
-
-	// return CreateH264Format(webrtc::H264::kProfileConstrainedBaseline, webrtc::H264::kLevel3_1) };
 }
 
 FPixelStreamingVideoEncoderFactory::CodecInfo FPixelStreamingVideoEncoderFactory::QueryVideoEncoder(const SdpVideoFormat& format) const
@@ -43,16 +45,21 @@ FPixelStreamingVideoEncoderFactory::CodecInfo FPixelStreamingVideoEncoderFactory
 
 std::unique_ptr<VideoEncoder> FPixelStreamingVideoEncoderFactory::CreateVideoEncoder(const SdpVideoFormat& format)
 {
-	FPlayerSession* Session;
-	bool res = PendingPlayerSessions.Dequeue(Session);
-	checkf(res, TEXT("no player session associated with encoder instance"));
+	if (absl::EqualsIgnoreCase(format.name, cricket::kVp8CodecName))
+		return VP8Encoder::Create();
+	else
+	{
+		FPlayerSession* Session;
+		bool res = PendingPlayerSessions.Dequeue(Session);
+		checkf(res, TEXT("no player session associated with encoder instance"));
 
-	UE_LOG(PixelStreamer, Log, TEXT("Encoder factory addded encoder for PlayerId=%s"), *Session->GetPlayerId());
+		UE_LOG(PixelStreamer, Log, TEXT("Encoder factory addded encoder for PlayerId=%s"), *Session->GetPlayerId());
 
-	auto VideoEncoder = std::make_unique<FPixelStreamingVideoEncoder>(Session, &EncoderContext);
-	this->ActiveEncoders.Add(VideoEncoder.get());
-	Session->SetVideoEncoder(VideoEncoder.get());
-	return VideoEncoder;
+		auto VideoEncoder = std::make_unique<FPixelStreamingVideoEncoder>(Session, &EncoderContext);
+		this->ActiveEncoders.Add(VideoEncoder.get());
+		Session->SetVideoEncoder(VideoEncoder.get());
+		return VideoEncoder;
+	}
 }
 
 void FPixelStreamingVideoEncoderFactory::OnEncodedImage(const webrtc::EncodedImage& encoded_image, const webrtc::CodecSpecificInfo* codec_specific_info, const webrtc::RTPFragmentationHeader* fragmentation)
@@ -81,5 +88,3 @@ void FPixelStreamingVideoEncoderFactory::ReleaseVideoEncoder(FPixelStreamingVide
 	UE_LOG(PixelStreamer, Log, TEXT("Encoder factory asked to remove encoder for PlayerId=%s"), *PlayerId);
 }
 
-//TODO real encoder class in here
-// Maintain TMap<FPlayerId, WeakPtr<ThinEncoder>>
\ No newline at end of file
diff --git a/Engine/Plugins/Media/PixelStreaming/Source/PixelStreaming/Private/PixelStreamingEncoderFactory.h b/Engine/Plugins/Media/PixelStreaming/Source/PixelStreaming/Private/PixelStreamingEncoderFactory.h
--- a/Engine/Plugins/Media/PixelStreaming/Source/PixelStreaming/Private/PixelStreamingEncoderFactory.h
+++ b/Engine/Plugins/Media/PixelStreaming/Source/PixelStreaming/Private/PixelStreamingEncoderFactory.h
@@ -22,7 +22,7 @@ public:
 	virtual ~FPixelStreamingVideoEncoderFactory() override;
 
 	/**
-	* This is used from the FPlayerSession::OnSucess to let the factory know
+	* This is used from the FPlayerSession::OnSuccess to let the factory know
 	* what session the next created encoder should belong to.
 	* It allows us to get the right FPlayerSession <-> FVideoEncoder relationship
 	*/
@@ -46,50 +46,4 @@ private:
 	// Each encoder is associated with a particular player (peer).
 	TArray<FPixelStreamingVideoEncoder*> ActiveEncoders;
 	
-};
-
-/*
-class FPixelStreamingVideoDecoderFactory : public webrtc::VideoDecoderFactory
-{
-public:
-	virtual std::vector<SdpVideoFormat> GetSupportedFormats() const override { return std::vector<SdpVideoFormat>(); };
-	virtual std::unique_ptr<VideoDecoder> CreateVideoDecoder(const SdpVideoFormat& format) override { return std::unique_ptr<VideoDecoder>();};
-	virtual std::unique_ptr<VideoDecoder> LegacyCreateVideoDecoder(const SdpVideoFormat& format, const std::string& receive_stream_id) override {return std::unique_ptr<VideoDecoder>();};
-};
-*/
-
-/*
- class FPixelStreamingAudioEncoderFactory : public webrtc::AudioEncoderFactory
- {
- public:
- 	struct Config
- 	{
-
- 	};
-
- 	virtual std::vector<AudioCodecSpec> GetSupportedEncoders() override { return std::vector<AudioCodecSpec>(); };
- 	virtual absl::optional<AudioCodecInfo> QueryAudioEncoder(const SdpAudioFormat& format) override { return absl::optional<AudioCodecInfo>(); };
- 	virtual std::unique_ptr<AudioEncoder> MakeAudioEncoder(int payload_type, const SdpAudioFormat& format, absl::optional<AudioCodecPairId> codec_pair_id) override { return std::unique_ptr<AudioEncoder>(); };
-
- 	static void AppendSupportedEncoders(std::vector<AudioCodecSpec>* specs) {};
- 	static Config* SdpToConfig(const SdpAudioFormat& format) { return nullptr; };
- };
- */
-
- /*
- class FPixelStreamingAudioDecoderFactory : public webrtc::AudioDecoderFactory
- {
- public:
- 	struct Config
- 	{
-
- 	};
-
- 	virtual std::vector<AudioCodecSpec> GetSupportedDecoders() override { return std::vector<AudioCodecSpec>(); };
- 	virtual bool IsSupportedDecoder(const SdpAudioFormat& format) override { return false; };
- 	virtual std::unique_ptr<AudioDecoder> MakeAudioDecoder(const SdpAudioFormat& format, absl::optional<AudioCodecPairId> codec_pair_id) override { return std::unique_ptr<AudioDecoder>(); };
-
- 	static void AppendSupportedDecoders(std::vector<AudioCodecSpec>* specs) {};
- 	static Config* SdpToConfig(const SdpAudioFormat& format) { return nullptr; };
- };
- */
\ No newline at end of file
+};
\ No newline at end of file
diff --git a/Engine/Plugins/Media/PixelStreaming/Source/PixelStreaming/Private/PixelStreamingFrameBuffer.h b/Engine/Plugins/Media/PixelStreaming/Source/PixelStreaming/Private/PixelStreamingFrameBuffer.h
--- a/Engine/Plugins/Media/PixelStreaming/Source/PixelStreaming/Private/PixelStreamingFrameBuffer.h
+++ b/Engine/Plugins/Media/PixelStreaming/Source/PixelStreaming/Private/PixelStreamingFrameBuffer.h
@@ -8,8 +8,8 @@
 class FPixelStreamingFrameBuffer : public webrtc::VideoFrameBuffer
 {
 public:
-	explicit FPixelStreamingFrameBuffer(AVEncoder::FVideoEncoderInputFrame* InputFrame, TSharedPtr<AVEncoder::FVideoEncoderInput> InputVideoEncoderInput)
-		: Frame(InputFrame), VideoEncoderInput(InputVideoEncoderInput)
+	explicit FPixelStreamingFrameBuffer(FTexture2DRHIRef SourceTexture, AVEncoder::FVideoEncoderInputFrame* InputFrame, TSharedPtr<AVEncoder::FVideoEncoderInput> InputVideoEncoderInput)
+		: TextureRef(SourceTexture), Frame(InputFrame), VideoEncoderInput(InputVideoEncoderInput)
 	{
 		Frame->Obtain();
 	}
@@ -37,23 +37,93 @@ public:
 		return Frame->GetHeight();
 	}
 
-	// we should not hit this as we support a native_buffer
+	//////////////////////////////////////////////
+	// NOTE: Only used for non-hardware encoders
+	//////////////////////////////////////////////
 	rtc::scoped_refptr<webrtc::I420BufferInterface> ToI420() override
-	{		
-		const AVEncoder::FVideoEncoderInputFrame::FYUV420P& YUV420P = Frame->GetYUV420P();
-		return webrtc::I420Buffer::Copy(
-			Frame->GetWidth(),
-			Frame->GetHeight(),
-			YUV420P.Data[0],
-			YUV420P.StrideY,
-			YUV420P.Data[1],
-			YUV420P.StrideU,
-			YUV420P.Data[2],
-			YUV420P.StrideV
-		);
+	{
+		rtc::scoped_refptr<webrtc::I420Buffer> Buffer = webrtc::I420Buffer::Create(Frame->GetWidth(), Frame->GetHeight());
+
+		// TODO Texture conversion can happen here as we can add a refernce to the Unreal texture to this class rather than doing raw RHI calls
+#if 1
+		{
+			// Slow CPU based path for testing
+			uint32 stride;
+			uint8* TextureData = (uint8*)GDynamicRHI->RHILockTexture2D(TextureRef, 0, EResourceLockMode::RLM_ReadOnly, stride, true);
+
+			uint32 components = GPixelFormats[TextureRef->GetFormat()].NumComponents;
+
+			// Convert from BGRA to I420
+			uint32 height = Frame->GetHeight();
+			uint32 width = Frame->GetWidth();
+
+			uint32 image_size = width * height;
+			uint32 upos = 0;
+			uint32 vpos = 0;
+			uint32 i = 0;
+
+			uint8* DataY = Buffer->MutableDataY();
+			uint8* DataU = Buffer->MutableDataU();
+			uint8* DataV = Buffer->MutableDataV();
+
+			FColor Pixel;
+
+			for (uint32 row = 0; row < TextureRef->GetSizeY(); ++row)
+			{
+				for (uint32 col = 0; col < TextureRef->GetSizeX(); ++col)
+				{
+					i = row * stride + col * components;
+					Pixel = { TextureData[i+2] , TextureData[i+1], TextureData[i] };
+					if (row % 2)
+					{
+						DataY[row * width + col] = ((66 * Pixel.R + 129 * Pixel.G + 25 * Pixel.B) >> 8) + 16;
+						// on every second pixel
+						DataU[row / 2 * (width / 2) + col / 2] = ((-38 * Pixel.R + -74 * Pixel.G + 112 * Pixel.B) >> 8) + 128;
+						DataV[row / 2 * (width / 2) + col / 2] = ((112 * Pixel.R + -94 * Pixel.G + -18 * Pixel.B) >> 8) + 128;
+
+						col += 1;
+						i = row * stride + col * components;
+						Pixel = { TextureData[i + 2] , TextureData[i + 1], TextureData[i] };
+						DataY[row * width + col] = ((66 * Pixel.R + 129 * Pixel.G + 25 * Pixel.B) >> 8) + 16;
+					}
+					else
+					{
+						DataY[row * width + col] = ((66 * Pixel.R + 129 * Pixel.G + 25 * Pixel.B) >> 8) + 16;
+					}
+				}
+			}
+
+			GDynamicRHI->RHIUnlockTexture2D(TextureRef, 0, true);
+		}
+#else
+		// Faster GPU path
+		uint8* DataY = Buffer->MutableDataY();
+		uint8* DataU = Buffer->MutableDataU();
+		uint8* DataV = Buffer->MutableDataV();
+
+		// TODO Write shader that repacks RGBA texture into I420
+		FComputeFenceRHIRef fence = GDynamicRHI->RHICreateComputeFence(TEXT("RGBA_to_YUV420_fence"));
+
+
+
+		fence.wait();
+
+		// Copy data out of I420 packed into RGBA
+		uint8* TextureData = (uint8*)GDynamicRHI->RHILockTexture2D(TextureRef, 0, EResourceLockMode::RLM_ReadOnly, stride, true);
+
+		uint32 size = Frame->GetWidth() * FrameGetHeight();
+		uint32 half_size = size / 2;
+
+		FMemory::Memcpy(DataY, TextureData, size);
+		FMemory::Memcpy(DataU, TextureData + size, half_size);
+		FMemory::Memcpy(DataV, TextureData + half_size, half_size);
+
+		GDynamicRHI->RHIUnlockTexture2D(TextureRef, 0, true);
+#endif
+
+		return Buffer;
 	}
 
-	// HACK (M84FIX) Don't blame me this const_cast is from the reference implimentation
 	AVEncoder::FVideoEncoderInputFrame* GetFrame() const
 	{
 		return Frame;
@@ -65,6 +135,7 @@ public:
 	}
 
 private:
+	FTexture2DRHIRef TextureRef;
 	AVEncoder::FVideoEncoderInputFrame* Frame;
 	TSharedPtr<AVEncoder::FVideoEncoderInput> VideoEncoderInput;
 };
\ No newline at end of file
diff --git a/Engine/Plugins/Media/PixelStreaming/Source/PixelStreaming/Private/PixelStreamingModule.cpp b/Engine/Plugins/Media/PixelStreaming/Source/PixelStreaming/Private/PixelStreamingModule.cpp
--- a/Engine/Plugins/Media/PixelStreaming/Source/PixelStreaming/Private/PixelStreamingModule.cpp
+++ b/Engine/Plugins/Media/PixelStreaming/Source/PixelStreaming/Private/PixelStreamingModule.cpp
@@ -2,7 +2,6 @@
 
 #include "PixelStreamingModule.h"
 #include "FreezeFrame.h"
-#include "Player.h"
 #include "Streamer.h"
 #include "InputDevice.h"
 #include "PixelStreamerInputComponent.h"
@@ -10,6 +9,8 @@
 #include "SignallingServerConnection.h"
 #include "HUDStats.h"
 #include "PixelStreamingPrivate.h"
+#include "PixelStreamingSettings.h"
+#include "LatencyTester.h"
 
 #include "CoreMinimal.h"
 #include "Modules/ModuleManager.h"
@@ -47,12 +48,7 @@ DEFINE_LOG_CATEGORY(PixelStreaming);
 
 namespace
 {
-	TAutoConsoleVariable<int32> CVarFreezeFrameQuality(
-		TEXT("PixelStreaming.FreezeFrameQuality"),
-		100,
-		TEXT("Compression quality of the freeze frame"),
-		ECVF_Default
-	);
+	
 
 	#if PLATFORM_WINDOWS
 	// required for WMF video decoding
@@ -85,8 +81,8 @@ void FPixelStreamingModule::InitStreamer()
 	FString SignallingServerIP;
 	uint16 SignallingServerPort = 8888;
 	FParse::Value(FCommandLine::Get(), TEXT("PixelStreamingID="), StreamerId);
-	if (!FParse::Value(FCommandLine::Get(), TEXT("PixelStreamingIP="), SignallingServerIP) ||
-		!FParse::Value(FCommandLine::Get(), TEXT("PixelStreamingPort="), SignallingServerPort))
+	if (!PixelStreamingSettings::GetSignallingServerIP(SignallingServerIP) ||
+		!PixelStreamingSettings::GetSignallingServerPort(SignallingServerPort))
 	{
 		UE_LOG(PixelStreamer, Log, TEXT("PixelStreaming is disabled, provide `PixelStreamingIP` and `PixelStreamingPort` cmd-args to enable it"));
 		return;
@@ -144,44 +140,6 @@ void FPixelStreamingModule::InitStreamer()
 	Streamer = MakeUnique<FStreamer>(FString::Printf(TEXT("ws://%s:%d"), *SignallingServerIP, SignallingServerPort), StreamerId);
 }
 
-void FPixelStreamingModule::InitPlayer()
-{
-	check(!bPlayerInitialized);
-
-#if PLATFORM_WINDOWS
-	// Win7+ only
-	if (!IsWindows7Plus())
-	{
-		UE_LOG(PixelPlayer, Log, TEXT("PixelStreamingPlayer plugin is incompatible with Windows prior to 7.0 version: %s"), *FPlatformMisc::GetOSVersion());
-		return;
-	}
-
-	if (!LoadMediaFoundationDLLs())
-	{
-		UE_LOG(PixelPlayer, Log, TEXT("Can't load Media Foundation, %s"), *FPlatformMisc::GetOSVersion());
-		return;
-	}
-
-	//Doesn't compile anymore?
-	//HRESULT Res = MFStartup(MF_VERSION);
-	//checkf(SUCCEEDED(Res), TEXT("MFStartup failed: %d"), Res);
-#elif PLATFORM_LINUX
-	// Linux pre-setup code here, none currently needed
-#endif
-
-	if (GIsClient)
-	{
-		bool bRes = FPlayer::CreateManagerAndDevice();
-		if (!bRes)
-		{
-			UE_LOG(PixelPlayer, Warning, TEXT("Failed to create DXGI Manager and Device"));
-		}
-	}
-
-	bPlayerInitialized = true;
-
-}
-
 /** IModuleInterface implementation */
 void FPixelStreamingModule::StartupModule()
 {
@@ -197,15 +155,13 @@ void FPixelStreamingModule::StartupModule()
 	else if( GDynamicRHI->GetName() == FString(TEXT("D3D11")) || 
 		   	 GDynamicRHI->GetName() == FString(TEXT("D3D12")))
 	{
-		// By calling InitStreamer and InitPlayer post engine init we can use pixel streaming in standalone editor mode
+		// By calling InitStreamer post engine init we can use pixel streaming in standalone editor mode
 		FCoreDelegates::OnPostEngineInit.AddRaw(this, &FPixelStreamingModule::InitStreamer);
-		FCoreDelegates::OnPostEngineInit.AddRaw(this, &FPixelStreamingModule::InitPlayer);
 	}
 	else if (GDynamicRHI->GetName() == FString(TEXT("Vulkan")))
 	{
 #if PLATFORM_LINUX
 		FModuleManager::LoadModuleChecked<FCUDAModule>("CUDA").OnPostCUDAInit.AddRaw(this, &FPixelStreamingModule::InitStreamer);
-		FModuleManager::GetModuleChecked<FCUDAModule>("CUDA").OnPostCUDAInit.AddRaw(this, &FPixelStreamingModule::InitPlayer);
 #endif
 	}
 
@@ -220,8 +176,6 @@ void FPixelStreamingModule::ShutdownModule()
 	}
 
 	IModularFeatures::Get().UnregisterModularFeature(GetModularFeatureName(), this);
-
-	FPlayer::DestroyManagerAndDevice();
 }
 
 bool FPixelStreamingModule::CheckPlatformCompatibility() const
@@ -353,13 +307,13 @@ void FPixelStreamingModule::AddPlayerConfig(TSharedRef<FJsonObject>& JsonObject)
 	JsonObject->SetBoolField(TEXT("FakingTouchEvents"), InputDevice->IsFakingTouchEvents());
 
 	FString PixelStreamingControlScheme;
-	if (FParse::Value(FCommandLine::Get(), TEXT("PixelStreamingControlScheme="), PixelStreamingControlScheme))
+	if (PixelStreamingSettings::GetControlScheme(PixelStreamingControlScheme))
 	{
 		JsonObject->SetStringField(TEXT("ControlScheme"), PixelStreamingControlScheme);
 	}
 
 	float PixelStreamingFastPan;
-	if (FParse::Value(FCommandLine::Get(), TEXT("PixelStreamingFastPan="), PixelStreamingFastPan))
+	if (PixelStreamingSettings::GetFastPan(PixelStreamingFastPan))
 	{
 		JsonObject->SetNumberField(TEXT("FastPan"), PixelStreamingFastPan);
 	}
@@ -419,7 +373,7 @@ void FPixelStreamingModule::SendJpeg(TArray<FColor> RawData, const FIntRect& Rec
 	if (bSuccess)
 	{
 		// Compress to a JPEG of the maximum possible quality.
-		int32 Quality = CVarFreezeFrameQuality.GetValueOnAnyThread();
+		int32 Quality = PixelStreamingSettings::CVarFreezeFrameQuality.GetValueOnAnyThread();
 		const TArray64<uint8> JpegBytes = ImageWrapper->GetCompressed(Quality);
 		Streamer->SendFreezeFrame(JpegBytes);
 	}
@@ -442,6 +396,18 @@ bool FPixelStreamingModule::IsTickableInEditor() const
 void FPixelStreamingModule::Tick(float DeltaTime)
 {
 	FHUDStats::Get().Tick();
+
+	// If we are running a latency test then check if we have timing results and if we do transmit them
+	if(FLatencyTester::IsTestRunning() && FLatencyTester::GetTestStage() == FLatencyTester::ELatencyTestStage::RESULTS_READY)
+	{
+		FString LatencyResults;
+		bool bEnded = FLatencyTester::End(LatencyResults);
+		if(bEnded)
+		{
+			Streamer->SendPlayerMessage(PixelStreamingProtocol::EToPlayerMsg::LatencyTest, LatencyResults);
+		}
+	}
+	
 }
 
 TStatId FPixelStreamingModule::GetStatId() const
@@ -449,9 +415,4 @@ TStatId FPixelStreamingModule::GetStatId() const
 	RETURN_QUICK_DECLARE_CYCLE_STAT(FPixelStreamingModule, STATGROUP_Tickables);
 }
 
-TSharedPtr<IMediaPlayer, ESPMode::ThreadSafe> FPixelStreamingModule::CreatePlayer(IMediaEventSink& EventSink)
-{
-	return bPlayerInitialized ? MakeShareable(new FPlayer(EventSink)) : nullptr;
-}
-
 IMPLEMENT_MODULE(FPixelStreamingModule, PixelStreaming)
diff --git a/Engine/Plugins/Media/PixelStreaming/Source/PixelStreaming/Private/PixelStreamingModule.h b/Engine/Plugins/Media/PixelStreaming/Source/PixelStreaming/Private/PixelStreamingModule.h
--- a/Engine/Plugins/Media/PixelStreaming/Source/PixelStreaming/Private/PixelStreamingModule.h
+++ b/Engine/Plugins/Media/PixelStreaming/Source/PixelStreaming/Private/PixelStreamingModule.h
@@ -27,8 +27,6 @@ private:
 
 	TSharedPtr<class IInputDevice> CreateInputDevice(const TSharedRef<FGenericApplicationMessageHandler>& InMessageHandler) override;
 
-	TSharedPtr<IMediaPlayer, ESPMode::ThreadSafe> CreatePlayer(IMediaEventSink& EventSink) override;
-
 	/** IPixelStreamingModule implementation */
 	FInputDevice& GetInputDevice() override;
 	void AddPlayerConfig(TSharedRef<FJsonObject>& JsonObject) override;
@@ -59,12 +57,6 @@ private:
 	void SendJpeg(TArray<FColor> RawData, const FIntRect& Rect);
 
 	void InitStreamer();
-	void InitPlayer();
-
-	bool IsPlayerInitialized() const override
-	{
-		return bPlayerInitialized;
-	}
 
 private:
 	TUniquePtr<FStreamer> Streamer;
@@ -72,6 +64,4 @@ private:
 	TArray<UPixelStreamerInputComponent*> InputComponents;
 	bool bFrozen = false;
 	bool bCaptureNextBackBufferAndStream = false;
-	
-	bool bPlayerInitialized = false;
 };
diff --git a/Engine/Plugins/Media/PixelStreaming/Source/PixelStreaming/Private/PixelStreamingSettings.cpp b/Engine/Plugins/Media/PixelStreaming/Source/PixelStreaming/Private/PixelStreamingSettings.cpp
--- a/Engine/Plugins/Media/PixelStreaming/Source/PixelStreaming/Private/PixelStreamingSettings.cpp
+++ b/Engine/Plugins/Media/PixelStreaming/Source/PixelStreaming/Private/PixelStreamingSettings.cpp
@@ -3,6 +3,183 @@
 #include "PixelStreamingSettings.h"
 #include "PixelStreamingPrivate.h"
 
+
+namespace PixelStreamingSettings
+{
+// Begin Encoder CVars
+
+	TAutoConsoleVariable<int32> CVarPixelStreamingEncoderTargetBitrate(
+		TEXT("PixelStreaming.Encoder.TargetBitrate"),
+		-1,
+		TEXT("Target bitrate (bps). Ignore the bitrate WebRTC wants (not recommended). Set to -1 to disable. Default -1."),
+		ECVF_RenderThreadSafe);
+
+	TAutoConsoleVariable<int32> CVarPixelStreamingEncoderMaxBitrate(
+		TEXT("PixelStreaming.Encoder.MaxBitrateVBR"),
+		20000000,
+		TEXT("Max bitrate (bps). Does not work in CBR rate control mode with NVENC."),
+		ECVF_RenderThreadSafe);
+
+	TAutoConsoleVariable<bool> CVarPixelStreamingDebugDumpFrame(
+		TEXT("PixelStreaming.Encoder.DumpDebugFrames"),
+		false,
+		TEXT("Dumps frames from the encoder to a file on disk for debugging purposes."),
+		ECVF_Default);
+
+	TAutoConsoleVariable<int32> CVarPixelStreamingEncoderMinQP(
+		TEXT("PixelStreaming.Encoder.MinQP"),
+		-1,
+		TEXT("0-51, lower values result in better quality but higher bitrate. Default -1. -1 will disable any hard limit on a minimum QP."),
+		ECVF_Default);
+
+	TAutoConsoleVariable<int32> CVarPixelStreamingEncoderMaxQP(
+		TEXT("PixelStreaming.Encoder.MaxQP"),
+		-1,
+		TEXT("0-51, lower values result in better quality but higher bitrate. Default -1. -1 will disable any hard limit on a maximum QP."),
+		ECVF_Default);
+
+	TAutoConsoleVariable<FString> CVarPixelStreamingEncoderRateControl(
+		TEXT("PixelStreaming.Encoder.RateControl"),
+		TEXT("CBR"),
+		TEXT("PixelStreaming video encoder RateControl mode. Supported modes are `ConstQP`, `VBR`, `CBR`. Default: CBR, which we recommend."),
+		ECVF_Default);
+
+	TAutoConsoleVariable<bool> CVarPixelStreamingEnableFillerData(
+		TEXT("PixelStreaming.Encoder.EnableFillerData"),
+		false,
+		TEXT("Maintains constant bitrate by filling with junk data. Note: Should not be required with CBR and MinQP = -1. Default: false."),
+		ECVF_Default);
+
+	TAutoConsoleVariable<FString> CVarPixelStreamingEncoderMultipass(
+		TEXT("PixelStreaming.Encoder.Multipass"),
+		TEXT("FULL"),
+		TEXT("PixelStreaming encoder multipass. Supported modes are `DISABLED`, `QUARTER`, `FULL`"),
+		ECVF_Default);
+
+// End Encoder CVars
+
+// Begin Capturer CVars
+
+	TAutoConsoleVariable<int32> CVarPixelStreamingUseBackBufferCaptureSize(
+		TEXT("PixelStreaming.Capturer.UseBackBufferSize"),
+		3,
+		TEXT("Whether to use back buffer size or custom size"),
+		ECVF_Default);
+
+	TAutoConsoleVariable<FString> CVarPixelStreamingCaptureSize(
+		TEXT("PixelStreaming.Capturer.CaptureSize"),
+		TEXT("1920x1080"),
+		TEXT("Capture size in format widthxheight. Recommended to UseBackBufferSize instead."),
+		ECVF_Default);
+
+// End Capturer CVars
+
+// Begin WebRTC CVars
+
+	TAutoConsoleVariable<FString> CVarPixelStreamingDegradationPreference(
+		TEXT("PixelStreaming.WebRTC.DegradationPreference"),
+		TEXT("BALANCED"),
+		TEXT("PixelStreaming degradation preference. Supported modes are `BALANCED`, `MAINTAIN_FRAMERATE`, `MAINTAIN_RESOLUTION`"),
+		ECVF_Default);
+
+	TAutoConsoleVariable<int32> CVarPixelStreamingWebRTCMaxFps(
+		TEXT("PixelStreaming.WebRTC.MaxFps"),
+		60,
+		TEXT("Maximum fps WebRTC will try to request. Default: 60"),
+		ECVF_Default);
+
+	TAutoConsoleVariable<int32> CVarPixelStreamingWebRTCMinBitrate(
+		TEXT("PixelStreaming.WebRTC.MinBitrate"),
+		100000,
+		TEXT("Min bitrate (bps) that WebRTC will not request below. Careful not to set too high otherwise WebRTC will just drop frames. Default: 100000"),
+		ECVF_RenderThreadSafe);
+
+	TAutoConsoleVariable<int32> CVarPixelStreamingWebRTCMaxBitrate(
+		TEXT("PixelStreaming.WebRTC.MaxBitrate"),
+		20000000,
+		TEXT("Max bitrate (bps) that WebRTC will not request above. Careful not to set too high otherwise because a local (ideal network) will actually reach this. Default: 20000000"),
+		ECVF_RenderThreadSafe);
+
+	TAutoConsoleVariable<int> CVarPixelStreamingWebRTCLowQpThreshold(
+		TEXT("PixelStreaming.WebRTC.LowQpThreshold"),
+		25,
+		TEXT("Only useful when MinQP=-1. If WebRTC is getting frames below this QP it will try to increase resolution when not in MAINTAIN_RESOLUTION mode."),
+		ECVF_Default);
+
+	TAutoConsoleVariable<int> CVarPixelStreamingWebRTCHighQpThreshold(
+		TEXT("PixelStreaming.WebRTC.HighQpThreshold"),
+		37,
+		TEXT("Only useful when MinQP=-1. If WebRTC is getting frames above this QP it will decrease resolution when not in MAINTAIN_RESOLUTION mode."),
+		ECVF_Default);
+
+// End WebRTC CVars
+
+// Begin Pixel Streaming Plugin CVars
+
+	TAutoConsoleVariable<bool> CVarPixelStreamingHudStats(
+		TEXT("PixelStreaming.HUDStats"),
+		false,
+		TEXT("Whether to show PixelStreaming stats on the in-game HUD."),
+		ECVF_RenderThreadSafe);
+
+	TAutoConsoleVariable<int32> CVarFreezeFrameQuality(
+		TEXT("PixelStreaming.FreezeFrameQuality"),
+		100,
+		TEXT("Compression quality of the freeze frame"),
+		ECVF_Default
+	);
+
+// Ends Pixel Streaming Plugin CVars
+
+// Begin utility functions etc.
+	std::map<FString, AVEncoder::FVideoEncoder::RateControlMode> const RateControlCVarMap{
+		{"ConstQP", AVEncoder::FVideoEncoder::RateControlMode::CONSTQP },
+		{"VBR", AVEncoder::FVideoEncoder::RateControlMode::VBR },
+		{"CBR", AVEncoder::FVideoEncoder::RateControlMode::CBR },
+	};
+
+	std::map<FString, AVEncoder::FVideoEncoder::MultipassMode> const MultipassCVarMap{
+		{"DISABLED", AVEncoder::FVideoEncoder::MultipassMode::DISABLED },
+		{"QUARTER", AVEncoder::FVideoEncoder::MultipassMode::QUARTER },
+		{"FULL", AVEncoder::FVideoEncoder::MultipassMode::FULL },
+	};
+
+	AVEncoder::FVideoEncoder::RateControlMode GetRateControlCVar()
+	{
+		auto const cvarStr = CVarPixelStreamingEncoderRateControl.GetValueOnAnyThread();
+		auto const it = RateControlCVarMap.find(cvarStr);
+		if (it == std::end(RateControlCVarMap))
+			return AVEncoder::FVideoEncoder::RateControlMode::CBR;
+		return it->second;
+	}
+
+	AVEncoder::FVideoEncoder::MultipassMode GetMultipassCVar()
+	{
+		auto const cvarStr = CVarPixelStreamingEncoderMultipass.GetValueOnAnyThread();
+		auto const it = MultipassCVarMap.find(cvarStr);
+		if (it == std::end(MultipassCVarMap))
+			return AVEncoder::FVideoEncoder::MultipassMode::FULL;
+		return it->second;
+	}
+
+	webrtc::DegradationPreference GetDegradationPreference()
+	{
+		FString DegradationPrefStr = CVarPixelStreamingDegradationPreference.GetValueOnAnyThread();
+		if(DegradationPrefStr == "MAINTAIN_FRAMERATE")
+		{
+			return webrtc::DegradationPreference::MAINTAIN_FRAMERATE;
+		}
+		else if(DegradationPrefStr == "MAINTAIN_RESOLUTION")
+		{
+			return webrtc::DegradationPreference::MAINTAIN_RESOLUTION;
+		}
+		// Everything else, return balanced.
+		return webrtc::DegradationPreference::BALANCED;
+	}
+// End utility functions etc.
+
+}
+
 UPixelStreamingSettings::UPixelStreamingSettings(const FObjectInitializer& ObjectInitlaizer)
 	: Super(ObjectInitlaizer)
 {
diff --git a/Engine/Plugins/Media/PixelStreaming/Source/PixelStreaming/Private/PixelStreamingVideoEncoder.cpp b/Engine/Plugins/Media/PixelStreaming/Source/PixelStreaming/Private/PixelStreamingVideoEncoder.cpp
--- a/Engine/Plugins/Media/PixelStreaming/Source/PixelStreaming/Private/PixelStreamingVideoEncoder.cpp
+++ b/Engine/Plugins/Media/PixelStreaming/Source/PixelStreaming/Private/PixelStreamingVideoEncoder.cpp
@@ -6,133 +6,63 @@
 #include "PixelStreamingFrameBuffer.h"
 #include "PlayerSession.h"
 #include "HUDStats.h"
+#include "UnrealEngine.h"
+#include "Misc/Paths.h"
+#include "PixelStreamingSettings.h"
+#include "LatencyTester.h"
 
-TAutoConsoleVariable<int32> CVarPixelStreamingEncoderTargetBitrate(
-	TEXT("PixelStreaming.Encoder.TargetBitrate"),
-	-1,
-	TEXT("Target bitrate (bps). Ignore the bitrate WebRTC wants (not recommended). Set to -1 to disable."),
-	ECVF_RenderThreadSafe);
-
-TAutoConsoleVariable<int32> CVarPixelStreamingEncoderMaxBitrate(
-	TEXT("PixelStreaming.Encoder.MaxBitrateVBR"),
-	50000000,
-	TEXT("Max bitrate (bps). Does not work in CBR rate control mode with NVENC."),
-	ECVF_RenderThreadSafe);
-
-TAutoConsoleVariable<int32> CVarPixelStreamingEncoderUseBackBufferSize(
-	TEXT("PixelStreaming.Encoder.UseBackBufferSize"),
-	3,
-	TEXT("Whether to use back buffer size or custom size"),
-	ECVF_Cheat);
-
-TAutoConsoleVariable<FString> CVarPixelStreamingEncoderTargetSize(
-	TEXT("PixelStreaming.Encoder.TargetSize"),
-	TEXT("1920x1080"),
-	TEXT("Encoder target size in format widthxheight"),
-	ECVF_Cheat);
-
-TAutoConsoleVariable<bool> CVarPixelStreamingDebugDumpFrame(
-	TEXT("PixelStreaming.Encoder.DumpDebugFrames"),
-	false,
-	TEXT("Dumps frames from the encoder to a file on disk for debugging purposes."),
-	ECVF_Cheat);
-
-TAutoConsoleVariable<int32> CVarPixelStreamingEncoderMinQP(
-	TEXT("PixelStreaming.Encoder.MinQP"),
-	20,
-	TEXT("0-51, lower values result in better quality but higher bitrate"),
-	ECVF_Default);
-
-TAutoConsoleVariable<FString> CVarPixelStreamingEncoderRateControl(
-	TEXT("PixelStreaming.Encoder.RateControl"),
-	TEXT("CBR"),
-	TEXT("PixelStreaming video encoder RateControl mode. Supported modes are `ConstQP`, `VBR`, `CBR`"),
-	ECVF_Default);
-
-TAutoConsoleVariable<bool> CVarPixelStreamingEnableFillerData(
-	TEXT("PixelStreaming.Encoder.EnableFillerData"),
-	false,
-	TEXT("Maintains constant bitrate by filling with junk data."),
-	ECVF_Cheat);
-
-TAutoConsoleVariable<int> CVarPixelStreamingLowQpThreshold(
-	TEXT("PixelStreaming.WebRTC.LowQpThreshold"),
-	24,
-	TEXT("Low QP threshold for quality scaling."),
-	ECVF_Default);
-
-TAutoConsoleVariable<int> CVarPixelStreamingHighQpThreshold(
-	TEXT("PixelStreaming.WebRTC.HighQpThreshold"),
-	37,
-	TEXT("High QP threshold for quality scaling."),
-	ECVF_Default);
 
 using namespace webrtc;
 
 namespace
 {
-	std::map<FString, AVEncoder::FVideoEncoder::RateControlMode> const RateControlCVarMap{
-		{"ConstQP", AVEncoder::FVideoEncoder::RateControlMode::CONSTQP },
-		{"VBR", AVEncoder::FVideoEncoder::RateControlMode::VBR },
-		{"CBR", AVEncoder::FVideoEncoder::RateControlMode::CBR },
-	};
-
-	AVEncoder::FVideoEncoder::RateControlMode GetRateControlCVar()
+	void CreateH264FragmentHeader(uint8 const* CodedData, size_t CodedDataSize, RTPFragmentationHeader& Fragments)
 	{
-		auto const cvarStr = CVarPixelStreamingEncoderRateControl.GetValueOnAnyThread();
-		auto const it = RateControlCVarMap.find(cvarStr);
-		if (it == std::end(RateControlCVarMap))
-			return AVEncoder::FVideoEncoder::RateControlMode::CBR;
-		return it->second;
-}
-
-void CreateH264FragmentHeader(const uint8* CodedData, size_t CodedDataSize, RTPFragmentationHeader& Fragments)
-{
-	// count the number of nal units
-	for (int pass = 0; pass < 2; ++pass)
-	{
-		size_t num_nal = 0;
-		size_t offset = 0;
-		while (offset < CodedDataSize)
+		// count the number of nal units
+		for (int pass = 0; pass < 2; ++pass)
 		{
-			// either a 0,0,1 or 0,0,0,1 sequence indicates a new 'nal'
-			size_t nal_maker_length = 3;
-			if (offset < (CodedDataSize - 3) && CodedData[offset] == 0 &&
-				CodedData[offset + 1] == 0 && CodedData[offset + 2] == 1)
-			{
-			}
-			else if (offset < (CodedDataSize - 4) && CodedData[offset] == 0 &&
-				CodedData[offset + 1] == 0 && CodedData[offset + 2] == 0 &&
-				CodedData[offset + 3] == 1)
+			size_t num_nal = 0;
+			size_t offset = 0;
+			while (offset < CodedDataSize)
 			{
-				nal_maker_length = 4;
+				// either a 0,0,1 or 0,0,0,1 sequence indicates a new 'nal'
+				size_t nal_maker_length = 3;
+				if (offset < (CodedDataSize - 3) && CodedData[offset] == 0 &&
+					CodedData[offset + 1] == 0 && CodedData[offset + 2] == 1)
+				{
+				}
+				else if (offset < (CodedDataSize - 4) && CodedData[offset] == 0 &&
+					CodedData[offset + 1] == 0 && CodedData[offset + 2] == 0 &&
+					CodedData[offset + 3] == 1)
+				{
+					nal_maker_length = 4;
+				}
+				else
+				{
+					++offset;
+					continue;
+				}
+				if (pass == 1)
+				{
+					Fragments.fragmentationOffset[num_nal] = offset + nal_maker_length;
+					Fragments.fragmentationLength[num_nal] = 0;
+					if (num_nal > 0)
+					{
+						Fragments.fragmentationLength[num_nal - 1] = offset - Fragments.fragmentationOffset[num_nal - 1];
+					}
+				}
+				offset += nal_maker_length;
+				++num_nal;
 			}
-			else
+			if (pass == 0)
 			{
-				++offset;
-				continue;
+				Fragments.VerifyAndAllocateFragmentationHeader(num_nal);
 			}
-			if (pass == 1)
+			else if (pass == 1 && num_nal > 0)
 			{
-				Fragments.fragmentationOffset[num_nal] = offset + nal_maker_length;
-				Fragments.fragmentationLength[num_nal] = 0;
-				if (num_nal > 0)
-				{
-					Fragments.fragmentationLength[num_nal - 1] = offset - Fragments.fragmentationOffset[num_nal - 1];
-				}
+				Fragments.fragmentationLength[num_nal - 1] = offset - Fragments.fragmentationOffset[num_nal - 1];
 			}
-			offset += nal_maker_length;
-			++num_nal;
 		}
-		if (pass == 0)
-		{
-			Fragments.VerifyAndAllocateFragmentationHeader(num_nal);
-		}
-		else if (pass == 1 && num_nal > 0)
-		{
-			Fragments.fragmentationLength[num_nal - 1] = offset - Fragments.fragmentationOffset[num_nal - 1];
-		}
-	}
 	}
 }
 
@@ -160,7 +90,7 @@ void FPixelStreamingVideoEncoder::SetQualityController(bool bControlsQualityNow)
 	}
 }
 
-int FPixelStreamingVideoEncoder::InitEncode(const VideoCodec* codec_settings, const VideoEncoder::Settings& settings)
+int FPixelStreamingVideoEncoder::InitEncode(VideoCodec const* codec_settings, VideoEncoder::Settings const& settings)
 {
 	UE_LOG(PixelStreamer, Log, TEXT("PixelStreaming video encoder initialise for PlayerId=%s"), *this->GetPlayerId());
 
@@ -169,14 +99,29 @@ int FPixelStreamingVideoEncoder::InitEncode(const VideoCodec* codec_settings, co
 
 	checkf(AVEncoder::FVideoEncoderFactory::Get().IsSetup(), TEXT("FVideoEncoderFactory not setup"));
 
-	VideoInit.Width = codec_settings->width;
-	VideoInit.Height = codec_settings->height;
-	VideoInit.MaxBitrate = codec_settings->maxBitrate;
-	VideoInit.TargetBitrate = codec_settings->startBitrate;
-	VideoInit.MaxFramerate = codec_settings->maxFramerate;
+	EncoderConfig.Width = codec_settings->width;
+	EncoderConfig.Height = codec_settings->height;
+	EncoderConfig.MaxBitrate = codec_settings->maxBitrate;
+	EncoderConfig.TargetBitrate = codec_settings->startBitrate;
+	EncoderConfig.MaxFramerate = codec_settings->maxFramerate;
 
 	WebRtcProposedTargetBitrate = codec_settings->startBitrate;
 
+	// Make viewport match the requested encode size. Is this a good idea?
+	// FVector2D ViewportRes = FVector2D(0,0);
+    // if (GEngine && GEngine->GameViewport)
+    // {
+	// 	GEngine->GameViewport->GetViewportSize(ViewportRes);
+    // }
+	// if(ViewportRes.X != VideoInit.Width || ViewportRes.Y != VideoInit.Height)
+	// {
+
+	// 	AsyncTask(ENamedThreads::GameThread, [this]() {
+	// 		FSystemResolution::RequestResolutionChange(VideoInit.Width, VideoInit.Height, EWindowMode::Type::Windowed);
+	// 	});
+		
+	// }
+
 	return WEBRTC_VIDEO_CODEC_OK;
 }
 
@@ -195,7 +140,78 @@ int32 FPixelStreamingVideoEncoder::Release()
 	return WEBRTC_VIDEO_CODEC_OK;
 }
 
-int32 FPixelStreamingVideoEncoder::Encode(const VideoFrame& frame, const std::vector<VideoFrameType>* frame_types)
+void FPixelStreamingVideoEncoder::CreateAVEncoder(TSharedPtr<AVEncoder::FVideoEncoderInput> encoderInput)
+{
+	// TODO: When we have multiple HW encoders do some factory checking and find the best encoder.
+	auto& Available = AVEncoder::FVideoEncoderFactory::Get().GetAvailable();
+
+	Context->Encoder = AVEncoder::FVideoEncoderFactory::Get().Create(Available[0].ID, encoderInput, EncoderConfig);
+	Context->Encoder->SetOnEncodedPacket([FactoryPtr = Context->Factory](uint32 InLayerIndex, const AVEncoder::FVideoEncoderInputFrame* InFrame, const AVEncoder::FCodecPacket& InPacket)
+	{
+		webrtc::EncodedImage Image;
+
+		RTPFragmentationHeader FragHeader;
+		CreateH264FragmentHeader(InPacket.Data, InPacket.DataSize, FragHeader);
+
+		Image.timing_.packetization_finish_ms = FTimespan::FromSeconds(FPlatformTime::Seconds()).GetTotalMilliseconds();
+		Image.timing_.encode_start_ms = InPacket.Timings.StartTs.GetTotalMilliseconds();
+		Image.timing_.encode_finish_ms = InPacket.Timings.FinishTs.GetTotalMilliseconds();
+		Image.timing_.flags = webrtc::VideoSendTiming::kTriggeredByTimer;
+
+		auto encoded_data = webrtc::EncodedImageBuffer::Create(const_cast<uint8_t*>(InPacket.Data), InPacket.DataSize);
+		Image.SetEncodedData(encoded_data);
+		Image._encodedWidth = InFrame->GetWidth();
+		Image._encodedHeight = InFrame->GetHeight();
+		Image._frameType = InPacket.IsKeyFrame ? VideoFrameType::kVideoFrameKey : VideoFrameType::kVideoFrameDelta;
+		Image.content_type_ = VideoContentType::UNSPECIFIED;
+		Image.qp_ = InPacket.VideoQP;
+		Image.SetSpatialIndex(InLayerIndex);
+		Image._completeFrame = true;
+		Image.rotation_ = webrtc::VideoRotation::kVideoRotation_0;
+		Image.SetTimestamp(InPacket.PTS);
+		Image.capture_time_ms_ = InPacket.PTS / 1000.0; //presentation timestamp is in microseconds
+		Image.ntp_time_ms_ = InFrame->NTP;
+
+		CodecSpecificInfo CodecInfo;
+		CodecInfo.codecType = webrtc::VideoCodecType::kVideoCodecH264;
+		CodecInfo.codecSpecific.H264.packetization_mode = webrtc::H264PacketizationMode::NonInterleaved;
+		CodecInfo.codecSpecific.H264.temporal_idx = webrtc::kNoTemporalIdx;
+		CodecInfo.codecSpecific.H264.idr_frame = InPacket.IsKeyFrame;
+		CodecInfo.codecSpecific.H264.base_layer_sync = false;
+
+		FactoryPtr->OnEncodedImage(Image, &CodecInfo, &FragHeader);
+
+		FHUDStats& Stats = FHUDStats::Get();
+		const double LatencyMs = InPacket.Timings.FinishTs.GetTotalMilliseconds() - InPacket.Timings.StartTs.GetTotalMilliseconds();
+		const double BitrateMbps = InPacket.DataSize * 8 * InPacket.Framerate / 1000000.0;
+		const double CaptureMs = (InPacket.Timings.StartTs.GetTotalMicroseconds() - InPacket.PTS) / 1000.0;
+
+		if (PixelStreamingSettings::CVarPixelStreamingHudStats.GetValueOnAnyThread())
+		{
+			Stats.EncoderLatencyMs.Update(LatencyMs);
+			Stats.EncoderBitrateMbps.Update(BitrateMbps);
+			Stats.EncoderQP.Update(InPacket.VideoQP);
+			Stats.EncoderFPS.Update(InPacket.Framerate);
+			Stats.CaptureLatencyMs.Update(CaptureMs);
+		}
+
+		UE_LOG(PixelStreamer, VeryVerbose, TEXT("QP %d/%0.f, capture latency %.0f ms, latency %.0f/%.0f ms, bitrate %.3f/%.3f Mbps, %d bytes")
+			, InPacket.VideoQP, Stats.EncoderQP.Get()
+			, Stats.CaptureLatencyMs.Get()
+			, LatencyMs, Stats.EncoderLatencyMs.Get()
+			, BitrateMbps, Stats.EncoderBitrateMbps.Get()
+			, (int)InPacket.DataSize);
+
+		// If we are running a latency test then record pre-encode timing
+		if(FLatencyTester::IsTestRunning() && FLatencyTester::GetTestStage() == FLatencyTester::ELatencyTestStage::POST_ENCODE)
+		{
+			FLatencyTester::RecordPostEncodeTime(InFrame->GetFrameID());
+		}
+
+	});
+}
+
+int32 FPixelStreamingVideoEncoder::Encode(VideoFrame const& frame, std::vector<VideoFrameType> const* frame_types)
 {
 	if (!bControlsQuality)
 	{
@@ -203,105 +219,64 @@ int32 FPixelStreamingVideoEncoder::Encode(const VideoFrame& frame, const std::ve
 		return WEBRTC_VIDEO_CODEC_OK;
 	}
 
+	auto const WebRTCVideoFrame = static_cast<FPixelStreamingFrameBuffer*>(frame.video_frame_buffer().get());
+
+	if (!Context->Encoder)
+		CreateAVEncoder(WebRTCVideoFrame->GetVideoEncoderInput());
+
 	// Detect video frame not matching encoder encoding resolution
-	int FrameWidth = frame.width();
-	int FrameHeight = frame.height();
-	if(this->Context->Encoder && (this->Context->Encoder->GetWidth(0) != FrameWidth || this->Context->Encoder->GetHeight(0) != FrameHeight))
+	// Note: This can happen when UE application changes its resolution and the encoder is not programattically.
+	int const FrameWidth = frame.width();
+	int const FrameHeight = frame.height();
+	if (Context->Encoder)
 	{
-		this->Context->Encoder->UpdateLayerResolution(0, FrameWidth, FrameHeight);
-		this->VideoInit.Width = FrameWidth;
-		this->VideoInit.Height = FrameHeight;
-	}	
-
-	const FPixelStreamingFrameBuffer* WebRTCVideoFrame = static_cast<FPixelStreamingFrameBuffer*>(frame.video_frame_buffer().get());
+		if (EncoderConfig.Width != FrameWidth || EncoderConfig.Height != FrameHeight)
+		{
+			EncoderConfig.Width = FrameWidth;
+			EncoderConfig.Height = FrameHeight;
+			Context->Encoder->UpdateLayerConfig(0, EncoderConfig);
+		}
+	}
+	
+	// Change encoder settings through CVars
+	const int32 MaxBitrateCVar 											= PixelStreamingSettings::CVarPixelStreamingEncoderMaxBitrate.GetValueOnRenderThread();
+	const int32 TargetBitrateCVar 										= PixelStreamingSettings::CVarPixelStreamingEncoderTargetBitrate.GetValueOnRenderThread();
+	const int32 MinQPCVar 												= PixelStreamingSettings::CVarPixelStreamingEncoderMinQP.GetValueOnRenderThread();
+	const int32 MaxQPCVar 												= PixelStreamingSettings::CVarPixelStreamingEncoderMaxQP.GetValueOnRenderThread();
+	const AVEncoder::FVideoEncoder::RateControlMode RateControlCVar 	= PixelStreamingSettings::GetRateControlCVar();
+	const AVEncoder::FVideoEncoder::MultipassMode MultiPassCVar 		= PixelStreamingSettings::GetMultipassCVar();
+	const bool FillerDataCVar 											= PixelStreamingSettings::CVarPixelStreamingEnableFillerData.GetValueOnRenderThread();
+
+	AVEncoder::FVideoEncoder::FLayerConfig NewConfig = EncoderConfig;
+	NewConfig.MaxBitrate = MaxBitrateCVar > -1 ? MaxBitrateCVar : NewConfig.MaxBitrate;
+	NewConfig.TargetBitrate = TargetBitrateCVar > -1 ? TargetBitrateCVar : WebRtcProposedTargetBitrate;
+	NewConfig.QPMin = MinQPCVar;
+	NewConfig.QPMax = MaxQPCVar;
+	NewConfig.RateControlMode = RateControlCVar;
+	NewConfig.MultipassMode = MultiPassCVar;
+	NewConfig.FillData = FillerDataCVar;
+
+	if (NewConfig != EncoderConfig)
+	{
+		UpdateConfig(NewConfig);
+	}
 
-	AVEncoder::FVideoEncoderInputFrame* NativeVideoFrame = WebRTCVideoFrame->GetFrame();
+	AVEncoder::FVideoEncoderInputFrame* const NativeVideoFrame = WebRTCVideoFrame->GetFrame();
 	NativeVideoFrame->PTS = frame.timestamp();
+	NativeVideoFrame->NTP = frame.ntp_time_ms();
 
-	AVEncoder::FVideoEncoder::FEncodeOptions	EncodeOptions;
-	
+	AVEncoder::FVideoEncoder::FEncodeOptions EncodeOptions;
 	if (frame_types && (*frame_types)[0] == webrtc::VideoFrameType::kVideoFrameKey || ForceNextKeyframe)
 	{
 		EncodeOptions.bForceKeyFrame = true;
 		ForceNextKeyframe = false;
 	}
 
-	if (!Context->Encoder)
+	// If we are running a latency test then record pre-encode timing
+	if(FLatencyTester::IsTestRunning() && FLatencyTester::GetTestStage() == FLatencyTester::ELatencyTestStage::PRE_ENCODE)
 	{
-		// TODO: When we have multiple HW encoders do some factory checking and find the best encoder.
-		auto& Available = AVEncoder::FVideoEncoderFactory::Get().GetAvailable();
-
-		Context->Encoder = AVEncoder::FVideoEncoderFactory::Get().Create(Available[0].ID, WebRTCVideoFrame->GetVideoEncoderInput(), VideoInit);
-		Context->Encoder->SetOnEncodedPacket([FactoryPtr = Context->Factory](uint32 InLayerIndex, const AVEncoder::FVideoEncoderInputFrame* InFrame, const AVEncoder::FCodecPacket& InPacket)
-		{
-			webrtc::EncodedImage Image;
-			auto encoded_data = webrtc::EncodedImageBuffer::Create(const_cast<uint8_t*>(InPacket.Data), InPacket.DataSize);
-			Image.SetEncodedData(encoded_data);
-			Image.SetTimestamp(InPacket.PTS);
-			Image._encodedWidth = InFrame->GetWidth();
-			Image._encodedHeight = InFrame->GetHeight();
-			Image._frameType = InPacket.IsKeyFrame ? VideoFrameType::kVideoFrameKey : VideoFrameType::kVideoFrameDelta;
-			Image.content_type_ = VideoContentType::UNSPECIFIED;
-			Image.qp_ = InPacket.VideoQP;
-			Image.SetSpatialIndex(InLayerIndex);
-			Image.timing_.encode_finish_ms = rtc::TimeMicros() / 1000;
-			Image.timing_.flags = webrtc::VideoSendTiming::kTriggeredByTimer;
-			Image.capture_time_ms_ = InPacket.PTS;
-			Image.ntp_time_ms_ = FTimespan::FromSeconds(FPlatformTime::Seconds()).GetTotalMilliseconds();
-			Image.timing_.encode_start_ms = InPacket.Timings.StartTs.GetTotalMilliseconds();
-			Image._completeFrame = true;
-			Image.rotation_ = webrtc::VideoRotation::kVideoRotation_0;
-
-			CodecSpecificInfo CodecInfo;
-			CodecInfo.codecType = webrtc::VideoCodecType::kVideoCodecH264;
-			CodecInfo.codecSpecific.H264.packetization_mode = webrtc::H264PacketizationMode::NonInterleaved;
-			CodecInfo.codecSpecific.H264.temporal_idx = webrtc::kNoTemporalIdx;
-			CodecInfo.codecSpecific.H264.idr_frame = InPacket.IsKeyFrame;
-			CodecInfo.codecSpecific.H264.base_layer_sync = false;
-
-			RTPFragmentationHeader FragHeader;
-			
-			CreateH264FragmentHeader(InPacket.Data, InPacket.DataSize, FragHeader);
-			FactoryPtr->OnEncodedImage(Image, &CodecInfo, &FragHeader);
-
-			FHUDStats& Stats = FHUDStats::Get();
-			double LatencyMs = InPacket.Timings.FinishTs.GetTotalMilliseconds() - InPacket.Timings.StartTs.GetTotalMilliseconds();
-			double BitrateMbps = InPacket.DataSize * 8 * InPacket.Framerate / 1000000.0;
-			if (Stats.bEnabled)
-			{
-				Stats.EncoderLatencyMs.Update(LatencyMs);
-				Stats.EncoderBitrateMbps.Update(BitrateMbps);
-				Stats.EncoderQP.Update(InPacket.VideoQP);
-				Stats.EncoderFPS.Update(InPacket.Framerate);
-				
-			}
-
-			UE_LOG(PixelStreamer, VeryVerbose, TEXT("QP %d/%0.f, latency %.0f/%.0f ms, bitrate %.3f/%.3f Mbps, %d bytes")
-				, InPacket.VideoQP, Stats.EncoderQP.Get()
-				, LatencyMs, Stats.EncoderLatencyMs.Get()
-				, BitrateMbps, Stats.EncoderBitrateMbps.Get()
-				, (int)InPacket.DataSize);
-
-
-		});
+		FLatencyTester::RecordPreEncodeTime(NativeVideoFrame->GetFrameID());
 	}
-	
-	// Change encoder settings through CVars
-
-	auto const MaxBitrateCVar = CVarPixelStreamingEncoderMaxBitrate.GetValueOnAnyThread();
-	if (MaxBitrateCVar > -1)
-		SetMaxBitrate(MaxBitrateCVar);
-
-	auto const TargetBitrateCVar = CVarPixelStreamingEncoderTargetBitrate.GetValueOnAnyThread();
-	SetTargetBitrate(TargetBitrateCVar > -1 ? TargetBitrateCVar : WebRtcProposedTargetBitrate);
-
-	auto const MinQPCVar = CVarPixelStreamingEncoderMinQP.GetValueOnAnyThread();
-	SetMinQP(MinQPCVar > -1 ? MinQPCVar : 20);
-
-	SetRateControl(GetRateControlCVar());
-
-	auto const FillerDataCVar = CVarPixelStreamingEnableFillerData.GetValueOnAnyThread();
-	EnableFillerData(FillerDataCVar);
 
 	// Encode the frame!
 	Context->Encoder->Encode(NativeVideoFrame, EncodeOptions);
@@ -311,112 +286,59 @@ int32 FPixelStreamingVideoEncoder::Encode(const VideoFrame& frame, const std::ve
 
 // Pass rate control parameters from WebRTC to our encoder
 // This is how WebRTC can control the bitrate/framerate of the encoder.
-void FPixelStreamingVideoEncoder::SetRates(const RateControlParameters& parameters)
+void FPixelStreamingVideoEncoder::SetRates(RateControlParameters const& parameters)
 {
-	if (!bControlsQuality)
-		return;
-
-	VideoInit.MaxFramerate = parameters.framerate_fps;
+	EncoderConfig.MaxFramerate = parameters.framerate_fps;
 
-	int32 TargetBitrateCVar = CVarPixelStreamingEncoderTargetBitrate.GetValueOnAnyThread();
+	const int32 TargetBitrateCVar = PixelStreamingSettings::CVarPixelStreamingEncoderTargetBitrate.GetValueOnRenderThread();
 	// We store what WebRTC wants as the bitrate, even if we are overriding it, so we can restore back to it when user stops using CVar.
 	WebRtcProposedTargetBitrate = parameters.bitrate.get_sum_kbps() * 1000; 
-	VideoInit.TargetBitrate = TargetBitrateCVar > -1 ? TargetBitrateCVar : WebRtcProposedTargetBitrate;
-
-	if (Context->Encoder)
-	{
-		Context->Encoder->UpdateLayerBitrate(0, VideoInit.MaxBitrate, VideoInit.TargetBitrate);
-		Context->Encoder->UpdateFrameRate(VideoInit.MaxFramerate);
-	}
-}
-
-// This is mostly used for testing, it may not do anything depending on the rate control mode of the encoder.
-// For example, when using NVENC encoder max bitrate is ignored in CBR rate control mode.
-void FPixelStreamingVideoEncoder::SetMaxBitrate(int32 MaxBitrate)
-{
-	if (!bControlsQuality)
-		return;
-
-	if(Context->Encoder && VideoInit.MaxBitrate != MaxBitrate)
-	{
-		VideoInit.MaxBitrate = MaxBitrate;
-		Context->Encoder->UpdateLayerBitrate(0, VideoInit.MaxBitrate, VideoInit.TargetBitrate);
-	}
-}
-
-// This is mostly used for testing, in practice it is advised to let WebRTC set the target bitrate.
-void FPixelStreamingVideoEncoder::SetTargetBitrate(int32 TargetBitrate)
-{
-	if (!bControlsQuality)
-		return;
-
-	if(Context->Encoder && VideoInit.TargetBitrate != TargetBitrate)
-	{
-		VideoInit.TargetBitrate = TargetBitrate;
-		Context->Encoder->UpdateLayerBitrate(0, VideoInit.MaxBitrate, VideoInit.TargetBitrate);
-	}
-}
-
-void FPixelStreamingVideoEncoder::SetMinQP(int32 MinQP)
-{
-	if (!bControlsQuality)
-		return;
-
-	if (Context->Encoder && VideoInit.QPMin != MinQP)
-	{
-		VideoInit.QPMin = MinQP;
-		Context->Encoder->UpdateMinQP(VideoInit.QPMin);
-	}
-}
+	EncoderConfig.TargetBitrate = TargetBitrateCVar > -1 ? TargetBitrateCVar : WebRtcProposedTargetBitrate;
 
-void FPixelStreamingVideoEncoder::SetRateControl(AVEncoder::FVideoEncoder::RateControlMode mode)
-{
-	if (!bControlsQuality)
-		return;
-
-	if (Context->Encoder && VideoInit.RateControlMode != mode)
+	if (bControlsQuality && Context->Encoder)
 	{
-		VideoInit.RateControlMode = mode;
-		Context->Encoder->UpdateRateControl(VideoInit.RateControlMode);
-	}
-}
-
-void FPixelStreamingVideoEncoder::EnableFillerData(bool enable)
-{
-	if (!bControlsQuality)
-		return;
-
-	if (Context->Encoder && VideoInit.FillData != enable)
-	{
-		VideoInit.FillData = enable;
-		Context->Encoder->UpdateFillData(VideoInit.FillData);
+		Context->Encoder->UpdateLayerConfig(0, EncoderConfig);
 	}
+		
 }
 
 VideoEncoder::EncoderInfo FPixelStreamingVideoEncoder::GetEncoderInfo() const
 {
 	VideoEncoder::EncoderInfo info;
 	info.supports_native_handle = true;
-	info.implementation_name = TCHAR_TO_UTF8(*FString::Printf(TEXT("NVENC_%s"), GDynamicRHI->GetName()));
+	info.is_hardware_accelerated = true;
+	info.has_internal_source = false;
+	info.supports_simulcast = false;
+	info.implementation_name = TCHAR_TO_UTF8(*FString::Printf(TEXT("PIXEL_STREAMING_HW_ENCODER_%s"), GDynamicRHI->GetName()));
 
-	auto const lowQp = CVarPixelStreamingLowQpThreshold.GetValueOnAnyThread();
-	auto const highQp = CVarPixelStreamingHighQpThreshold.GetValueOnAnyThread();
-	info.scaling_settings = VideoEncoder::ScalingSettings(lowQp, highQp);
+	const int LowQP  = PixelStreamingSettings::CVarPixelStreamingWebRTCLowQpThreshold.GetValueOnRenderThread();
+	const int HighQP = PixelStreamingSettings::CVarPixelStreamingWebRTCHighQpThreshold.GetValueOnRenderThread();
+	info.scaling_settings = VideoEncoder::ScalingSettings(LowQP, HighQP);
+
+	// basically means HW encoder must be perfect and drop frames itself etc
+	info.has_trusted_rate_controller = false;
 
 	return info;
 }
 
-void FPixelStreamingVideoEncoder::SendEncodedImage(const webrtc::EncodedImage& encoded_image, const webrtc::CodecSpecificInfo* codec_specific_info, const webrtc::RTPFragmentationHeader* fragmentation)
+void FPixelStreamingVideoEncoder::UpdateConfig(AVEncoder::FVideoEncoder::FLayerConfig const& config)
 {
+	EncoderConfig = config;
+	if (bControlsQuality && Context->Encoder)
+		Context->Encoder->UpdateLayerConfig(0, EncoderConfig);
+}
 
+void FPixelStreamingVideoEncoder::SendEncodedImage(webrtc::EncodedImage const& encoded_image, webrtc::CodecSpecificInfo const* codec_specific_info, webrtc::RTPFragmentationHeader const* fragmentation)
+{
 	// Dump H264 frames to file for debugging if CVar is turned on.
-	if(CVarPixelStreamingDebugDumpFrame.GetValueOnAnyThread())
+	if(PixelStreamingSettings::CVarPixelStreamingDebugDumpFrame.GetValueOnRenderThread())
 	{
 		static IFileHandle* FileHandle = nullptr;
 		if (!FileHandle)
 		{
 			IPlatformFile& PlatformFile = FPlatformFileManager::Get().GetPlatformFile();
-			FileHandle = PlatformFile.OpenWrite(TEXT("c:/tmp/nvenc.h264"));
+			FString TempFilePath = FPaths::CreateTempFilename(*FPaths::ProjectSavedDir(), TEXT("encoded_frame"), TEXT(".h264"));
+			FileHandle = PlatformFile.OpenWrite(*TempFilePath);
 			check(FileHandle);
 		}
 
@@ -424,7 +346,7 @@ void FPixelStreamingVideoEncoder::SendEncodedImage(const webrtc::EncodedImage& e
 		FileHandle->Flush();
 	}
 
-	if(OnEncodedImageCallback)
+	if (OnEncodedImageCallback)
 	{
 		webrtc::EncodedImageCallback::Result res = OnEncodedImageCallback->OnEncodedImage(encoded_image, codec_specific_info, fragmentation);
 	}
diff --git a/Engine/Plugins/Media/PixelStreaming/Source/PixelStreaming/Private/PixelStreamingVideoEncoder.h b/Engine/Plugins/Media/PixelStreaming/Source/PixelStreaming/Private/PixelStreamingVideoEncoder.h
--- a/Engine/Plugins/Media/PixelStreaming/Source/PixelStreaming/Private/PixelStreamingVideoEncoder.h
+++ b/Engine/Plugins/Media/PixelStreaming/Source/PixelStreaming/Private/PixelStreamingVideoEncoder.h
@@ -18,11 +18,11 @@ public:
 	void SetQualityController(bool bControlsQuality);
 
 	// WebRTC Interface
-	virtual int InitEncode(const webrtc::VideoCodec* codec_settings,const webrtc::VideoEncoder::Settings& settings) override;
+	virtual int InitEncode(webrtc::VideoCodec const* codec_settings, webrtc::VideoEncoder::Settings const& settings) override;
 	virtual int32 RegisterEncodeCompleteCallback(webrtc::EncodedImageCallback* callback) override;
 	virtual int32 Release() override;
-	virtual int32 Encode(const webrtc::VideoFrame& frame, const std::vector<webrtc::VideoFrameType>* frame_types) override;
-	virtual void SetRates(const RateControlParameters& parameters) override;
+	virtual int32 Encode(webrtc::VideoFrame const& frame, std::vector<webrtc::VideoFrameType> const* frame_types) override;
+	virtual void SetRates(RateControlParameters const& parameters) override;
 	virtual webrtc::VideoEncoder::EncoderInfo GetEncoderInfo() const override;
 
 	// Note: These funcs can also be overriden but are not pure virtual
@@ -32,24 +32,24 @@ public:
 	// virtual void OnLossNotification(const LossNotification& loss_notification) override;
 	// End WebRTC Interface.
 
-	void SetMaxBitrate(int32 MaxBitrate);
-	void SetTargetBitrate(int32 TargetBitrate);
-	void SetMinQP(int32 maxqp);
-	void SetRateControl(AVEncoder::FVideoEncoder::RateControlMode mode);
-	void EnableFillerData(bool enable);
-	void SendEncodedImage(const webrtc::EncodedImage& encoded_image, const webrtc::CodecSpecificInfo* codec_specific_info, const webrtc::RTPFragmentationHeader* fragmentation);
+	AVEncoder::FVideoEncoder::FLayerConfig GetConfig() const { return EncoderConfig; }
+	void UpdateConfig(AVEncoder::FVideoEncoder::FLayerConfig const& config);
+
+	void SendEncodedImage(webrtc::EncodedImage const& encoded_image, webrtc::CodecSpecificInfo const* codec_specific_info, webrtc::RTPFragmentationHeader const* fragmentation);
 	FPlayerId GetPlayerId();
 	bool IsRegisteredWithWebRTC();
 
 	void ForceKeyFrame() { ForceNextKeyframe = true; }
 
 private:
+	void CreateAVEncoder(TSharedPtr<AVEncoder::FVideoEncoderInput> encoderInput);
+
 	// We store this so we can restore back to it if the user decides to use then stop using the PixelStreaming.Encoder.TargetBitrate CVar.
-	int32 WebRtcProposedTargetBitrate = 20000000;
+	int32 WebRtcProposedTargetBitrate = 5000000; 
 	FEncoderContext* Context;
 	FPlayerId PlayerId;
 
-	AVEncoder::FVideoEncoder::FInit VideoInit;
+	AVEncoder::FVideoEncoder::FLayerConfig EncoderConfig;
 
 	webrtc::EncodedImageCallback* OnEncodedImageCallback = nullptr;
 	
diff --git a/Engine/Plugins/Media/PixelStreaming/Source/PixelStreaming/Private/PlayerSession.cpp b/Engine/Plugins/Media/PixelStreaming/Source/PixelStreaming/Private/PlayerSession.cpp
--- a/Engine/Plugins/Media/PixelStreaming/Source/PixelStreaming/Private/PlayerSession.cpp
+++ b/Engine/Plugins/Media/PixelStreaming/Source/PixelStreaming/Private/PlayerSession.cpp
@@ -7,11 +7,12 @@
 #include "VideoEncoder.h"
 #include "VideoEncoderFactory.h"
 #include "WebRtcObservers.h"
-
 #include "Modules/ModuleManager.h"
-
 #include "PixelStreamingEncoderFactory.h"
 #include "PixelStreamingVideoEncoder.h"
+#include "PixelStreamingSettings.h"
+#include "LatencyTester.h"
+#include <chrono>
 
 FPlayerSession::FPlayerSession(FStreamer& InStreamer, FPlayerId InPlayerId, bool bInOriginalQualityController)
     : Streamer(InStreamer)
@@ -70,6 +71,40 @@ void FPlayerSession::OnOffer(TUniquePtr<webrtc::SessionDescriptionInterface> SDP
 		Streamer.GetVideoEncoderFactory()->AddSession(*this);
 
 		PeerConnection->SetLocalDescription(SetLocalDescriptionObserver, SDP);
+
+		// Once local description has been set we can start setting some encoding information for the video stream rtp sender
+		std::vector<rtc::scoped_refptr<webrtc::RtpSenderInterface>> SendersArr = PeerConnection->GetSenders();
+		for(rtc::scoped_refptr<webrtc::RtpSenderInterface> Sender : SendersArr)
+		{
+			cricket::MediaType MediaType = Sender->media_type();
+			if(MediaType == cricket::MediaType::MEDIA_TYPE_VIDEO)
+			{
+				webrtc::RtpParameters ExistingParams = Sender->GetParameters();
+				
+				// Set the start min/max bitrate and max framerate for the codec.
+				for (webrtc::RtpEncodingParameters& codec_params : ExistingParams.encodings)
+				{
+					codec_params.max_bitrate_bps = PixelStreamingSettings::CVarPixelStreamingWebRTCMaxBitrate.GetValueOnAnyThread();
+					codec_params.min_bitrate_bps = PixelStreamingSettings::CVarPixelStreamingWebRTCMinBitrate.GetValueOnAnyThread();
+					codec_params.max_framerate = PixelStreamingSettings::CVarPixelStreamingWebRTCMaxFps.GetValueOnAnyThread();
+				}
+				
+				// Set the degradation preference based on our CVar for it.
+				webrtc::DegradationPreference DegradationPref = PixelStreamingSettings::GetDegradationPreference();
+				ExistingParams.degradation_preference = DegradationPref;
+
+				webrtc::RTCError Err = Sender->SetParameters(ExistingParams);
+				if(!Err.ok())
+				{
+					const char* ErrMsg = Err.message();
+					FString ErrorStr(ErrMsg);
+    				UE_LOG(PixelStreamer, Error, TEXT("Failed to set RTP Sender params: %s"), *ErrorStr);
+				}
+			}
+		}
+
+		// Note: if desired, manually limiting total bitrate for a peer is possible in PeerConnection::SetBitrate(const BitrateSettings& bitrate)
+
 	};
 
 	FCreateSessionDescriptionObserver* CreateAnswerObserver = FCreateSessionDescriptionObserver::Create
@@ -84,7 +119,20 @@ void FPlayerSession::OnOffer(TUniquePtr<webrtc::SessionDescriptionInterface> SDP
 
 	auto OnSetRemoteDescriptionSuccess = [this, CreateAnswerObserver, OnCreateAnswerSuccess = MoveTemp(OnCreateAnswerSuccess)]()
 	{
-		webrtc::PeerConnectionInterface::RTCOfferAnswerOptions AnswerOption{ 0, 0, true, true, true };
+		int offer_to_receive_video = 0;
+		int offer_to_receive_audio = 0; // ToDo: Make CVar to receive browser audio.
+		bool voice_activity_detection = true;
+		bool ice_restart = true;
+		bool use_rtp_mux = true;
+
+		webrtc::PeerConnectionInterface::RTCOfferAnswerOptions AnswerOption{ 
+			offer_to_receive_video, 
+			offer_to_receive_audio,
+			voice_activity_detection,
+			ice_restart,
+			use_rtp_mux
+		};
+		
 		PeerConnection->CreateAnswer(CreateAnswerObserver, AnswerOption);
 	};
 
@@ -181,6 +229,47 @@ void FPlayerSession::OnIceConnectionReceivingChange(bool Receiving)
 void FPlayerSession::OnTrack(rtc::scoped_refptr<webrtc::RtpTransceiverInterface> transceiver)
 {
 	UE_LOG(PixelStreamer, Log, TEXT("%s : PlayerId=%s"), TEXT("FPlayerSession::OnTrack"), *PlayerId);
+	
+	// print out track type
+	auto mediaType = transceiver->media_type();
+	switch(mediaType)
+	{
+		case cricket::MediaType::MEDIA_TYPE_AUDIO:
+			UE_LOG(PixelStreamer, Log, TEXT("Track was type: audio"));
+			break;
+		case cricket::MediaType::MEDIA_TYPE_VIDEO:
+			UE_LOG(PixelStreamer, Log, TEXT("Track was type: video"));
+			break;
+		case cricket::MediaType::MEDIA_TYPE_DATA:
+			UE_LOG(PixelStreamer, Log, TEXT("Track was type: data"));
+			break;
+		default:
+			UE_LOG(PixelStreamer, Log, TEXT("Track was an unsupported type"));
+			break;
+  
+	}
+
+	// print out track direction
+	webrtc::RtpTransceiverDirection direction = transceiver->direction();
+	switch(direction)
+	{
+		case webrtc::RtpTransceiverDirection::kSendRecv:
+			UE_LOG(PixelStreamer, Log, TEXT("Track direction: send+recv"));
+			break;
+		case webrtc::RtpTransceiverDirection::kSendOnly:
+			UE_LOG(PixelStreamer, Log, TEXT("Track direction: send only"));
+			break;
+		case webrtc::RtpTransceiverDirection::kRecvOnly:
+			UE_LOG(PixelStreamer, Log, TEXT("Track direction: recv only"));
+			break;
+		case webrtc::RtpTransceiverDirection::kInactive:
+			UE_LOG(PixelStreamer, Log, TEXT("Track direction: inactive"));
+			break;
+		case webrtc::RtpTransceiverDirection::kStopped:
+			UE_LOG(PixelStreamer, Log, TEXT("Track direction: stopped"));
+			break;
+	}
+
 }
 
 void FPlayerSession::OnRemoveTrack(rtc::scoped_refptr<webrtc::RtpReceiverInterface> receiver)
@@ -194,15 +283,28 @@ void FPlayerSession::OnRemoveTrack(rtc::scoped_refptr<webrtc::RtpReceiverInterfa
 
 void FPlayerSession::OnMessage(const webrtc::DataBuffer& Buffer)
 {
-	auto MsgType = static_cast<PixelStreamingProtocol::EToStreamerMsg>(Buffer.data.data()[0]);
+	const uint8* Data = Buffer.data.data();
+	uint32 Size = static_cast<uint32>(Buffer.data.size());
+	PixelStreamingProtocol::EToStreamerMsg MsgType = static_cast<PixelStreamingProtocol::EToStreamerMsg>(Data[0]);
+
 	if (MsgType == PixelStreamingProtocol::EToStreamerMsg::RequestQualityControl)
 	{
-		check(Buffer.data.size() == 1);
+		check(Size == 1);
 		Streamer.OnQualityOwnership(PlayerId);
 	}
+	else if(MsgType == PixelStreamingProtocol::EToStreamerMsg::LatencyTest)
+	{
+		// Have to parse even though we already have this so that pointer arithmetic lines up when we parse the rest of the data
+		MsgType = PixelStreamingProtocol::ParseBuffer<PixelStreamingProtocol::EToStreamerMsg>(Data, Size);
+		FString TestStartTimeInBrowserMs = PixelStreamingProtocol::ParseString(Data, Size);
+		FLatencyTester::Start();
+		FLatencyTester::RecordReceiptTime();
+		unsigned long long NowMs = std::chrono::duration_cast<std::chrono::milliseconds>(std::chrono::system_clock::now().time_since_epoch()).count();
+		UE_LOG(PixelStreamer, Log, TEXT("Browser start time: %s | UE start time: %llu"), *TestStartTimeInBrowserMs, NowMs);
+	}
 	else if (!IsEngineExitRequested())
 	{
-		InputDevice.OnMessage(Buffer.data.data(), static_cast<uint32>(Buffer.data.size()));
+		InputDevice.OnMessage(Data, Size);
 	}
 }
 
@@ -302,10 +404,12 @@ void FPlayerSession::SendUnfreezeFrame()
 	DataChannel->Send(webrtc::DataBuffer(Buffer, true));
 }
 
-void FPlayerSession::OnNewSecondarySession()
+void FPlayerSession::SendKeyFrame()
 {
 	if (IsQualityController())
+	{
 		VideoEncoder->ForceKeyFrame();
+	}
 }
 
 void FPlayerSession::OnBufferedAmountChange(uint64_t PreviousAmount)
diff --git a/Engine/Plugins/Media/PixelStreaming/Source/PixelStreaming/Private/PlayerSession.h b/Engine/Plugins/Media/PixelStreaming/Source/PixelStreaming/Private/PlayerSession.h
--- a/Engine/Plugins/Media/PixelStreaming/Source/PixelStreaming/Private/PlayerSession.h
+++ b/Engine/Plugins/Media/PixelStreaming/Source/PixelStreaming/Private/PlayerSession.h
@@ -38,7 +38,7 @@ public:
 	void SendFreezeFrame(const TArray64<uint8>& JpegBytes);
 	void SendUnfreezeFrame();
 
-	void OnNewSecondarySession();
+	void SendKeyFrame();
 
 private:
 	//
diff --git a/Engine/Plugins/Media/PixelStreaming/Source/PixelStreaming/Private/ProtocolDefs.h b/Engine/Plugins/Media/PixelStreaming/Source/PixelStreaming/Private/ProtocolDefs.h
--- a/Engine/Plugins/Media/PixelStreaming/Source/PixelStreaming/Private/ProtocolDefs.h
+++ b/Engine/Plugins/Media/PixelStreaming/Source/PixelStreaming/Private/ProtocolDefs.h
@@ -20,6 +20,7 @@ namespace PixelStreamingProtocol
 		AverageBitrateRequest = 3,
 		StartStreaming = 4,
 		StopStreaming = 5,
+		LatencyTest = 6,
 
 		/**********************************************************************/
 
@@ -61,11 +62,36 @@ namespace PixelStreamingProtocol
 
 	//! Messages that can be sent to the webrtc players
 	enum class EToPlayerMsg : uint8 {
-		QualityControlOwnership,
-		Response,
-		Command,
-		FreezeFrame,
-		UnfreezeFrame,
-		VideoEncoderAvgQP // average Quantisation Parameter value of Video Encoder, roughly depicts video encoding quality
+		QualityControlOwnership = 0,
+		Response = 1,
+		Command = 2,
+		FreezeFrame = 3,
+		UnfreezeFrame = 4,
+		VideoEncoderAvgQP = 5, // average Quantisation Parameter value of Video Encoder, roughly depicts video encoding quality
+		LatencyTest = 6
 	};
+
+	template<typename T>
+	static const T& ParseBuffer(const uint8*& Data, uint32& Size)
+	{
+		checkf(sizeof(T) <= Size, TEXT("%d - %d"), sizeof(T), Size);
+		const T& Value = *reinterpret_cast<const T*>(Data);
+		Data += sizeof(T);
+		Size -= sizeof(T);
+		return Value;
+	}
+
+	static FString ParseString(const uint8*& Data, uint32& Size)
+	{
+		uint16 StrLen = ParseBuffer<uint16>(Data, Size);
+		checkf(StrLen * sizeof(TCHAR) <= Size, TEXT("%d - %d"), StrLen, Size);
+		FString Res;
+		Res.GetCharArray().SetNumUninitialized(StrLen + 1);
+		FMemory::Memcpy(Res.GetCharArray().GetData(), Data, StrLen * sizeof(TCHAR));
+		Res.GetCharArray()[StrLen] = '\0';
+		Data += StrLen * sizeof(TCHAR);
+		Size -= StrLen * sizeof(TCHAR);
+		return Res;
+	}
+
 };
\ No newline at end of file
diff --git a/Engine/Plugins/Media/PixelStreaming/Source/PixelStreaming/Private/SignallingServerConnection.cpp b/Engine/Plugins/Media/PixelStreaming/Source/PixelStreaming/Private/SignallingServerConnection.cpp
--- a/Engine/Plugins/Media/PixelStreaming/Source/PixelStreaming/Private/SignallingServerConnection.cpp
+++ b/Engine/Plugins/Media/PixelStreaming/Source/PixelStreaming/Private/SignallingServerConnection.cpp
@@ -186,6 +186,7 @@ void FSignallingServerConnection::OnClosed(int32 StatusCode, const FString& Reas
 {
 	UE_LOG(LogPixelStreamingSS, Log, TEXT("Connection to SS closed: \n\tstatus %d\n\treason: %s\n\twas clean: %s"), StatusCode, *Reason, bWasClean ? TEXT("true") : TEXT("false"));
 	Observer.OnSignallingServerDisconnected();
+	GWorld->GetTimerManager().ClearTimer(TimerHandle_KeepAlive);
 }
 
 void FSignallingServerConnection::OnMessage(const FString& Msg)
diff --git a/Engine/Plugins/Media/PixelStreaming/Source/PixelStreaming/Private/Streamer.cpp b/Engine/Plugins/Media/PixelStreaming/Source/PixelStreaming/Private/Streamer.cpp
--- a/Engine/Plugins/Media/PixelStreaming/Source/PixelStreaming/Private/Streamer.cpp
+++ b/Engine/Plugins/Media/PixelStreaming/Source/PixelStreaming/Private/Streamer.cpp
@@ -7,34 +7,10 @@
 #include "VideoEncoderFactory.h"
 #include "PixelStreamerDelegates.h"
 #include "PixelStreamingEncoderFactory.h"
+#include "PixelStreamingSettings.h"
 
 DEFINE_LOG_CATEGORY(PixelStreamer);
 
-extern TAutoConsoleVariable<int32> CVarPixelStreamingEncoderMaxBitrate;
-
-TAutoConsoleVariable<FString> CVarPixelStreamingDegradationPreference(
-	TEXT("PixelStreaming.DegradationPreference"),
-	TEXT("BALANCED"),
-	TEXT("PixelStreaming degradation preference. Supported modes are `BALANCED`, `MAINTAIN_FRAMERATE`, `MAINTAIN_RESOLUTION`"),
-	ECVF_Default);
-
-namespace
-{
-	std::map<FString, webrtc::DegradationPreference> const StringToDegradationPrefMap {
-		{"BALANCED", webrtc::DegradationPreference::BALANCED},
-		{"MAINTAIN_FRAMERATE", webrtc::DegradationPreference::MAINTAIN_FRAMERATE},
-		{"MAINTAIN_RESOLUTION", webrtc::DegradationPreference::MAINTAIN_RESOLUTION},
-	};
-
-	webrtc::DegradationPreference GetDegradationPreferenceCVar()
-	{
-		auto const prefStr = CVarPixelStreamingDegradationPreference.GetValueOnAnyThread();
-		auto const it = StringToDegradationPrefMap.find(prefStr);
-		if (it == std::end(StringToDegradationPrefMap))
-			return webrtc::DegradationPreference::BALANCED;
-		return it->second;
-	}
-}
 
 bool FStreamer::CheckPlatformCompatibility()
 {
@@ -178,7 +154,6 @@ void FStreamer::OnFrameBufferReady(const FTexture2DRHIRef& FrameBuffer)
 {
 	if (bStreamingStarted && VideoSource)
 		VideoSource->OnFrameReady(FrameBuffer);
-	SendVideoEncoderQP();
 }
 
 void FStreamer::OnConfig(const webrtc::PeerConnectionInterface::RTCConfiguration& Config)
@@ -199,7 +174,7 @@ void FStreamer::OnOffer(FPlayerId PlayerId, TUniquePtr<webrtc::SessionDescriptio
 	{
 		FScopeLock PlayersLock(&PlayersCS);
 		for (auto&& PlayerEntry : Players)
-			PlayerEntry.Value->OnNewSecondarySession();
+			PlayerEntry.Value->SendKeyFrame();
 	}
 }
 
@@ -353,8 +328,8 @@ void FStreamer::AddStreams(FPlayerId PlayerId)
 	}
 	else
 	{
-		//If desired, limiting entire bitrate for a peer is desired there is PeerConnection SetBitrate(const BitrateSettings& bitrate)
-		switch (GetDegradationPreferenceCVar())
+		webrtc::DegradationPreference DegradationPref = PixelStreamingSettings::GetDegradationPreference();
+		switch (DegradationPref)
 		{
 		case webrtc::DegradationPreference::MAINTAIN_FRAMERATE:
 			VideoTrack->set_content_hint(webrtc::VideoTrackInterface::ContentHint::kFluid);
@@ -366,6 +341,7 @@ void FStreamer::AddStreams(FPlayerId PlayerId)
 			break;
 		}
 	}
+
 }
 
 void FStreamer::OnQualityOwnership(FPlayerId PlayerId)
@@ -431,19 +407,3 @@ void FStreamer::SendUnfreezeFrame()
 	
 	CachedJpegBytes.Empty();
 }
-
-void FStreamer::SendVideoEncoderQP()
-{
-	// TODO: We no longer use hardware encoder details so this needs to get hooked up to the new way of doing things.
-
-	// if (HWEncoderDetails.LastAvgQP != FHWEncoderDetails::InvalidQP)
-	// {
-	// 	VideoEncoderAvgQP.Update(HWEncoderDetails.LastAvgQP);
-	// }
-	// double Now = FPlatformTime::Seconds();
-	// if (Now - LastVideoEncoderQPReportTime > 1)
-	// {
-	// 	SendPlayerMessage(PixelStreamingProtocol::EToPlayerMsg::VideoEncoderAvgQP, FString::Printf(TEXT("%.0f"), VideoEncoderAvgQP.Get()));
-	// 	LastVideoEncoderQPReportTime = Now;
-	// }
-}
diff --git a/Engine/Plugins/Media/PixelStreaming/Source/PixelStreaming/Private/Streamer.h b/Engine/Plugins/Media/PixelStreaming/Source/PixelStreaming/Private/Streamer.h
--- a/Engine/Plugins/Media/PixelStreaming/Source/PixelStreaming/Private/Streamer.h
+++ b/Engine/Plugins/Media/PixelStreaming/Source/PixelStreaming/Private/Streamer.h
@@ -49,7 +49,6 @@ private:
 	void DeleteAllPlayerSessions();
 	FPlayerSession* GetPlayerSession(FPlayerId PlayerId);
 	void AddStreams(FPlayerId PlayerId);
-	void SendVideoEncoderQP();
 
 private:
 	FString SignallingServerUrl;
