Remove Win32 reference

diff --git a/Engine/Plugins/Media/PixelStreaming/PixelStreaming.uplugin b/Engine/Plugins/Media/PixelStreaming/PixelStreaming.uplugin
--- a/Engine/Plugins/Media/PixelStreaming/PixelStreaming.uplugin
+++ b/Engine/Plugins/Media/PixelStreaming/PixelStreaming.uplugin
@@ -19,7 +19,7 @@
 			"Name": "PixelStreaming",
 			"Type": "RuntimeNoCommandlet",
 			"LoadingPhase": "PostDefault",
-			"WhitelistPlatforms": [ "Win32", "Win64" ],
+			"WhitelistPlatforms": [ "Win64" ],
 			"BlacklistTargets": [ "Server" ]
 		}
 	],
