Enable D3D12 to use CUDA for PixelStreaming backbuffer textures in order to resolve NVENC encoder latency issues

diff --git a/Engine/Plugins/Media/HardwareEncoders/Source/EncoderNVENC/Private/NVENC_EncoderH264.cpp b/Engine/Plugins/Media/HardwareEncoders/Source/EncoderNVENC/Private/NVENC_EncoderH264.cpp
--- a/Engine/Plugins/Media/HardwareEncoders/Source/EncoderNVENC/Private/NVENC_EncoderH264.cpp
+++ b/Engine/Plugins/Media/HardwareEncoders/Source/EncoderNVENC/Private/NVENC_EncoderH264.cpp
@@ -156,6 +156,12 @@ namespace AVEncoder
 			NVENCCommandLineParseOption(TEXT("-NVENCKeyFrameQPUseLastQP"), CVarNVENCKeyframeQPUseLastQP);
 			NVENCCommandLineParseOption(TEXT("-NVENCEnableStats"), CVarNVENCEnableStats);
 
+			// When NVENC stats CVar changes, change the output to screen flag.
+            CVarNVENCEnableStats.AsVariable()->SetOnChangedCallback(FConsoleVariableDelegate::CreateLambda([](IConsoleVariable* ChangedCVar)
+            {
+                FNVENCStats::Get().SetOutputToScreen(CVarNVENCEnableStats.GetValueOnAnyThread());
+            }));
+
 			if(CVarNVENCEnableStats.GetValueOnAnyThread())
 			{
 				FNVENCStats::Get().SetOutputToScreen(true);
@@ -913,9 +919,20 @@ namespace AVEncoder
 			return;
 		}
 
+		switch(InBuffer->SourceFrame->GetCUDA().UnderlyingRHI) {
+			case FVideoEncoderInputFrame::EUnderlyingRHI::Vulkan:
+				InBuffer->Pitch = TextureSize.X * GPixelFormats[EPixelFormat::PF_A8R8G8B8].BlockBytes;
+				break;
+			case FVideoEncoderInputFrame::EUnderlyingRHI::D3D11:
+			case FVideoEncoderInputFrame::EUnderlyingRHI::D3D12:
+				InBuffer->Pitch = 0;
+				break;
+			default:
+				break;
+		}
+
 		InBuffer->Width = TextureSize.X;
 		InBuffer->Height = TextureSize.Y;
-		InBuffer->Pitch = TextureSize.X * 4;
 		InBuffer->BufferFormat = NV_ENC_BUFFER_FORMAT_ARGB;
 
 		RegisterParam.resourceType = NV_ENC_INPUT_RESOURCE_TYPE_CUDAARRAY;
diff --git a/Engine/Plugins/Media/HardwareEncoders/Source/EncoderNVENC/Private/NVENC_EncoderModule.cpp b/Engine/Plugins/Media/HardwareEncoders/Source/EncoderNVENC/Private/NVENC_EncoderModule.cpp
--- a/Engine/Plugins/Media/HardwareEncoders/Source/EncoderNVENC/Private/NVENC_EncoderModule.cpp
+++ b/Engine/Plugins/Media/HardwareEncoders/Source/EncoderNVENC/Private/NVENC_EncoderModule.cpp
@@ -40,12 +40,7 @@ public:
 #endif
 					VulkanRHIBridge::AddEnabledDeviceExtensionsAndLayers(ExtentionsToAdd, TArray<const ANSICHAR*>());
 				}
-
-#if PLATFORM_LINUX
 				FModuleManager::LoadModuleChecked<FCUDAModule>("CUDA").OnPostCUDAInit.AddLambda([]() {FVideoEncoderNVENC_H264::Register(FVideoEncoderFactory::Get());});
-#else
-				FCoreDelegates::OnPostEngineInit.AddLambda([]() {FVideoEncoderNVENC_H264::Register(FVideoEncoderFactory::Get());});
-#endif
 			}
 		}
 	}
diff --git a/Engine/Plugins/Media/PixelStreaming/Source/PixelStreaming/PixelStreaming.Build.cs b/Engine/Plugins/Media/PixelStreaming/Source/PixelStreaming/PixelStreaming.Build.cs
--- a/Engine/Plugins/Media/PixelStreaming/Source/PixelStreaming/PixelStreaming.Build.cs
+++ b/Engine/Plugins/Media/PixelStreaming/Source/PixelStreaming/PixelStreaming.Build.cs
@@ -68,6 +68,12 @@ namespace UnrealBuildTool.Rules
 
 			var EngineDir = Path.GetFullPath(Target.RelativeEnginePath);
 
+			// This is so for game projects using our public headers don't have to include extra modules they might not know about.
+			PublicDependencyModuleNames.AddRange(new string[]
+			{
+				"InputDevice"
+			});
+
 			// NOTE: General rule is not to access the private folder of another module
 			PrivateIncludePaths.AddRange(new string[]
 				{
@@ -96,12 +102,6 @@ namespace UnrealBuildTool.Rules
 				"AVEncoder"
 			});
 
-			// This is so for game projects using our public headers don't have to include extra modules they might not know about.
-			PublicDependencyModuleNames.AddRange(new string[]
-			{
-				"InputDevice"
-			});
-
 			AddEngineThirdPartyPrivateStaticDependencies(Target, "Vulkan");
 
 			// required for casting UE4 BackBuffer to Vulkan Texture2D for NvEnc
@@ -115,6 +115,20 @@ namespace UnrealBuildTool.Rules
 
 			if (Target.IsInPlatformGroup(UnrealPlatformGroup.Windows))
 			{
+				PublicDependencyModuleNames.Add("D3D11RHI");
+				PublicDependencyModuleNames.Add("D3D12RHI");
+				PrivateIncludePaths.AddRange(
+					new string[]{
+						Path.Combine(EngineDir, "Source/Runtime/D3D12RHI/Private"),
+						Path.Combine(EngineDir, "Source/Runtime/D3D12RHI/Private/Windows")
+					});
+
+				AddEngineThirdPartyPrivateStaticDependencies(Target, "DX12");
+				PublicSystemLibraries.AddRange(new string[] {
+					"DXGI.lib",
+					"d3d11.lib",
+				});
+
 				PrivateIncludePaths.Add(Path.Combine(EngineDir, "Source/Runtime/VulkanRHI/Private/Windows"));
 			}
 			else if (Target.IsInPlatformGroup(UnrealPlatformGroup.Linux))
diff --git a/Engine/Plugins/Media/PixelStreaming/Source/PixelStreaming/Private/PixelStreamingEncoderFrameFactory.cpp b/Engine/Plugins/Media/PixelStreaming/Source/PixelStreaming/Private/PixelStreamingEncoderFrameFactory.cpp
--- a/Engine/Plugins/Media/PixelStreaming/Source/PixelStreaming/Private/PixelStreamingEncoderFrameFactory.cpp
+++ b/Engine/Plugins/Media/PixelStreaming/Source/PixelStreaming/Private/PixelStreamingEncoderFrameFactory.cpp
@@ -5,6 +5,17 @@
 #include "CudaModule.h"
 #include "VulkanRHIPrivate.h"
 #include "PixelStreamingSettings.h"
+#include "VideoEncoderInput.h"
+
+#if PLATFORM_WINDOWS
+    #include "VideoCommon.h"
+    #include "D3D11State.h"
+    #include "D3D11Resources.h"
+    #include "D3D12RHICommon.h"
+    #include "D3D12RHIPrivate.h"
+    #include "D3D12Resources.h"
+    #include "D3D12Texture.h"
+#endif
 
 FPixelStreamingEncoderFrameFactory::FPixelStreamingEncoderFrameFactory()
 {
@@ -116,11 +127,25 @@ TSharedPtr<AVEncoder::FVideoEncoderInput> FPixelStreamingEncoderFrameFactory::Cr
 #if PLATFORM_WINDOWS
     else if(RHIName == TEXT("D3D11"))
     {
-        return AVEncoder::FVideoEncoderInput::CreateForD3D11(GDynamicRHI->RHIGetNativeDevice(), InWidth, InHeight, bIsResizable, IsRHIDeviceAMD());
+        if(IsRHIDeviceAMD())
+        {
+            return AVEncoder::FVideoEncoderInput::CreateForD3D11(GDynamicRHI->RHIGetNativeDevice(), InWidth, InHeight, bIsResizable, true);
+        }
+        else if(IsRHIDeviceNVIDIA())
+        {
+            return AVEncoder::FVideoEncoderInput::CreateForD3D11(GDynamicRHI->RHIGetNativeDevice(), InWidth, InHeight, bIsResizable, false);
+        }
     }
     else if(RHIName == TEXT("D3D12"))
     {
-        return AVEncoder::FVideoEncoderInput::CreateForD3D12(GDynamicRHI->RHIGetNativeDevice(), InWidth, InHeight, bIsResizable, IsRHIDeviceNVIDIA());
+        if(IsRHIDeviceAMD())
+        {
+            return AVEncoder::FVideoEncoderInput::CreateForD3D12(GDynamicRHI->RHIGetNativeDevice(), InWidth, InHeight, bIsResizable, false);
+        }
+        else if(IsRHIDeviceNVIDIA())
+        {
+            return AVEncoder::FVideoEncoderInput::CreateForCUDA(FModuleManager::GetModuleChecked<FCUDAModule>("CUDA").GetCudaContext(), InWidth, InHeight, bIsResizable);
+        }
     }
 #endif
 
@@ -138,7 +163,7 @@ void FPixelStreamingEncoderFrameFactory::SetTexture(AVEncoder::FVideoEncoderInpu
         if(IsRHIDeviceAMD())
         {
             FVulkanTexture2D* VulkanTexture = static_cast<FVulkanTexture2D*>(Texture.GetReference());
-            InputFrame->SetTexture(VulkanTexture->Surface.Image, [this, InputFrame](VkImage NativeTexture) { /* Do something with released texture if needed */ });
+            InputFrame->SetTexture(VulkanTexture->Surface.Image, [](VkImage NativeTexture) { /* Do something with released texture if needed */ });
         }
         else if(IsRHIDeviceNVIDIA())
         {
@@ -149,16 +174,28 @@ void FPixelStreamingEncoderFrameFactory::SetTexture(AVEncoder::FVideoEncoderInpu
             UE_LOG(PixelStreamer, Error, TEXT("Pixel Streaming only supports AMD and NVIDIA devices, this device is neither of those."));
         }
     }
-#if PLATFORM_WINDOWS	
+#if PLATFORM_WINDOWS
+    // TODO: Fix CUDA DX11 (Using CUDA as a bridge between DX11 and NVENC currently produces garbled results)	
     // DX11
     else if(RHIName == TEXT("D3D11"))
     {
-        InputFrame->SetTexture((ID3D11Texture2D*)Texture->GetNativeResource(), [this, InputFrame](ID3D11Texture2D* NativeTexture) { /* Do something with released texture if needed */ });
+        InputFrame->SetTexture((ID3D11Texture2D*)Texture->GetNativeResource(), [](ID3D11Texture2D* NativeTexture) { /* Do something with released texture if needed */ });
     }
     // DX12
     else if(RHIName == TEXT("D3D12"))
     {
-        InputFrame->SetTexture((ID3D12Resource*)Texture->GetNativeResource(), [this, InputFrame](ID3D12Resource* NativeTexture) { /* Do something with released texture if needed */ });
+        if(IsRHIDeviceAMD())
+        {
+            InputFrame->SetTexture((ID3D12Resource*)Texture->GetNativeResource(), [](ID3D12Resource* NativeTexture) { /* Do something with released texture if needed */ });
+        }
+        else if(IsRHIDeviceNVIDIA())
+        {
+            this->SetTextureCUDAD3D12(InputFrame, Texture);
+        }
+        else
+        {
+            UE_LOG(PixelStreamer, Error, TEXT("Pixel Streaming only supports AMD and NVIDIA devices, this device is neither of those."));
+        }
     }
 #endif // PLATFORM_WINDOWS
     else
@@ -241,7 +278,7 @@ void FPixelStreamingEncoderFrameFactory::SetTextureCUDAVulkan(AVEncoder::FVideoE
 
     FCUDAModule::CUDA().cuCtxPopCurrent(NULL);
 
-    InputFrame->SetTexture(mappedArray, [this, mappedArray, mappedMipArray, mappedExternalMemory, InputFrame](CUarray NativeTexture)
+    InputFrame->SetTexture(mappedArray, AVEncoder::FVideoEncoderInputFrame::EUnderlyingRHI::Vulkan, nullptr, [mappedArray, mappedMipArray, mappedExternalMemory](CUarray NativeTexture)
     {
         // free the cuda types
         FCUDAModule::CUDA().cuCtxPushCurrent(FModuleManager::GetModuleChecked<FCUDAModule>("CUDA").GetCudaContext());
@@ -275,4 +312,227 @@ void FPixelStreamingEncoderFrameFactory::SetTextureCUDAVulkan(AVEncoder::FVideoE
 
         FCUDAModule::CUDA().cuCtxPopCurrent(NULL);
     });
-}
\ No newline at end of file
+}
+
+#if PLATFORM_WINDOWS
+void FPixelStreamingEncoderFrameFactory::SetTextureCUDAD3D11(AVEncoder::FVideoEncoderInputFrame* InputFrame, const FTexture2DRHIRef& Texture)
+{
+    FD3D11TextureBase* D3D11Texture = GetD3D11TextureFromRHITexture(Texture);
+    unsigned long long TextureMemorySize = D3D11Texture->GetMemorySize();
+
+    ID3D11Texture2D* D3D11NativeTexture = static_cast<ID3D11Texture2D*>(D3D11Texture->GetResource());
+
+    TRefCountPtr<IDXGIResource> DXGIResource;
+    HRESULT Result = D3D11NativeTexture->QueryInterface(IID_PPV_ARGS(DXGIResource.GetInitReference()));
+	if(Result != S_OK) {
+        UE_LOG(PixelStreamer, Error, TEXT("Failed to get DX texture handle for importing memory to CUDA: %d"), Result);
+	}
+
+    HANDLE D3D11TextureHandle;
+	DXGIResource->GetSharedHandle(&D3D11TextureHandle); 
+	DXGIResource->Release();
+
+    FCUDAModule::CUDA().cuCtxPushCurrent(FModuleManager::GetModuleChecked<FCUDAModule>("CUDA").GetCudaContext());
+
+    CUexternalMemory mappedExternalMemory = nullptr;
+
+    {
+        // generate a cudaExternalMemoryHandleDesc
+        CUDA_EXTERNAL_MEMORY_HANDLE_DESC cudaExtMemHandleDesc = {};
+        cudaExtMemHandleDesc.type = CU_EXTERNAL_MEMORY_HANDLE_TYPE_D3D11_RESOURCE_KMT;
+        cudaExtMemHandleDesc.handle.win32.name = NULL;
+        cudaExtMemHandleDesc.handle.win32.handle = D3D11TextureHandle;
+        cudaExtMemHandleDesc.size = TextureMemorySize;
+        // Necessary for committed resources (DX11 and committed DX12 resources)
+		cudaExtMemHandleDesc.flags |= CUDA_EXTERNAL_MEMORY_DEDICATED;
+
+        // import external memory
+        auto result = FCUDAModule::CUDA().cuImportExternalMemory(&mappedExternalMemory, &cudaExtMemHandleDesc);
+        if(result != CUDA_SUCCESS)
+        {
+            UE_LOG(PixelStreamer, Error, TEXT("Failed to import external memory from vulkan error: %d"), result);
+        }
+    }
+
+    CUmipmappedArray mappedMipArray = nullptr;
+    CUarray mappedArray = nullptr;
+
+    {
+        CUDA_EXTERNAL_MEMORY_MIPMAPPED_ARRAY_DESC mipmapDesc = {};
+        mipmapDesc.numLevels = 1;
+        mipmapDesc.offset = 0;
+        mipmapDesc.arrayDesc.Width = Texture->GetSizeX();
+        mipmapDesc.arrayDesc.Height = Texture->GetSizeY();
+        mipmapDesc.arrayDesc.Depth = 1;
+        mipmapDesc.arrayDesc.NumChannels = 4;
+        mipmapDesc.arrayDesc.Format = CU_AD_FORMAT_UNSIGNED_INT8;
+        mipmapDesc.arrayDesc.Flags = CUDA_ARRAY3D_SURFACE_LDST | CUDA_ARRAY3D_COLOR_ATTACHMENT;
+
+        // get the CUarray from the external memory
+		CUresult result = FCUDAModule::CUDA().cuExternalMemoryGetMappedMipmappedArray(&mappedMipArray, mappedExternalMemory, &mipmapDesc);
+        if(result != CUDA_SUCCESS)
+        {
+            UE_LOG(PixelStreamer, Error, TEXT("Failed to bind mipmappedArray error: %d"), result);
+        }
+    }
+
+    // get the CUarray from the external memory
+    CUresult mipMapArrGetLevelErr = FCUDAModule::CUDA().cuMipmappedArrayGetLevel(&mappedArray, mappedMipArray, 0);
+    if(mipMapArrGetLevelErr != CUDA_SUCCESS)
+    {
+        UE_LOG(PixelStreamer, Error, TEXT("Failed to bind to mip 0."));
+    }
+
+    FCUDAModule::CUDA().cuCtxPopCurrent(NULL);
+
+    InputFrame->SetTexture(mappedArray, AVEncoder::FVideoEncoderInputFrame::EUnderlyingRHI::D3D11, nullptr, [mappedArray, mappedMipArray, mappedExternalMemory](CUarray NativeTexture)
+    {
+        // free the cuda types
+        FCUDAModule::CUDA().cuCtxPushCurrent(FModuleManager::GetModuleChecked<FCUDAModule>("CUDA").GetCudaContext());
+
+        if(mappedArray)
+        {
+            auto result = FCUDAModule::CUDA().cuArrayDestroy(mappedArray);
+            if(result != CUDA_SUCCESS)
+            {
+                UE_LOG(PixelStreamer, Error, TEXT("Failed to destroy mappedArray: %d"), result);
+            }
+        }
+
+        if(mappedMipArray)
+        {
+            auto result = FCUDAModule::CUDA().cuMipmappedArrayDestroy(mappedMipArray);
+            if(result != CUDA_SUCCESS)
+            {
+                UE_LOG(PixelStreamer, Error, TEXT("Failed to destroy mappedMipArray: %d"), result);
+            }
+        }
+
+        if(mappedExternalMemory)
+        {
+            auto result = FCUDAModule::CUDA().cuDestroyExternalMemory(mappedExternalMemory);
+            if(result != CUDA_SUCCESS)
+            {
+                UE_LOG(PixelStreamer, Error, TEXT("Failed to destroy mappedExternalMemoryArray: %d"), result);
+            }
+        }
+
+        FCUDAModule::CUDA().cuCtxPopCurrent(NULL);
+    });
+}
+
+void FPixelStreamingEncoderFrameFactory::SetTextureCUDAD3D12(AVEncoder::FVideoEncoderInputFrame* InputFrame, const FTexture2DRHIRef& Texture)
+{
+    FD3D12TextureBase* D3D12Texture = GetD3D12TextureFromRHITexture(Texture);
+    ID3D12Resource* NativeD3D12Resource = (ID3D12Resource*)Texture->GetNativeResource();
+    unsigned long long TextureMemorySize = D3D12Texture->GetMemorySize();
+
+    // Because we create our texture as RenderTargetable, it is created as a committed resource, which is what our current implementation here supports.
+    // To prevent a mystery crash in future, check that our resource is a committed resource
+    check(!D3D12Texture->GetResource()->IsPlacedResource());
+
+    TRefCountPtr<ID3D12Device> OwnerDevice;
+    HRESULT Result;
+    if((Result = NativeD3D12Resource->GetDevice(IID_PPV_ARGS(OwnerDevice.GetInitReference()))) != S_OK)
+    {
+        UE_LOG(PixelStreamer, Error, TEXT("Failed to get DX texture handle for importing memory to CUDA: %d (Get Device)"), Result);
+    }
+
+    //
+    // ID3D12Device::CreateSharedHandle gives as an NT Handle, and so we need to call CloseHandle on it
+    //
+    HANDLE D3D12TextureHandle;
+    if ((Result = OwnerDevice->CreateSharedHandle(NativeD3D12Resource, NULL, GENERIC_ALL, NULL, &D3D12TextureHandle)) != S_OK)
+    {
+        UE_LOG(PixelStreamer, Error, TEXT("Failed to get DX texture handle for importing memory to CUDA: %d"), Result);
+    }
+
+    FCUDAModule::CUDA().cuCtxPushCurrent(FModuleManager::GetModuleChecked<FCUDAModule>("CUDA").GetCudaContext());
+
+    CUexternalMemory mappedExternalMemory = nullptr;
+
+    {
+        // generate a cudaExternalMemoryHandleDesc
+        CUDA_EXTERNAL_MEMORY_HANDLE_DESC cudaExtMemHandleDesc = {};
+        cudaExtMemHandleDesc.type = CU_EXTERNAL_MEMORY_HANDLE_TYPE_D3D12_RESOURCE;
+        cudaExtMemHandleDesc.handle.win32.name = NULL;
+        cudaExtMemHandleDesc.handle.win32.handle = D3D12TextureHandle;
+        cudaExtMemHandleDesc.size = TextureMemorySize;
+        // Necessary for committed resources (DX11 and committed DX12 resources)
+		cudaExtMemHandleDesc.flags |= CUDA_EXTERNAL_MEMORY_DEDICATED;
+
+        // import external memory
+        auto result = FCUDAModule::CUDA().cuImportExternalMemory(&mappedExternalMemory, &cudaExtMemHandleDesc);
+        if(result != CUDA_SUCCESS)
+        {
+            UE_LOG(PixelStreamer, Error, TEXT("Failed to import external memory from vulkan error: %d"), result);
+        }
+    }
+
+    CUmipmappedArray mappedMipArray = nullptr;
+    CUarray mappedArray = nullptr;
+
+    {
+        CUDA_EXTERNAL_MEMORY_MIPMAPPED_ARRAY_DESC mipmapDesc = {};
+        mipmapDesc.numLevels = 1;
+        mipmapDesc.offset = 0;
+        mipmapDesc.arrayDesc.Width = Texture->GetSizeX();
+        mipmapDesc.arrayDesc.Height = Texture->GetSizeY();
+        mipmapDesc.arrayDesc.Depth = 1;
+        mipmapDesc.arrayDesc.NumChannels = 4;
+        mipmapDesc.arrayDesc.Format = CU_AD_FORMAT_UNSIGNED_INT8;
+        mipmapDesc.arrayDesc.Flags = CUDA_ARRAY3D_SURFACE_LDST | CUDA_ARRAY3D_COLOR_ATTACHMENT;
+
+        // get the CUarray from the external memory
+		CUresult result = FCUDAModule::CUDA().cuExternalMemoryGetMappedMipmappedArray(&mappedMipArray, mappedExternalMemory, &mipmapDesc);
+        if(result != CUDA_SUCCESS)
+        {
+            UE_LOG(PixelStreamer, Error, TEXT("Failed to bind mipmappedArray error: %d"), result);
+        }
+    }
+
+    // get the CUarray from the external memory
+    CUresult mipMapArrGetLevelErr = FCUDAModule::CUDA().cuMipmappedArrayGetLevel(&mappedArray, mappedMipArray, 0);
+    if(mipMapArrGetLevelErr != CUDA_SUCCESS)
+    {
+        UE_LOG(PixelStreamer, Error, TEXT("Failed to bind to mip 0."));
+    }
+
+    FCUDAModule::CUDA().cuCtxPopCurrent(NULL);
+
+    InputFrame->SetTexture(mappedArray, AVEncoder::FVideoEncoderInputFrame::EUnderlyingRHI::D3D12, D3D12TextureHandle, [mappedArray, mappedMipArray, mappedExternalMemory](CUarray NativeTexture)
+    {
+        // free the cuda types
+        FCUDAModule::CUDA().cuCtxPushCurrent(FModuleManager::GetModuleChecked<FCUDAModule>("CUDA").GetCudaContext());
+
+        if(mappedArray)
+        {
+            auto result = FCUDAModule::CUDA().cuArrayDestroy(mappedArray);
+            if(result != CUDA_SUCCESS)
+            {
+                UE_LOG(PixelStreamer, Error, TEXT("Failed to destroy mappedArray: %d"), result);
+            }
+        }
+
+        if(mappedMipArray)
+        {
+            auto result = FCUDAModule::CUDA().cuMipmappedArrayDestroy(mappedMipArray);
+            if(result != CUDA_SUCCESS)
+            {
+                UE_LOG(PixelStreamer, Error, TEXT("Failed to destroy mappedMipArray: %d"), result);
+            }
+        }
+
+        if(mappedExternalMemory)
+        {
+            auto result = FCUDAModule::CUDA().cuDestroyExternalMemory(mappedExternalMemory);
+            if(result != CUDA_SUCCESS)
+            {
+                UE_LOG(PixelStreamer, Error, TEXT("Failed to destroy mappedExternalMemoryArray: %d"), result);
+            }
+        }
+
+        FCUDAModule::CUDA().cuCtxPopCurrent(NULL);
+    });
+}
+#endif //PLATFORM_WINDOWS
\ No newline at end of file
diff --git a/Engine/Plugins/Media/PixelStreaming/Source/PixelStreaming/Private/PixelStreamingEncoderFrameFactory.h b/Engine/Plugins/Media/PixelStreaming/Source/PixelStreaming/Private/PixelStreamingEncoderFrameFactory.h
--- a/Engine/Plugins/Media/PixelStreaming/Source/PixelStreaming/Private/PixelStreamingEncoderFrameFactory.h
+++ b/Engine/Plugins/Media/PixelStreaming/Source/PixelStreaming/Private/PixelStreamingEncoderFrameFactory.h
@@ -26,6 +26,10 @@ class FPixelStreamingEncoderFrameFactory
         TSharedPtr<AVEncoder::FVideoEncoderInput> CreateVideoEncoderInput(int InWidth, int InHeight) const;
         void SetTexture(AVEncoder::FVideoEncoderInputFrame* Frame, const FTexture2DRHIRef& Texture);
         void SetTextureCUDAVulkan(AVEncoder::FVideoEncoderInputFrame* Frame, const FTexture2DRHIRef& Texture);
+#if PLATFORM_WINDOWS        
+		void SetTextureCUDAD3D11(AVEncoder::FVideoEncoderInputFrame* Frame, const FTexture2DRHIRef& Texture);
+		void SetTextureCUDAD3D12(AVEncoder::FVideoEncoderInputFrame* Frame, const FTexture2DRHIRef& Texture);
+#endif // PLATFORM_WINDOWS
 
     private:
 
diff --git a/Engine/Source/Runtime/AVEncoder/Private/VideoEncoderInput.cpp b/Engine/Source/Runtime/AVEncoder/Private/VideoEncoderInput.cpp
--- a/Engine/Source/Runtime/AVEncoder/Private/VideoEncoderInput.cpp
+++ b/Engine/Source/Runtime/AVEncoder/Private/VideoEncoderInput.cpp
@@ -789,6 +789,7 @@ namespace AVEncoder
 			{
 				UE_LOG(LogVideoEncoder, Warning, TEXT("VideoEncoderFame - D3D11 input texture still holds %d references."), NumRef);
 			}
+			// Need to call release twice as we have added an extra reference just above (only way to count how many references we have)
 			D3D11.EncoderTexture->Release();
 			D3D11.EncoderTexture->Release();
 			D3D11.EncoderTexture = nullptr;
@@ -810,6 +811,7 @@ namespace AVEncoder
 			{
 				UE_LOG(LogVideoEncoder, Warning, TEXT("VideoEncoderFame - D3D12 input texture still holds %d references."), NumRef);
 			}
+			// Need to call release twice as we have added an extra reference just above (only way to count how many references we have)
 			D3D12.EncoderTexture->Release();
 			D3D12.EncoderTexture->Release();
 			D3D12.EncoderTexture = nullptr;
@@ -819,6 +821,13 @@ namespace AVEncoder
 			OnReleaseD3D12Texture(D3D12.Texture);
 			D3D12.Texture = nullptr;
 		}
+
+		// D3D12 specific handle atm
+		if (CUDA.SharedHandle)
+		{
+			CloseHandle(CUDA.SharedHandle);
+			CUDA.SharedHandle = nullptr;
+		}
 #endif 
 
 		if (CUDA.EncoderTexture)
@@ -970,10 +979,12 @@ namespace AVEncoder
 	}
 #endif
 
-	void FVideoEncoderInputFrame::SetTexture(CUarray InTexture, FReleaseCUDATextureCallback InOnReleaseTexture)
+	void FVideoEncoderInputFrame::SetTexture(CUarray InTexture, EUnderlyingRHI UnderlyingRHI, void* SharedHandle, FReleaseCUDATextureCallback InOnReleaseTexture)
 	{
 		if (Format == EVideoFrameFormat::CUDA_R8G8B8A8_UNORM)
 		{
+			CUDA.UnderlyingRHI = UnderlyingRHI;
+			CUDA.SharedHandle = SharedHandle;
 			CUDA.EncoderTexture = InTexture;
 			OnReleaseCUDATexture = InOnReleaseTexture;
 			if (!CUDA.EncoderTexture)
diff --git a/Engine/Source/Runtime/AVEncoder/Public/VideoEncoderInput.h b/Engine/Source/Runtime/AVEncoder/Public/VideoEncoderInput.h
--- a/Engine/Source/Runtime/AVEncoder/Public/VideoEncoderInput.h
+++ b/Engine/Source/Runtime/AVEncoder/Public/VideoEncoderInput.h
@@ -210,11 +210,21 @@ namespace AVEncoder
 
 #endif // PLATFORM_WINDOWS
 
+		enum class EUnderlyingRHI
+		{
+			Undefined,
+			D3D11,
+			D3D12,
+			Vulkan
+		};
+
 		// --- CUDA
 		struct FCUDA
 		{
-			CUarray		EncoderTexture = nullptr;
-			CUcontext   EncoderDevice = nullptr;
+			CUarray			EncoderTexture = nullptr;
+			CUcontext   	EncoderDevice = nullptr;
+			EUnderlyingRHI	UnderlyingRHI = EUnderlyingRHI::Undefined;
+			void*			SharedHandle = nullptr;
 		};
 
 		const FCUDA& GetCUDA() const { return CUDA; }
@@ -223,7 +233,7 @@ namespace AVEncoder
 		// the callback type used to create a registered encoder
 		using FReleaseCUDATextureCallback = TFunction<void(CUarray)>;
 
-		void SetTexture(CUarray InTexture, FReleaseCUDATextureCallback InOnReleaseTexture);
+		void SetTexture(CUarray InTexture, EUnderlyingRHI UnderlyingRHI, void* SharedHandle, FReleaseCUDATextureCallback InOnReleaseTexture);
 
 #if PLATFORM_DESKTOP && !PLATFORM_APPLE
 		// --- Vulkan
diff --git a/Engine/Source/Runtime/CUDA/Source/CUDA.Build.cs b/Engine/Source/Runtime/CUDA/Source/CUDA.Build.cs
--- a/Engine/Source/Runtime/CUDA/Source/CUDA.Build.cs
+++ b/Engine/Source/Runtime/CUDA/Source/CUDA.Build.cs
@@ -19,7 +19,8 @@ public class CUDA : ModuleRules
 		PublicDependencyModuleNames.Add("CUDAHeader");
 
 		if (Target.IsInPlatformGroup(UnrealPlatformGroup.Windows) || Target.IsInPlatformGroup(UnrealPlatformGroup.Unix))
-		{		
+		{
+			PublicDefinitions.Add("PLATFORM_SUPPORTS_CUDA=1");
 			var EngineDir = Path.GetFullPath(Target.RelativeEnginePath);
 
 			AddEngineThirdPartyPrivateStaticDependencies(Target, "NVAftermath");
@@ -30,13 +31,22 @@ public class CUDA : ModuleRules
 			if (Target.IsInPlatformGroup(UnrealPlatformGroup.Unix))
 			{
 				PrivateIncludePaths.Add(Path.Combine(EngineDir, "Source/Runtime/VulkanRHI/Private/Linux"));	
-			} 
-			else if (Target.Platform == UnrealTargetPlatform.Win64)
+			} else if (Target.IsInPlatformGroup(UnrealPlatformGroup.Windows))
 			{
-				PrivateIncludePaths.Add(Path.Combine(EngineDir, "Source/Runtime/VulkanRHI/Private/Windows"));	
+				PublicSystemLibraries.AddRange(new string[] {
+					"DXGI.lib",
+					"d3d11.lib",
+					"d3d12.lib"
+				});
+
+				PrivateIncludePaths.Add(Path.Combine(EngineDir, "Source/Runtime/VulkanRHI/Private/Windows"));
 			}
 
 			AddEngineThirdPartyPrivateStaticDependencies(Target, "Vulkan");
 		}
+		else
+		{
+			PublicDefinitions.Add("PLATFORM_SUPPORTS_CUDA=0");
+		}
 	}
 }
diff --git a/Engine/Source/Runtime/CUDA/Source/Private/CudaModule.cpp b/Engine/Source/Runtime/CUDA/Source/Private/CudaModule.cpp
--- a/Engine/Source/Runtime/CUDA/Source/Private/CudaModule.cpp
+++ b/Engine/Source/Runtime/CUDA/Source/Private/CudaModule.cpp
@@ -3,10 +3,19 @@
 
 #include "Misc/CoreDelegates.h"
 
-#if PLATFORM_DESKTOP && !PLATFORM_APPLE
+#if PLATFORM_SUPPORTS_CUDA
 #include "VulkanRHIPrivate.h"
 #endif
 
+#if PLATFORM_WINDOWS
+	#include "DynamicRHI.h"
+    #include "Windows/AllowWindowsPlatformTypes.h"
+	#include "dxgi.h"
+    #include "d3d12.h"
+    #include "d3d11.h"
+    #include "Windows/HideWindowsPlatformTypes.h"
+#endif
+
 #include "HAL/UnrealMemory.h"
 
 DEFINE_LOG_CATEGORY(LogCUDA);
@@ -84,22 +93,111 @@ void FCUDAModule::UnloadCuda()
 	FMemory::Memset(&DriverApiPtrs, 0, sizeof(CUDA_DRIVER_API_FUNCTION_LIST));
 }
 
+bool FCUDAModule::IsRHISelectedDevice(CUdevice cuDevice) {
+	FString RHIName = GDynamicRHI->GetName();
+
+#if PLATFORM_SUPPORTS_CUDA
+    // VULKAN
+    if(RHIName == TEXT("Vulkan"))
+    {
+		uint8 deviceUUID[16];
+
+		// We find the device that the RHI has selected for us so that later we can create a CUDA context on this device.
+		FVulkanDynamicRHI *vkDynamicRHI = static_cast<FVulkanDynamicRHI *>(GDynamicRHI);
+		FMemory::Memcpy(deviceUUID, vkDynamicRHI->GetDevice()->GetDeviceIdProperties().deviceUUID, 16);	
+
+		// Get the device UUID so we can compare this with what the RHI selected.
+		CUuuid cudaDeviceUUID;
+		CUresult getUuidErr = DriverApiPtrs.cuDeviceGetUuid(&cudaDeviceUUID, cuDevice);
+		if(getUuidErr != CUDA_SUCCESS)
+		{
+			UE_LOG(LogCUDA, Warning, TEXT("Could not get CUDA device UUID at device %d."), cuDevice);
+			return false;
+		}
+
+		// Compare the CUDA device UUID with RHI selected UUID.
+		return FMemory::Memcmp(&cudaDeviceUUID, deviceUUID, 16) == 0;
+    }
+	else
+#if PLATFORM_WINDOWS
+    // DX11
+    if(RHIName == TEXT("D3D11"))
+    {
+		char* DeviceLuid = new char[64];
+		unsigned int DeviceNodeMask;
+		CUresult getLuidErr = DriverApiPtrs.cuDeviceGetLuid(DeviceLuid, &DeviceNodeMask, cuDevice);
+		if(getLuidErr != CUDA_SUCCESS)
+		{
+			UE_LOG(LogCUDA, Warning, TEXT("Could not get CUDA device LUID at device %d."), cuDevice);
+			return false;
+		}
+
+		ID3D11Device *NativeD3D11Device = static_cast<ID3D11Device *>(GDynamicRHI->RHIGetNativeDevice());
+
+		TRefCountPtr<IDXGIDevice> DXGIDevice;
+		
+		HRESULT Result = NativeD3D11Device->QueryInterface(__uuidof(IDXGIDevice), (void**)DXGIDevice.GetInitReference());
+		if(FAILED(Result)) {
+			UE_LOG(LogCUDA, Error, TEXT("Failed to get DXGIDevice when trying to resolve to CUDA device at %s:%u \n with error %u"), ANSI_TO_TCHAR(__FILE__), __LINE__, Result);
+			return false;
+		}
+
+		//VERIFYD3D11RESULT_EX(GDynamicRHI->GetDevice()->QueryInterface(IID_IDXGIDevice, (void**)DXGIDevice.GetInitReference()), GDynamicRHI->GetDevice());
+
+		TRefCountPtr<IDXGIAdapter> DXGIAdapter;
+		DXGIDevice->GetAdapter((IDXGIAdapter**)DXGIAdapter.GetInitReference());
+
+		DXGI_ADAPTER_DESC AdapterDesc;
+		DXGIAdapter->GetDesc(&AdapterDesc);
+
+		LUID AdapterLUID = AdapterDesc.AdapterLuid;
+	
+		return ((memcmp(&AdapterLUID.LowPart, DeviceLuid,
+					sizeof(AdapterLUID.LowPart)) == 0) &&
+			(memcmp(&AdapterLUID.HighPart,
+					DeviceLuid + sizeof(AdapterLUID.LowPart),
+					sizeof(AdapterLUID.HighPart)) == 0));
+    }
+    // DX12
+    else if(RHIName == TEXT("D3D12"))
+	{
+		char* DeviceLuid = new char[64];
+		unsigned int DeviceNodeMask;
+		CUresult getLuidErr = DriverApiPtrs.cuDeviceGetLuid(DeviceLuid, &DeviceNodeMask, cuDevice);
+		if(getLuidErr != CUDA_SUCCESS)
+		{
+			UE_LOG(LogCUDA, Warning, TEXT("Could not get CUDA device LUID at device %d."), cuDevice);
+			return false;
+		}
+
+		//FD3D12DynamicRHI *D3D12DynamicRHI = static_cast<FD3D12DynamicRHI *>(GDynamicRHI);
+		//FD3D12Device* D3D12Device = static_cast<FD3D12Device>(D3D12DynamicRHI->GetDevice();
+		ID3D12Device* NativeD3D12Device = static_cast<ID3D12Device *>(GDynamicRHI->RHIGetNativeDevice());
+
+		LUID AdapterLUID = NativeD3D12Device->GetAdapterLuid();
+
+		return ((memcmp(&AdapterLUID.LowPart, DeviceLuid,
+					sizeof(AdapterLUID.LowPart)) == 0) &&
+			(memcmp(&AdapterLUID.HighPart,
+					DeviceLuid + sizeof(AdapterLUID.LowPart),
+					sizeof(AdapterLUID.HighPart)) == 0));
+	}
+	else
+#endif //PLATFORM_WINDOWS
+#endif //PLATFORM_SUPPORTS_CUDA
+	{
+		UE_LOG(LogCUDA, Fatal, TEXT("Unsupported RHI or Platform for CUDA."));
+	}
+
+	return false;
+}
+
 void FCUDAModule::InitCuda()
 {
-// Don't Init CUDA on unsupported platforms	
-#if PLATFORM_DESKTOP && !PLATFORM_APPLE
-
 	// Initialise to rhiDeviceIndex -1, which is an invalid device index, if it remains -1 we will know no device was found.
 	rhiDeviceIndex = -1;
 
-	// TODO: add support for other RHIs (e.g. DX12)
-	// For now simply exit early if some other RHI is used.
-	if (GDynamicRHI != nullptr && GDynamicRHI->GetName() != FString(TEXT("Vulkan"))) 
-	{
-		UE_LOG(LogCUDA, Warning, TEXT("CUDA module only supports the Vulkan RHI presently. RHI detected: %s"), GDynamicRHI->GetName());
-		return;
-	}
-	
+#if PLATFORM_SUPPORTS_CUDA
 	// Initialise CUDA API
 	{
 		UE_LOG(LogCUDA, Display, TEXT("Initialising CUDA API..."));
@@ -114,12 +212,6 @@ void FCUDAModule::InitCuda()
 			UE_LOG(LogCUDA, Fatal, TEXT("CUDA API failed to initialise."));
 		}
 	}
-	
-	uint8 deviceUUID[16];
-
-	// We find the device that the RHI has selected for us so that later we can create a CUDA context on this device.
-	FVulkanDynamicRHI *vkDynamicRHI = static_cast<FVulkanDynamicRHI *>(GDynamicRHI);
-	FMemory::Memcpy(deviceUUID, vkDynamicRHI->GetDevice()->GetDeviceIdProperties().deviceUUID, 16);	
 
 	// Find out how many graphics devices there are so we find that one that matches the one the RHI selected.
 	int device_count = 0;
@@ -136,9 +228,9 @@ void FCUDAModule::InitCuda()
 	if (device_count == 0)
 	{
 		UE_LOG(LogCUDA, Error, TEXT("There are no available device(s) that support CUDA. If that is untrue check CUDA is installed."));
+		return;
 	}
 	
-	CUuuid    cudaDeviceUUID;
 	CUdevice  cuDevice;
 	CUcontext cudaContext;
 	
@@ -152,14 +244,6 @@ void FCUDAModule::InitCuda()
 			UE_LOG(LogCUDA, Warning, TEXT("Could not get CUDA device at device %d."), current_device);
 			continue;
 		}
-
-		// Get the device UUID so we can compare this with what the RHI selected.
-		CUresult getUuidErr = DriverApiPtrs.cuDeviceGetUuid(&cudaDeviceUUID, cuDevice);
-		if(getUuidErr != CUDA_SUCCESS)
-		{
-			UE_LOG(LogCUDA, Warning, TEXT("Could not get CUDA device UUID at device %d."), current_device);
-			continue;
-		}
 		
 		// Get device name for logging purposes.
 		char deviceNameArr[256];
@@ -168,43 +252,53 @@ void FCUDAModule::InitCuda()
 		if(getNameErr == CUDA_SUCCESS) 
 		{
 			deviceName = FString(UTF8_TO_TCHAR(deviceNameArr));
-			UE_LOG(LogCUDA, Display, TEXT("Found device %d called %s."), current_device, *deviceName);
+			UE_LOG(LogCUDA, Display, TEXT("Found device %d called \"%s\"."), current_device, *deviceName);
 		}
 		else 
 		{
 			UE_LOG(LogCUDA, Warning, TEXT("Could not get name of CUDA device %d."), current_device);
 		}
 		
-		// Compare the CUDA device UUID with RHI selected UUID.
-		bool bIsRHISelectedDevice = FMemory::Memcmp(&cudaDeviceUUID, deviceUUID, 16) == 0;
-		
-		if (bIsRHISelectedDevice)
+		// Compare the CUDA device with current RHI selected device.		
+		if (this->IsRHISelectedDevice(cuDevice))
 		{
 			UE_LOG(LogCUDA, Display, TEXT("Attempting to create CUDA context on GPU Device %d..."), current_device);
 			
 			CUresult createCtxErr = DriverApiPtrs.cuCtxCreate(&cudaContext, 0, current_device);
 			if(createCtxErr == CUDA_SUCCESS)
 			{
-				UE_LOG(LogCUDA, Display, TEXT("Created CUDA context on device %s!"), *deviceName);
+				UE_LOG(LogCUDA, Display, TEXT("Created CUDA context on device \"%s\"."), *deviceName);
 				contextMap.Add(current_device, cudaContext);
 				rhiDeviceIndex = current_device;
 				break;
 			}
 			else
 			{
-				UE_LOG(LogCUDA, Warning, TEXT("Could not create CUDA context on device %s."), *deviceName);
-				continue;
+				UE_LOG(LogCUDA, Error, TEXT("CUDA module could not create CUDA context on RHI device \"%s\"."), *deviceName);
+				return;
 			}
 		}
+		else
+		{
+			UE_LOG(LogCUDA, Verbose, TEXT("CUDA device %d \"%s\" does not match RHI device."), current_device, *deviceName);
+		}
 	}
-	
+
 	if (rhiDeviceIndex == -1)
 	{
-		UE_LOG(LogCUDA, Fatal, TEXT("CUDA module failed to create a CUDA context on the RHI selected device with UUID %d."), deviceUUID);
+		UE_LOG(LogCUDA, Error, TEXT("CUDA module failed to create a CUDA context on the RHI selected device."));
+		return;
+	}
+	else
+#else
+
+	UE_LOG(LogCUDA, Error, TEXT("CUDA attemped to be initialized on unsupported Platform (Currently supported platforms are Win32, Win64, Linux and LinuxAArch64)."));
+
+#endif //PLATFORM_SUPPORTS_CUDA
+
+	{
+		OnPostCUDAInit.Broadcast();
 	}
-	
-#endif
-	OnPostCUDAInit.Broadcast();
 }
 
 CUcontext FCUDAModule::GetCudaContext()
diff --git a/Engine/Source/Runtime/CUDA/Source/Public/CudaModule.h b/Engine/Source/Runtime/CUDA/Source/Public/CudaModule.h
--- a/Engine/Source/Runtime/CUDA/Source/Public/CudaModule.h
+++ b/Engine/Source/Runtime/CUDA/Source/Public/CudaModule.h
@@ -39,6 +39,7 @@ private:
 	void UnloadCuda();
 	
 	void InitCuda();
+	bool IsRHISelectedDevice(CUdevice cuDevice);
 	
 	void* DriverLibrary;
 	CUDA_DRIVER_API_FUNCTION_LIST DriverApiPtrs;
