Fix missing header by including module

diff --git a/Engine/Plugins/Media/PixelStreaming/Source/PixelStreaming/PixelStreaming.Build.cs b/Engine/Plugins/Media/PixelStreaming/Source/PixelStreaming/PixelStreaming.Build.cs
--- a/Engine/Plugins/Media/PixelStreaming/Source/PixelStreaming/PixelStreaming.Build.cs
+++ b/Engine/Plugins/Media/PixelStreaming/Source/PixelStreaming/PixelStreaming.Build.cs
@@ -91,6 +91,7 @@ namespace UnrealBuildTool.Rules
 				"Json",
 				"RenderCore",
 				"RHI",
+				"RHICore",
 				"Slate",
 				"SlateCore",
 				"AudioMixer",
