Merging //UE5/Dev-TensorWorks to Release-5.0 (//UE5/Release-5.0)
Adding copyright notices.

diff --git a/Engine/Plugins/Media/PixelStreaming/Source/PixelStreaming/Private/PixelStreamingFrameBuffer.cpp b/Engine/Plugins/Media/PixelStreaming/Source/PixelStreaming/Private/PixelStreamingFrameBuffer.cpp
--- a/Engine/Plugins/Media/PixelStreaming/Source/PixelStreaming/Private/PixelStreamingFrameBuffer.cpp
+++ b/Engine/Plugins/Media/PixelStreaming/Source/PixelStreaming/Private/PixelStreamingFrameBuffer.cpp
@@ -1,5 +1,5 @@
+// Copyright Epic Games, Inc. All Rights Reserved.
 #include "PixelStreamingFrameBuffer.h"
-
 #include "PixelStreamingFrameSource.h"
 
 FPixelStreamingSimulcastFrameBuffer::FPixelStreamingSimulcastFrameBuffer(FPixelStreamingFrameSource* InFrameSource)
diff --git a/Engine/Plugins/Media/PixelStreaming/Source/PixelStreaming/Private/PixelStreamingFramePump.cpp b/Engine/Plugins/Media/PixelStreaming/Source/PixelStreaming/Private/PixelStreamingFramePump.cpp
--- a/Engine/Plugins/Media/PixelStreaming/Source/PixelStreaming/Private/PixelStreamingFramePump.cpp
+++ b/Engine/Plugins/Media/PixelStreaming/Source/PixelStreaming/Private/PixelStreamingFramePump.cpp
@@ -1,3 +1,4 @@
+// Copyright Epic Games, Inc. All Rights Reserved.
 #include "PixelStreamingFramePump.h"
 
 #include "PlayerVideoSource.h"
diff --git a/Engine/Plugins/Media/PixelStreaming/Source/PixelStreaming/Private/PixelStreamingFramePump.h b/Engine/Plugins/Media/PixelStreaming/Source/PixelStreaming/Private/PixelStreamingFramePump.h
--- a/Engine/Plugins/Media/PixelStreaming/Source/PixelStreaming/Private/PixelStreamingFramePump.h
+++ b/Engine/Plugins/Media/PixelStreaming/Source/PixelStreaming/Private/PixelStreamingFramePump.h
@@ -1,3 +1,4 @@
+// Copyright Epic Games, Inc. All Rights Reserved.
 #pragma once
 
 #include "PlayerId.h"
diff --git a/Engine/Plugins/Media/PixelStreaming/Source/PixelStreaming/Private/PixelStreamingFrameSource.cpp b/Engine/Plugins/Media/PixelStreaming/Source/PixelStreaming/Private/PixelStreamingFrameSource.cpp
--- a/Engine/Plugins/Media/PixelStreaming/Source/PixelStreaming/Private/PixelStreamingFrameSource.cpp
+++ b/Engine/Plugins/Media/PixelStreaming/Source/PixelStreaming/Private/PixelStreamingFrameSource.cpp
@@ -1,3 +1,4 @@
+// Copyright Epic Games, Inc. All Rights Reserved.
 #include "PixelStreamingFrameSource.h"
 
 #include "LatencyTester.h"
diff --git a/Engine/Plugins/Media/PixelStreaming/Source/PixelStreaming/Private/PixelStreamingFrameSource.h b/Engine/Plugins/Media/PixelStreaming/Source/PixelStreaming/Private/PixelStreamingFrameSource.h
--- a/Engine/Plugins/Media/PixelStreaming/Source/PixelStreaming/Private/PixelStreamingFrameSource.h
+++ b/Engine/Plugins/Media/PixelStreaming/Source/PixelStreaming/Private/PixelStreamingFrameSource.h
@@ -1,3 +1,4 @@
+// Copyright Epic Games, Inc. All Rights Reserved.
 #pragma once
 
 #include "RHI.h"
diff --git a/Engine/Plugins/Media/PixelStreaming/Source/PixelStreaming/Private/PixelStreamingRealEncoder.cpp b/Engine/Plugins/Media/PixelStreaming/Source/PixelStreaming/Private/PixelStreamingRealEncoder.cpp
--- a/Engine/Plugins/Media/PixelStreaming/Source/PixelStreaming/Private/PixelStreamingRealEncoder.cpp
+++ b/Engine/Plugins/Media/PixelStreaming/Source/PixelStreaming/Private/PixelStreamingRealEncoder.cpp
@@ -1,3 +1,4 @@
+// Copyright Epic Games, Inc. All Rights Reserved.
 #include "PixelStreamingRealEncoder.h"
 #include "Utils.h"
 #include "VulkanRHIPrivate.h"
diff --git a/Engine/Plugins/Media/PixelStreaming/Source/PixelStreaming/Private/PixelStreamingRealEncoder.h b/Engine/Plugins/Media/PixelStreaming/Source/PixelStreaming/Private/PixelStreamingRealEncoder.h
--- a/Engine/Plugins/Media/PixelStreaming/Source/PixelStreaming/Private/PixelStreamingRealEncoder.h
+++ b/Engine/Plugins/Media/PixelStreaming/Source/PixelStreaming/Private/PixelStreamingRealEncoder.h
@@ -1,3 +1,4 @@
+// Copyright Epic Games, Inc. All Rights Reserved.
 #pragma once
 
 #include "VideoEncoder.h"
diff --git a/Engine/Plugins/Media/PixelStreaming/Source/PixelStreaming/Private/PixelStreamingSimulcastEncoderAdapter.cpp b/Engine/Plugins/Media/PixelStreaming/Source/PixelStreaming/Private/PixelStreamingSimulcastEncoderAdapter.cpp
--- a/Engine/Plugins/Media/PixelStreaming/Source/PixelStreaming/Private/PixelStreamingSimulcastEncoderAdapter.cpp
+++ b/Engine/Plugins/Media/PixelStreaming/Source/PixelStreaming/Private/PixelStreamingSimulcastEncoderAdapter.cpp
@@ -1,13 +1,4 @@
-/*
- *  Copyright (c) 2014 The WebRTC project authors. All Rights Reserved.
- *
- *  Use of this source code is governed by a BSD-style license
- *  that can be found in the LICENSE file in the root of the source
- *  tree. An additional intellectual property rights grant can be found
- *  in the file PATENTS.  All contributing project authors may
- *  be found in the AUTHORS file in the root of the source tree.
- */
-
+// Copyright Epic Games, Inc. All Rights Reserved.
 #include "PixelStreamingSimulcastEncoderAdapter.h"
 
 #include <stdio.h>
diff --git a/Engine/Plugins/Media/PixelStreaming/Source/PixelStreaming/Private/PixelStreamingSimulcastEncoderAdapter.h b/Engine/Plugins/Media/PixelStreaming/Source/PixelStreaming/Private/PixelStreamingSimulcastEncoderAdapter.h
--- a/Engine/Plugins/Media/PixelStreaming/Source/PixelStreaming/Private/PixelStreamingSimulcastEncoderAdapter.h
+++ b/Engine/Plugins/Media/PixelStreaming/Source/PixelStreaming/Private/PixelStreamingSimulcastEncoderAdapter.h
@@ -1,13 +1,4 @@
-/*
- *  Copyright (c) 2014 The WebRTC project authors. All Rights Reserved.
- *
- *  Use of this source code is governed by a BSD-style license
- *  that can be found in the LICENSE file in the root of the source
- *  tree. An additional intellectual property rights grant can be found
- *  in the file PATENTS.  All contributing project authors may
- *  be found in the AUTHORS file in the root of the source tree.
- *
- */
+// Copyright Epic Games, Inc. All Rights Reserved.
 #pragma once
 
 #include <memory>
diff --git a/Engine/Plugins/Media/PixelStreaming/Source/PixelStreaming/Private/PlayerVideoSource.cpp b/Engine/Plugins/Media/PixelStreaming/Source/PixelStreaming/Private/PlayerVideoSource.cpp
--- a/Engine/Plugins/Media/PixelStreaming/Source/PixelStreaming/Private/PlayerVideoSource.cpp
+++ b/Engine/Plugins/Media/PixelStreaming/Source/PixelStreaming/Private/PlayerVideoSource.cpp
@@ -1,3 +1,4 @@
+// Copyright Epic Games, Inc. All Rights Reserved.
 #include "PlayerVideoSource.h"
 
 FPlayerVideoSource::FPlayerVideoSource()
