Non-unity build fixes

diff --git a/Engine/Plugins/Media/PixelStreaming/Source/PixelStreaming/Private/PixelStreamingFrameBuffer.h b/Engine/Plugins/Media/PixelStreaming/Source/PixelStreaming/Private/PixelStreamingFrameBuffer.h
--- a/Engine/Plugins/Media/PixelStreaming/Source/PixelStreaming/Private/PixelStreamingFrameBuffer.h
+++ b/Engine/Plugins/Media/PixelStreaming/Source/PixelStreaming/Private/PixelStreamingFrameBuffer.h
@@ -3,6 +3,7 @@
 #pragma once
 
 #include "WebRTCIncludes.h"
+#include "RHI.h"
 
 class FPixelStreamingFrameSource;
 class FPixelStreamingLayerFrameSource;
