Merging //UE5/Dev-TensorWorks to Release-5.0 (//UE5/Release-5.0)
Fixing non unity builds.

diff --git a/Engine/Plugins/Media/PixelStreaming/Source/PixelStreaming/Private/PixelStreamingEncoderFactory.cpp b/Engine/Plugins/Media/PixelStreaming/Source/PixelStreaming/Private/PixelStreamingEncoderFactory.cpp
--- a/Engine/Plugins/Media/PixelStreaming/Source/PixelStreaming/Private/PixelStreamingEncoderFactory.cpp
+++ b/Engine/Plugins/Media/PixelStreaming/Source/PixelStreaming/Private/PixelStreamingEncoderFactory.cpp
@@ -9,23 +9,6 @@
 #include "Misc/ScopeLock.h"
 #include "PixelStreamingRealEncoder.h"
 
-#if PLATFORM_WINDOWS
-// some code in the simulcast adaptor has trouble finding the windows versions of these
-// functions. these are here to redirect them properly
-template<class T>
-T InterlockedIncrement(volatile T* i) {
-	return ::_InterlockedIncrement(i);
-}
-template<class T>
-T InterlockedDecrement(volatile T* i) {
-	return ::_InterlockedDecrement(i);
-}
-template<class T, class U>
-T InterlockedCompareExchange(volatile T* i, U old_value, U new_value) {
-	return ::_InterlockedCompareExchange(i, new_value, old_value);
-}
-#endif
-
 #include "PixelStreamingSimulcastEncoderAdapter.h"
 
 FPixelStreamingSimulcastEncoderFactory::FPixelStreamingSimulcastEncoderFactory(IPixelStreamingSessions* InPixelStreamingSessions)
diff --git a/Engine/Plugins/Media/PixelStreaming/Source/PixelStreaming/Private/PixelStreamingFrameBuffer.h b/Engine/Plugins/Media/PixelStreaming/Source/PixelStreaming/Private/PixelStreamingFrameBuffer.h
--- a/Engine/Plugins/Media/PixelStreaming/Source/PixelStreaming/Private/PixelStreamingFrameBuffer.h
+++ b/Engine/Plugins/Media/PixelStreaming/Source/PixelStreaming/Private/PixelStreamingFrameBuffer.h
@@ -2,6 +2,7 @@
 
 #pragma once
 
+#include "RHI.h"
 #include "WebRTCIncludes.h"
 #include "RHI.h"
 
diff --git a/Engine/Plugins/Media/PixelStreaming/Source/PixelStreaming/Private/PixelStreamingFramePump.h b/Engine/Plugins/Media/PixelStreaming/Source/PixelStreaming/Private/PixelStreamingFramePump.h
--- a/Engine/Plugins/Media/PixelStreaming/Source/PixelStreaming/Private/PixelStreamingFramePump.h
+++ b/Engine/Plugins/Media/PixelStreaming/Source/PixelStreaming/Private/PixelStreamingFramePump.h
@@ -2,6 +2,7 @@
 #pragma once
 
 #include "PlayerId.h"
+#include "HAL/Thread.h"
 
 class FPixelStreamingFrameSource;
 class FPlayerVideoSource;
diff --git a/Engine/Plugins/Media/PixelStreaming/Source/PixelStreaming/Private/PixelStreamingFrameSource.cpp b/Engine/Plugins/Media/PixelStreaming/Source/PixelStreaming/Private/PixelStreamingFrameSource.cpp
--- a/Engine/Plugins/Media/PixelStreaming/Source/PixelStreaming/Private/PixelStreamingFrameSource.cpp
+++ b/Engine/Plugins/Media/PixelStreaming/Source/PixelStreaming/Private/PixelStreamingFrameSource.cpp
@@ -3,6 +3,7 @@
 
 #include "LatencyTester.h"
 #include "PixelStreamingStats.h"
+#include "Async/Async.h"
 
 FPixelStreamingFrameSource::FPixelStreamingFrameSource()
 {
diff --git a/Engine/Plugins/Media/PixelStreaming/Source/PixelStreaming/Private/PixelStreamingFrameSource.h b/Engine/Plugins/Media/PixelStreaming/Source/PixelStreaming/Private/PixelStreamingFrameSource.h
--- a/Engine/Plugins/Media/PixelStreaming/Source/PixelStreaming/Private/PixelStreamingFrameSource.h
+++ b/Engine/Plugins/Media/PixelStreaming/Source/PixelStreaming/Private/PixelStreamingFrameSource.h
@@ -3,6 +3,7 @@
 
 #include "RHI.h"
 #include "Containers/TripleBuffer.h"
+#include "HAL/ThreadSafeBool.h"
 
 class FPixelStreamingLayerFrameSource;
 
diff --git a/Engine/Plugins/Media/PixelStreaming/Source/PixelStreaming/Private/PixelStreamingSettings.cpp b/Engine/Plugins/Media/PixelStreaming/Source/PixelStreaming/Private/PixelStreamingSettings.cpp
--- a/Engine/Plugins/Media/PixelStreaming/Source/PixelStreaming/Private/PixelStreamingSettings.cpp
+++ b/Engine/Plugins/Media/PixelStreaming/Source/PixelStreaming/Private/PixelStreamingSettings.cpp
@@ -3,6 +3,7 @@
 #include "PixelStreamingSettings.h"
 #include "PixelStreamingPrivate.h"
 #include "Misc/DefaultValueHelper.h"
+#include "Async/Async.h"
 
 template<typename T>
 void CommandLineParseValue(const TCHAR* Match, TAutoConsoleVariable<T>& CVar)
diff --git a/Engine/Plugins/Media/PixelStreaming/Source/PixelStreaming/Private/PixelStreamingSimulcastEncoderAdapter.cpp b/Engine/Plugins/Media/PixelStreaming/Source/PixelStreaming/Private/PixelStreamingSimulcastEncoderAdapter.cpp
--- a/Engine/Plugins/Media/PixelStreaming/Source/PixelStreaming/Private/PixelStreamingSimulcastEncoderAdapter.cpp
+++ b/Engine/Plugins/Media/PixelStreaming/Source/PixelStreaming/Private/PixelStreamingSimulcastEncoderAdapter.cpp
@@ -1,30 +1,30 @@
 // Copyright Epic Games, Inc. All Rights Reserved.
 #include "PixelStreamingSimulcastEncoderAdapter.h"
 
-#include <stdio.h>
-#include <string.h>
-
-#include <algorithm>
-#include <cstdint>
-#include <string>
-#include <utility>
-
-#include "api/scoped_refptr.h"
-#include "api/video/i420_buffer.h"
-#include "api/video/video_codec_constants.h"
-#include "api/video/video_frame_buffer.h"
-#include "api/video/video_rotation.h"
-#include "api/video_codecs/video_encoder.h"
-#include "api/video_codecs/video_encoder_factory.h"
-#include "api/video_codecs/video_encoder_software_fallback_wrapper.h"
-#include "media/base/video_common.h"
-#include "modules/video_coding/include/video_error_codes.h"
-#include "modules/video_coding/utility/simulcast_rate_allocator.h"
-#include "rtc_base/atomic_ops.h"
-#include "rtc_base/checks.h"
-#include "rtc_base/experiments/rate_control_settings.h"
-#include "rtc_base/logging.h"
-#include "system_wrappers/include/field_trial.h"
+//#include <stdio.h>
+//#include <string.h>
+//
+//#include <algorithm>
+//#include <cstdint>
+//#include <string>
+//#include <utility>
+
+//#include "api/scoped_refptr.h"
+//#include "api/video/i420_buffer.h"
+//#include "api/video/video_codec_constants.h"
+//#include "api/video/video_frame_buffer.h"
+//#include "api/video/video_rotation.h"
+//#include "api/video_codecs/video_encoder.h"
+//#include "api/video_codecs/video_encoder_factory.h"
+//#include "api/video_codecs/video_encoder_software_fallback_wrapper.h"
+//#include "media/base/video_common.h"
+//#include "modules/video_coding/include/video_error_codes.h"
+//#include "modules/video_coding/utility/simulcast_rate_allocator.h"
+//#include "rtc_base/atomic_ops.h"
+//#include "rtc_base/checks.h"
+//#include "rtc_base/experiments/rate_control_settings.h"
+//#include "rtc_base/logging.h"
+//#include "system_wrappers/include/field_trial.h"
 
 #include "PixelStreamingFrameBuffer.h"
 
diff --git a/Engine/Plugins/Media/PixelStreaming/Source/PixelStreaming/Private/PixelStreamingSimulcastEncoderAdapter.h b/Engine/Plugins/Media/PixelStreaming/Source/PixelStreaming/Private/PixelStreamingSimulcastEncoderAdapter.h
--- a/Engine/Plugins/Media/PixelStreaming/Source/PixelStreaming/Private/PixelStreamingSimulcastEncoderAdapter.h
+++ b/Engine/Plugins/Media/PixelStreaming/Source/PixelStreaming/Private/PixelStreamingSimulcastEncoderAdapter.h
@@ -1,21 +1,24 @@
 // Copyright Epic Games, Inc. All Rights Reserved.
 #pragma once
 
-#include <memory>
-#include <stack>
-#include <string>
-#include <utility>
-#include <vector>
-
-#include "absl/types/optional.h"
-#include "api/fec_controller_override.h"
-#include "api/video_codecs/sdp_video_format.h"
-#include "api/video_codecs/video_encoder.h"
-#include "modules/video_coding/include/video_codec_interface.h"
-#include "modules/video_coding/utility/framerate_controller.h"
-#include "rtc_base/atomic_ops.h"
-#include "rtc_base/synchronization/sequence_checker.h"
-#include "rtc_base/system/rtc_export.h"
+#include "CoreMinimal.h"
+
+#include "WebRTCIncludes.h"
+//#include <memory>
+
+//#include <string>
+//#include <utility>
+//#include <vector>
+
+//#include "absl/types/optional.h"
+//#include "api/fec_controller_override.h"
+//#include "api/video_codecs/sdp_video_format.h"
+//#include "api/video_codecs/video_encoder.h"
+//#include "modules/video_coding/include/video_codec_interface.h"
+//#include "modules/video_coding/utility/framerate_controller.h"
+//#include "rtc_base/atomic_ops.h"
+//#include "rtc_base/synchronization/sequence_checker.h"
+//#include "rtc_base/system/rtc_export.h"
 
 namespace webrtc {
 
diff --git a/Engine/Plugins/Media/PixelStreaming/Source/PixelStreaming/Private/PixelStreamingVideoEncoder.cpp b/Engine/Plugins/Media/PixelStreaming/Source/PixelStreaming/Private/PixelStreamingVideoEncoder.cpp
--- a/Engine/Plugins/Media/PixelStreaming/Source/PixelStreaming/Private/PixelStreamingVideoEncoder.cpp
+++ b/Engine/Plugins/Media/PixelStreaming/Source/PixelStreaming/Private/PixelStreamingVideoEncoder.cpp
@@ -14,6 +14,7 @@
 #include "Async/Async.h"
 #include "Utils.h"
 #include "PixelStreamingRealEncoder.h"
+#include "HAL/PlatformFilemanager.h"
 
 FPixelStreamingVideoEncoder::FPixelStreamingVideoEncoder(IPixelStreamingSessions* InPixelStreamingSessions, FPixelStreamingVideoEncoderFactory& InFactory)
 	: PixelStreamingSessions(InPixelStreamingSessions)
diff --git a/Engine/Plugins/Media/PixelStreaming/Source/PixelStreaming/Private/PlayerVideoSource.h b/Engine/Plugins/Media/PixelStreaming/Source/PixelStreaming/Private/PlayerVideoSource.h
--- a/Engine/Plugins/Media/PixelStreaming/Source/PixelStreaming/Private/PlayerVideoSource.h
+++ b/Engine/Plugins/Media/PixelStreaming/Source/PixelStreaming/Private/PlayerVideoSource.h
@@ -2,6 +2,8 @@
 
 #pragma once
 
+#include "WebRTCIncludes.h"
+
 // each connected player has a single source. it is fed by FVideoSource and feeds the encoder
 class FPlayerVideoSource : public rtc::AdaptedVideoTrackSource
 {
diff --git a/Engine/Plugins/Media/PixelStreaming/Source/PixelStreaming/Private/WebRTCIncludes.h b/Engine/Plugins/Media/PixelStreaming/Source/PixelStreaming/Private/WebRTCIncludes.h
--- a/Engine/Plugins/Media/PixelStreaming/Source/PixelStreaming/Private/WebRTCIncludes.h
+++ b/Engine/Plugins/Media/PixelStreaming/Source/PixelStreaming/Private/WebRTCIncludes.h
@@ -34,12 +34,31 @@ THIRD_PARTY_INCLUDES_START
 #include "api/audio_codecs/builtin_audio_encoder_factory.h"
 #include "api/video_codecs/video_decoder.h"
 #include "api/video_codecs/sdp_video_format.h"
+#include "api/video_codecs/video_encoder_software_fallback_wrapper.h"
 #include "api/video/video_frame.h"
 #include "api/video/video_rotation.h"
 #include "api/video/video_frame_buffer.h"
 #include "api/video/i420_buffer.h"
 #include "api/video/video_sink_interface.h"
 
+#if PLATFORM_WINDOWS
+// some code in the simulcast adaptor has trouble finding the windows versions of these
+// functions. these are here to redirect them properly
+template<class T>
+T InterlockedIncrement(volatile T* i) {
+	return ::_InterlockedIncrement(i);
+}
+template<class T>
+T InterlockedDecrement(volatile T* i) {
+	return ::_InterlockedDecrement(i);
+}
+template<class T, class U>
+T InterlockedCompareExchange(volatile T* i, U old_value, U new_value) {
+	return ::_InterlockedCompareExchange(i, new_value, old_value);
+}
+#endif
+
+#include "rtc_base/atomic_ops.h"
 #include "rtc_base/thread.h"
 #include "rtc_base/logging.h"
 #include "rtc_base/ssl_adapter.h"
@@ -47,6 +66,7 @@ THIRD_PARTY_INCLUDES_START
 #include "rtc_base/net_helpers.h"
 #include "rtc_base/string_utils.h"
 #include "rtc_base/signal_thread.h"
+#include "rtc_base/experiments/rate_control_settings.h"
 
 #include "pc/session_description.h"
 #include "pc/video_track_source.h"
@@ -63,15 +83,20 @@ THIRD_PARTY_INCLUDES_START
 #include "modules/audio_device/audio_device_buffer.h"
 #include "modules/audio_processing/include/audio_processing.h"
 #include "modules/video_coding/codecs/h264/include/h264.h"
+#include "modules/video_coding/utility/framerate_controller.h"
+#include "modules/video_coding/utility/simulcast_rate_allocator.h"
 
 #include "common_video/h264/h264_bitstream_parser.h"
 #include "common_video/h264/h264_common.h"
 
 #include "media/base/video_broadcaster.h"
 
+#include "system_wrappers/include/field_trial.h"
+
 // because WebRTC uses STL
 #include <string>
 #include <memory>
+#include <stack>
 
 #if PLATFORM_WINDOWS
 #pragma warning(pop)
diff --git a/Engine/Source/Runtime/AVEncoder/Public/CodecPacket.h b/Engine/Source/Runtime/AVEncoder/Public/CodecPacket.h
--- a/Engine/Source/Runtime/AVEncoder/Public/CodecPacket.h
+++ b/Engine/Source/Runtime/AVEncoder/Public/CodecPacket.h
@@ -2,6 +2,8 @@
 
 #pragma once
 
+#include "CoreMinimal.h"
+
 namespace AVEncoder
 {
 	class AVENCODER_API FCodecPacket
