fix Linux build

diff --git a/Engine/Plugins/Media/PixelStreaming/Source/PixelStreaming/Private/PixelStreamingVideoEncoder.cpp b/Engine/Plugins/Media/PixelStreaming/Source/PixelStreaming/Private/PixelStreamingVideoEncoder.cpp
--- a/Engine/Plugins/Media/PixelStreaming/Source/PixelStreaming/Private/PixelStreamingVideoEncoder.cpp
+++ b/Engine/Plugins/Media/PixelStreaming/Source/PixelStreaming/Private/PixelStreamingVideoEncoder.cpp
@@ -14,7 +14,7 @@
 #include "Async/Async.h"
 #include "Utils.h"
 #include "PixelStreamingRealEncoder.h"
-#include "HAL/PlatformFilemanager.h"
+#include "HAL/PlatformFileManager.h"
 
 FPixelStreamingVideoEncoder::FPixelStreamingVideoEncoder(IPixelStreamingSessions* InPixelStreamingSessions, FPixelStreamingVideoEncoderFactory& InFactory)
 	: PixelStreamingSessions(InPixelStreamingSessions)
