Fix for keyframes being lower quality than the rest of the stream.

diff --git a/Engine/Plugins/Media/PixelStreaming/Source/PixelStreaming/Private/PixelStreamingSettings.cpp b/Engine/Plugins/Media/PixelStreaming/Source/PixelStreaming/Private/PixelStreamingSettings.cpp
--- a/Engine/Plugins/Media/PixelStreaming/Source/PixelStreaming/Private/PixelStreamingSettings.cpp
+++ b/Engine/Plugins/Media/PixelStreaming/Source/PixelStreaming/Private/PixelStreamingSettings.cpp
@@ -5,7 +5,7 @@
 #include "Misc/DefaultValueHelper.h"
 #include "Async/Async.h"
 
-template<typename T>
+template <typename T>
 void CommandLineParseValue(const TCHAR* Match, TAutoConsoleVariable<T>& CVar)
 {
 	T Value;
@@ -25,11 +25,14 @@ void CommandLineParseOption(const TCHAR* Match, TAutoConsoleVariable<bool>& CVar
 	FString ValueMatch(Match);
 	ValueMatch.Append(TEXT("="));
 	FString Value;
-	if (FParse::Value(FCommandLine::Get(), *ValueMatch, Value)) {
-		if (Value.Equals(FString(TEXT("true")), ESearchCase::IgnoreCase)) {
+	if (FParse::Value(FCommandLine::Get(), *ValueMatch, Value))
+	{
+		if (Value.Equals(FString(TEXT("true")), ESearchCase::IgnoreCase))
+		{
 			CVar->Set(true, ECVF_SetByCommandline);
 		}
-		else if (Value.Equals(FString(TEXT("false")), ESearchCase::IgnoreCase)) {
+		else if (Value.Equals(FString(TEXT("false")), ESearchCase::IgnoreCase))
+		{
 			CVar->Set(false, ECVF_SetByCommandline);
 		}
 	}
@@ -41,7 +44,7 @@ void CommandLineParseOption(const TCHAR* Match, TAutoConsoleVariable<bool>& CVar
 
 namespace PixelStreamingSettings
 {
-// Begin Encoder CVars
+	// Begin Encoder CVars
 	TAutoConsoleVariable<int32> CVarPixelStreamingEncoderTargetBitrate(
 		TEXT("PixelStreaming.Encoder.TargetBitrate"),
 		-1,
@@ -62,14 +65,14 @@ namespace PixelStreamingSettings
 
 	TAutoConsoleVariable<int32> CVarPixelStreamingEncoderMinQP(
 		TEXT("PixelStreaming.Encoder.MinQP"),
-		-1,
-		TEXT("0-51, lower values result in better quality but higher bitrate. Default -1. -1 will disable any hard limit on a minimum QP."),
+		0,
+		TEXT("0-51, lower values result in better quality but higher bitrate. Default 0 - i.e. no limit on a minimum QP."),
 		ECVF_Default);
 
 	TAutoConsoleVariable<int32> CVarPixelStreamingEncoderMaxQP(
 		TEXT("PixelStreaming.Encoder.MaxQP"),
-		-1,
-		TEXT("0-51, lower values result in better quality but higher bitrate. Default -1. -1 will disable any hard limit on a maximum QP."),
+		51,
+		TEXT("0-51, lower values result in better quality but higher bitrate. Default 51 - i.e. no limit on a maximum QP."),
 		ECVF_Default);
 
 	TAutoConsoleVariable<FString> CVarPixelStreamingEncoderRateControl(
@@ -101,9 +104,9 @@ namespace PixelStreamingSettings
 		300,
 		TEXT("How many frames before a key frame is sent. Default: 300. Values <=0 will disable sending of periodic key frames. Note: NVENC does not support changing this after encoding has started."),
 		ECVF_Default);
-// End Encoder CVars
+	// End Encoder CVars
 
-// Begin Capturer CVars
+	// Begin Capturer CVars
 	TAutoConsoleVariable<bool> CVarPixelStreamingUseBackBufferCaptureSize(
 		TEXT("PixelStreaming.Capturer.UseBackBufferSize"),
 		true,
@@ -121,9 +124,9 @@ namespace PixelStreamingSettings
 		3,
 		TEXT("Maximum number of back buffers to use. A value of 0 will not limit the number of back buffers."),
 		ECVF_Default);
-// End Capturer CVars
+	// End Capturer CVars
 
-// Begin WebRTC CVars
+	// Begin WebRTC CVars
 	TAutoConsoleVariable<FString> CVarPixelStreamingDegradationPreference(
 		TEXT("PixelStreaming.WebRTC.DegradationPreference"),
 		TEXT("MAINTAIN_FRAMERATE"),
@@ -196,9 +199,9 @@ namespace PixelStreamingSettings
 		TEXT("Disables whether WebRTC can trigger a resolution a change of the stream (regardless of it wants to)."),
 		ECVF_Default);
 
-// End WebRTC CVars
+	// End WebRTC CVars
 
-// Begin Pixel Streaming Plugin CVars
+	// Begin Pixel Streaming Plugin CVars
 
 	TAutoConsoleVariable<bool> CVarPixelStreamingOnScreenStats(
 		TEXT("PixelStreaming.HUDStats"),
@@ -230,8 +233,8 @@ namespace PixelStreamingSettings
 		TEXT("If true disables latency tester being triggerable."),
 		ECVF_Default);
 
-	 TAutoConsoleVariable<FString> CVarPixelStreamingKeyFilter(
-	 	TEXT("PixelStreaming.KeyFilter"),
+	TAutoConsoleVariable<FString> CVarPixelStreamingKeyFilter(
+		TEXT("PixelStreaming.KeyFilter"),
 		"",
 		TEXT("Comma separated list of keys to ignore from streaming clients."),
 		ECVF_Default);
@@ -255,34 +258,33 @@ namespace PixelStreamingSettings
 		// HACK: Only doing it like this because this a hotfix change and cannot public ABI.
 		// Todo: Fix in next proper engine release.
 		AsyncTask(ENamedThreads::GameThread, [Var]()
-		{ 
-			IConsoleVariable* CVarNVENCKeyframeInterval = IConsoleManager::Get().FindConsoleVariable(TEXT("NVENC.KeyframeInterval"));
-			if(CVarNVENCKeyframeInterval)
-			{
-				CVarNVENCKeyframeInterval->Set(Var->GetInt(), ECVF_SetByCommandline);
-			}
-
-			IConsoleVariable* CVarAMFKeyframeInterval = IConsoleManager::Get().FindConsoleVariable(TEXT("AMF.KeyframeInterval"));
-			if(CVarNVENCKeyframeInterval)
-			{
-				CVarAMFKeyframeInterval->Set(Var->GetInt(), ECVF_SetByCommandline);
-			}
-		});
-		
+				  {
+					  IConsoleVariable* CVarNVENCKeyframeInterval = IConsoleManager::Get().FindConsoleVariable(TEXT("NVENC.KeyframeInterval"));
+					  if (CVarNVENCKeyframeInterval)
+					  {
+						  CVarNVENCKeyframeInterval->Set(Var->GetInt(), ECVF_SetByCommandline);
+					  }
+
+					  IConsoleVariable* CVarAMFKeyframeInterval = IConsoleManager::Get().FindConsoleVariable(TEXT("AMF.KeyframeInterval"));
+					  if (CVarNVENCKeyframeInterval)
+					  {
+						  CVarAMFKeyframeInterval->Set(Var->GetInt(), ECVF_SetByCommandline);
+					  }
+				  });
 	}
-// Ends Pixel Streaming Plugin CVars
+	// Ends Pixel Streaming Plugin CVars
 
-// Begin utility functions etc.
+	// Begin utility functions etc.
 	std::map<FString, AVEncoder::FVideoEncoder::RateControlMode> const RateControlCVarMap{
-		{"ConstQP", AVEncoder::FVideoEncoder::RateControlMode::CONSTQP },
-		{"VBR", AVEncoder::FVideoEncoder::RateControlMode::VBR },
-		{"CBR", AVEncoder::FVideoEncoder::RateControlMode::CBR },
+		{"ConstQP", AVEncoder::FVideoEncoder::RateControlMode::CONSTQP},
+		{"VBR", AVEncoder::FVideoEncoder::RateControlMode::VBR},
+		{"CBR", AVEncoder::FVideoEncoder::RateControlMode::CBR},
 	};
 
 	std::map<FString, AVEncoder::FVideoEncoder::MultipassMode> const MultipassCVarMap{
-		{"DISABLED", AVEncoder::FVideoEncoder::MultipassMode::DISABLED },
-		{"QUARTER", AVEncoder::FVideoEncoder::MultipassMode::QUARTER },
-		{"FULL", AVEncoder::FVideoEncoder::MultipassMode::FULL },
+		{"DISABLED", AVEncoder::FVideoEncoder::MultipassMode::DISABLED},
+		{"QUARTER", AVEncoder::FVideoEncoder::MultipassMode::QUARTER},
+		{"FULL", AVEncoder::FVideoEncoder::MultipassMode::FULL},
 	};
 
 	std::map<FString, AVEncoder::FVideoEncoder::H264Profile> const H264ProfileMap{
@@ -318,11 +320,11 @@ namespace PixelStreamingSettings
 	webrtc::DegradationPreference GetDegradationPreference()
 	{
 		FString DegradationPrefStr = CVarPixelStreamingDegradationPreference.GetValueOnAnyThread();
-		if(DegradationPrefStr == "MAINTAIN_FRAMERATE")
+		if (DegradationPrefStr == "MAINTAIN_FRAMERATE")
 		{
 			return webrtc::DegradationPreference::MAINTAIN_FRAMERATE;
 		}
-		else if(DegradationPrefStr == "MAINTAIN_RESOLUTION")
+		else if (DegradationPrefStr == "MAINTAIN_RESOLUTION")
 		{
 			return webrtc::DegradationPreference::MAINTAIN_RESOLUTION;
 		}
@@ -338,7 +340,7 @@ namespace PixelStreamingSettings
 			return AVEncoder::FVideoEncoder::H264Profile::BASELINE;
 		return it->second;
 	}
-// End utility functions etc.
+	// End utility functions etc.
 
 	FSimulcastParameters SimulcastParameters;
 
@@ -419,4 +421,4 @@ namespace PixelStreamingSettings
 
 		ReadSimulcastParameters();
 	}
-}
+} // namespace PixelStreamingSettings
