Windows Vulkan Pixel Streaming support

diff --git a/Engine/Plugins/Media/PixelStreaming/Source/PixelStreaming/Private/PixelStreamingEncoderFrameFactory.cpp b/Engine/Plugins/Media/PixelStreaming/Source/PixelStreaming/Private/PixelStreamingEncoderFrameFactory.cpp
--- a/Engine/Plugins/Media/PixelStreaming/Source/PixelStreaming/Private/PixelStreamingEncoderFrameFactory.cpp
+++ b/Engine/Plugins/Media/PixelStreaming/Source/PixelStreaming/Private/PixelStreamingEncoderFrameFactory.cpp
@@ -15,6 +15,11 @@
     #include "D3D12RHIPrivate.h"
     #include "D3D12Resources.h"
     #include "D3D12Texture.h"
+    #include "Windows/AllowWindowsPlatformTypes.h"
+    THIRD_PARTY_INCLUDES_START
+    #include <VersionHelpers.h>
+    THIRD_PARTY_INCLUDES_END
+    #include "Windows/HideWindowsPlatformTypes.h"
 #endif
 
 FPixelStreamingEncoderFrameFactory::FPixelStreamingEncoderFrameFactory()
@@ -206,6 +211,54 @@ void FPixelStreamingEncoderFrameFactory::SetTextureCUDAVulkan(AVEncoder::FVideoE
     FVulkanTexture2D* VulkanTexture = static_cast<FVulkanTexture2D*>(Texture.GetReference());
     VkDevice device = static_cast<FVulkanDynamicRHI*>(GDynamicRHI)->GetDevice()->GetInstanceHandle();
 
+#if PLATFORM_WINDOWS
+    HANDLE handle;
+    // It is recommended to use NT handles where available, but these are only supported from Windows 8 onward, for earliers versions of Windows
+    // we need to use a Win7 style handle. NT handles require us to close them when we are done with them to prevent memory leaks.
+    // Refer to remarks section of https://docs.microsoft.com/en-us/windows/win32/api/dxgi1_2/nf-dxgi1_2-idxgiresource1-createsharedhandle
+    bool bUseNTHandle = IsWindows8OrGreater();
+
+    { 
+        // Generate VkMemoryGetWin32HandleInfoKHR
+        VkMemoryGetWin32HandleInfoKHR vkMemoryGetHandleInfoKHR = {};
+        vkMemoryGetHandleInfoKHR.sType = VK_STRUCTURE_TYPE_MEMORY_GET_WIN32_HANDLE_INFO_KHR;  
+        vkMemoryGetHandleInfoKHR.pNext = NULL;
+        vkMemoryGetHandleInfoKHR.memory = VulkanTexture->Surface.GetAllocationHandle();
+        vkMemoryGetHandleInfoKHR.handleType = bUseNTHandle ? VK_EXTERNAL_MEMORY_HANDLE_TYPE_OPAQUE_WIN32_BIT : VK_EXTERNAL_MEMORY_HANDLE_TYPE_OPAQUE_WIN32_KMT_BIT;
+
+        // While this operation is safe (and unavoidable) C4191 has been enabled and this will trigger an error with warnings as errors
+#pragma warning(push)
+#pragma warning(disable : 4191)
+        PFN_vkGetMemoryWin32HandleKHR GetMemoryWin32HandleKHR = (PFN_vkGetMemoryWin32HandleKHR)VulkanRHI::vkGetDeviceProcAddr(device, "vkGetMemoryWin32HandleKHR");
+        VERIFYVULKANRESULT(GetMemoryWin32HandleKHR(device, &vkMemoryGetHandleInfoKHR, &handle));
+#pragma warning(pop)
+    }
+
+    FCUDAModule::CUDA().cuCtxPushCurrent(FModuleManager::GetModuleChecked<FCUDAModule>("CUDA").GetCudaContext());
+
+    CUexternalMemory mappedExternalMemory = nullptr;
+
+    {
+        // generate a cudaExternalMemoryHandleDesc
+        CUDA_EXTERNAL_MEMORY_HANDLE_DESC cudaExtMemHandleDesc = {};
+        cudaExtMemHandleDesc.type = bUseNTHandle ? CU_EXTERNAL_MEMORY_HANDLE_TYPE_OPAQUE_WIN32 : CU_EXTERNAL_MEMORY_HANDLE_TYPE_OPAQUE_WIN32_KMT;
+        cudaExtMemHandleDesc.handle.win32.name = NULL;
+        cudaExtMemHandleDesc.handle.win32.handle = handle;
+        cudaExtMemHandleDesc.size = VulkanTexture->Surface.GetAllocationOffset() + VulkanTexture->Surface.GetMemorySize();
+
+        // import external memory
+        auto result = FCUDAModule::CUDA().cuImportExternalMemory(&mappedExternalMemory, &cudaExtMemHandleDesc);
+        if(result != CUDA_SUCCESS)
+        {
+            UE_LOG(PixelStreamer, Error, TEXT("Failed to import external memory from vulkan error: %d"), result);
+        }
+    }
+
+    // Only store handle to be closed on frame destruction if it is an NT handle
+    handle = bUseNTHandle ? handle : NULL;
+#else
+    void* handle = nullptr;
+
     // Get the CUarray to that textures memory making sure the clear it when done
     int fd;
 
@@ -244,6 +297,9 @@ void FPixelStreamingEncoderFrameFactory::SetTextureCUDAVulkan(AVEncoder::FVideoE
         }
     }
 
+#endif
+
+
     CUmipmappedArray mappedMipArray = nullptr;
     CUarray mappedArray = nullptr;
 
@@ -275,7 +331,7 @@ void FPixelStreamingEncoderFrameFactory::SetTextureCUDAVulkan(AVEncoder::FVideoE
 
     FCUDAModule::CUDA().cuCtxPopCurrent(NULL);
 
-    InputFrame->SetTexture(mappedArray, AVEncoder::FVideoEncoderInputFrame::EUnderlyingRHI::Vulkan, nullptr, [mappedArray, mappedMipArray, mappedExternalMemory](CUarray NativeTexture)
+    InputFrame->SetTexture(mappedArray, AVEncoder::FVideoEncoderInputFrame::EUnderlyingRHI::Vulkan, handle, [mappedArray, mappedMipArray, mappedExternalMemory](CUarray NativeTexture)
     {
         // free the cuda types
         FCUDAModule::CUDA().cuCtxPushCurrent(FModuleManager::GetModuleChecked<FCUDAModule>("CUDA").GetCudaContext());
diff --git a/Engine/Plugins/Media/PixelStreaming/Source/PixelStreaming/Private/PixelStreamingFrameSource.cpp b/Engine/Plugins/Media/PixelStreaming/Source/PixelStreaming/Private/PixelStreamingFrameSource.cpp
--- a/Engine/Plugins/Media/PixelStreaming/Source/PixelStreaming/Private/PixelStreamingFrameSource.cpp
+++ b/Engine/Plugins/Media/PixelStreaming/Source/PixelStreaming/Private/PixelStreamingFrameSource.cpp
@@ -170,11 +170,14 @@ namespace
 	{
 		FRHIResourceCreateInfo CreateInfo(TEXT("VideoCapturerBackBuffer"));
 		FTexture2DRHIRef Texture;
-#if PLATFORM_WINDOWS
-		Texture = GDynamicRHI->RHICreateTexture2D(Width, Height, EPixelFormat::PF_B8G8R8A8, 1, 1, TexCreate_Shared | TexCreate_RenderTargetable, ERHIAccess::CopyDest, CreateInfo);
-#else
-		Texture = GDynamicRHI->RHICreateTexture2D(Width, Height, EPixelFormat::PF_B8G8R8A8, 1, 1, TexCreate_Shared, ERHIAccess::Present, CreateInfo);
-#endif
+		FString RHIName = GDynamicRHI->GetName();
+
+		if(RHIName == TEXT("Vulkan")) {
+			Texture = GDynamicRHI->RHICreateTexture2D(Width, Height, EPixelFormat::PF_B8G8R8A8, 1, 1, TexCreate_RenderTargetable | TexCreate_External, ERHIAccess::Present, CreateInfo);
+		} else
+		{
+			Texture = GDynamicRHI->RHICreateTexture2D(Width, Height, EPixelFormat::PF_B8G8R8A8, 1, 1, TexCreate_Shared | TexCreate_RenderTargetable, ERHIAccess::CopyDest, CreateInfo);
+		}
 	    return Texture;
 	}
 }
