Merging //UE5/Dev-TensorWorks to Release-5.0 (//UE5/Release-5.0)

- Updating simulcast support to allow direct connect peers

diff --git a/Engine/Plugins/Media/PixelStreaming/Source/PixelStreaming/Private/PixelStreamingEncoderFactory.cpp b/Engine/Plugins/Media/PixelStreaming/Source/PixelStreaming/Private/PixelStreamingEncoderFactory.cpp
--- a/Engine/Plugins/Media/PixelStreaming/Source/PixelStreaming/Private/PixelStreamingEncoderFactory.cpp
+++ b/Engine/Plugins/Media/PixelStreaming/Source/PixelStreaming/Private/PixelStreamingEncoderFactory.cpp
@@ -32,7 +32,7 @@ FPixelStreamingSimulcastEncoderFactory::CodecInfo FPixelStreamingSimulcastEncode
 
 std::unique_ptr<webrtc::VideoEncoder> FPixelStreamingSimulcastEncoderFactory::CreateVideoEncoder(const webrtc::SdpVideoFormat& format)
 {
-	return std::make_unique<PixelStreamingSimulcastEncoderAdapter>(RealFactory.Get(), nullptr, format);
+	return std::make_unique<PixelStreamingSimulcastEncoderAdapter>(RealFactory.Get(), format);
 }
 
 FPixelStreamingVideoEncoderFactory* FPixelStreamingSimulcastEncoderFactory::GetRealFactory() const
@@ -40,8 +40,6 @@ FPixelStreamingVideoEncoderFactory* FPixelStreamingSimulcastEncoderFactory::GetR
 	return RealFactory.Get();
 }
 
-
-
 FPixelStreamingVideoEncoderFactory::FPixelStreamingVideoEncoderFactory(IPixelStreamingSessions* InPixelStreamingSessions)
 	: PixelStreamingSessions(InPixelStreamingSessions)
 {
@@ -90,37 +88,15 @@ std::unique_ptr<webrtc::VideoEncoder> FPixelStreamingVideoEncoderFactory::Create
 	}
 }
 
-void FPixelStreamingVideoEncoderFactory::RemoveStaleEncoders()
-{
-	// Lock during removing stale encoders
-	FScopeLock Lock(&ActiveEncodersGuard);
-
-	// Iterate backwards so we can remove invalid encoders along the way
-	for (int32 Index = ActiveEncoders.Num()-1; Index >= 0; --Index)
-	{
-		FPixelStreamingVideoEncoder* Encoder = ActiveEncoders[Index];
-		if(!Encoder->IsRegisteredWithWebRTC())
-		{
-			ActiveEncoders.RemoveAt(Index);
-		}
-	}
-}
-
 void FPixelStreamingVideoEncoderFactory::OnEncodedImage(FHardwareEncoderId SourceEncoderId, const webrtc::EncodedImage& encoded_image, const webrtc::CodecSpecificInfo* codec_specific_info, const webrtc::RTPFragmentationHeader* fragmentation)
 {
-	// Before sending encoded image to each encoder's callback, check if all encoders we have are still relevant.
-	RemoveStaleEncoders();
-
 	// Lock as we send encoded image to each encoder.
 	FScopeLock Lock(&ActiveEncodersGuard);
 
 	// Go through each encoder and send our encoded image to its callback
 	for (FPixelStreamingVideoEncoder* Encoder : ActiveEncoders)
 	{
-		if (Encoder->IsRegisteredWithWebRTC())
-		{
-			Encoder->SendEncodedImage(SourceEncoderId, encoded_image, codec_specific_info, fragmentation);
-		}
+		Encoder->SendEncodedImage(SourceEncoderId, encoded_image, codec_specific_info, fragmentation);
 	}
 }
 
@@ -135,21 +111,12 @@ void FPixelStreamingVideoEncoderFactory::ForceKeyFrame()
 {
 	FScopeLock Lock(&ActiveEncodersGuard);
 	// Go through each encoder and send our encoded image to its callback
-	for(FPixelStreamingVideoEncoder* Encoder : ActiveEncoders)
+	for(auto& KeyAndEncoder : HardwareEncoders)
 	{
-		if(Encoder->IsRegisteredWithWebRTC())
-		{
-			Encoder->ForceKeyFrame();
-		}
+		KeyAndEncoder.Value->SetForceNextKeyframe();
 	}
 }
 
-double FPixelStreamingVideoEncoderFactory::GetLatestQP()
-{
-	return 0;
-	//return this->EncoderContext.SmoothedAvgQP.Get();
-}
-
 namespace
 {
 	uint64 HashResolution(uint32 Width, uint32 Height)
diff --git a/Engine/Plugins/Media/PixelStreaming/Source/PixelStreaming/Private/PixelStreamingEncoderFactory.h b/Engine/Plugins/Media/PixelStreaming/Source/PixelStreaming/Private/PixelStreamingEncoderFactory.h
--- a/Engine/Plugins/Media/PixelStreaming/Source/PixelStreaming/Private/PixelStreamingEncoderFactory.h
+++ b/Engine/Plugins/Media/PixelStreaming/Source/PixelStreaming/Private/PixelStreamingEncoderFactory.h
@@ -45,9 +45,7 @@ public:
 	virtual std::unique_ptr<webrtc::VideoEncoder> CreateVideoEncoder(const webrtc::SdpVideoFormat& format) override;
 
 	void ReleaseVideoEncoder(FPixelStreamingVideoEncoder* encoder);
-	void RemoveStaleEncoders();
 	void ForceKeyFrame();
-	double GetLatestQP();
 
 	using FHardwareEncoderId = uint64;
 	FHardwareEncoderId GetOrCreateHardwareEncoder(int Width, int Height, int MaxBitrate, int TargetBitrate, int MaxFramerate);
@@ -55,14 +53,17 @@ public:
 
 	void OnEncodedImage(FHardwareEncoderId SourceEncoderId, const webrtc::EncodedImage& encoded_image, const webrtc::CodecSpecificInfo* codec_specific_info, const webrtc::RTPFragmentationHeader* fragmentation);
 
+	int GetLatestQP() const { return 0; }
+
 private:
-	// Each encoder is associated with a particular player (peer).
-	TArray<FPixelStreamingVideoEncoder*> ActiveEncoders;
-	FCriticalSection ActiveEncodersGuard;
-	
 	// Used for checks such as whether a given player id is associated with the quality controlling player.
 	IPixelStreamingSessions* PixelStreamingSessions;
 
+	// These are the actual hardware encoders serving multiple pixelstreaming encoders
 	TMap<FHardwareEncoderId, TUniquePtr<FPixelStreamingRealEncoder>> HardwareEncoders;
 	FCriticalSection HardwareEncodersGuard;
+
+	// Encoders assigned to each peer. Each one of these will be assigned to one of the hardware encoders
+	TArray<FPixelStreamingVideoEncoder*> ActiveEncoders;
+	FCriticalSection ActiveEncodersGuard;
 };
diff --git a/Engine/Plugins/Media/PixelStreaming/Source/PixelStreaming/Private/PixelStreamingFrameBuffer.cpp b/Engine/Plugins/Media/PixelStreaming/Source/PixelStreaming/Private/PixelStreamingFrameBuffer.cpp
--- a/Engine/Plugins/Media/PixelStreaming/Source/PixelStreaming/Private/PixelStreamingFrameBuffer.cpp
+++ b/Engine/Plugins/Media/PixelStreaming/Source/PixelStreaming/Private/PixelStreamingFrameBuffer.cpp
@@ -2,6 +2,27 @@
 #include "PixelStreamingFrameBuffer.h"
 #include "PixelStreamingFrameSource.h"
 
+FPixelStreamingInitializeFrameBuffer::FPixelStreamingInitializeFrameBuffer(FPixelStreamingFrameSource* InFrameSource)
+	:FrameSource(InFrameSource)
+{
+
+}
+
+FPixelStreamingInitializeFrameBuffer::~FPixelStreamingInitializeFrameBuffer()
+{
+
+}
+
+int FPixelStreamingInitializeFrameBuffer::width() const
+{
+	return FrameSource->GetSourceWidth();
+}
+
+int FPixelStreamingInitializeFrameBuffer::height() const
+{
+	return FrameSource->GetSourceHeight();
+}
+
 FPixelStreamingSimulcastFrameBuffer::FPixelStreamingSimulcastFrameBuffer(FPixelStreamingFrameSource* InFrameSource)
 :FrameSource(InFrameSource)
 {
diff --git a/Engine/Plugins/Media/PixelStreaming/Source/PixelStreaming/Private/PixelStreamingFrameBuffer.h b/Engine/Plugins/Media/PixelStreaming/Source/PixelStreaming/Private/PixelStreamingFrameBuffer.h
--- a/Engine/Plugins/Media/PixelStreaming/Source/PixelStreaming/Private/PixelStreamingFrameBuffer.h
+++ b/Engine/Plugins/Media/PixelStreaming/Source/PixelStreaming/Private/PixelStreamingFrameBuffer.h
@@ -9,72 +9,88 @@
 class FPixelStreamingFrameSource;
 class FPixelStreamingLayerFrameSource;
 
-class FPixelStreamingSimulcastFrameBuffer : public webrtc::VideoFrameBuffer
+enum FPixelStreamingFrameBufferType
 {
-public:
-	FPixelStreamingSimulcastFrameBuffer(FPixelStreamingFrameSource* InFrameSource);
+	Initialize,
+	Simulcast,
+	Layer
+};
 
-	virtual ~FPixelStreamingSimulcastFrameBuffer();
+class FPixelStreamingFrameBuffer : public webrtc::VideoFrameBuffer
+{
+public:
+	virtual ~FPixelStreamingFrameBuffer() {}
+
+	virtual FPixelStreamingFrameBufferType GetFrameBufferType() const = 0;
+
+	// Begin webrtc::VideoFrameBuffer interface
+	virtual webrtc::VideoFrameBuffer::Type type() const override
+	{
+		return webrtc::VideoFrameBuffer::Type::kNative;
+	}
+
+	virtual int width() const = 0;
+	virtual int height() const = 0;
+
+	virtual rtc::scoped_refptr<webrtc::I420BufferInterface> ToI420() override
+	{
+		unimplemented();
+		return nullptr;
+	}
+
+	virtual const webrtc::I420BufferInterface* GetI420() const override
+	{
+		unimplemented();
+		return nullptr;
+	}
+	// End webrtc::VideoFrameBuffer interface
+};
 
-	FPixelStreamingLayerFrameSource* GetLayerFrameSource(int LayerIndex) const;
+// We use this frame to force webrtc to create encoders that will siphon off other streams
+// but never get frames directly pumped to them.
+class FPixelStreamingInitializeFrameBuffer : public FPixelStreamingFrameBuffer
+{
+public:
+	FPixelStreamingInitializeFrameBuffer(FPixelStreamingFrameSource* InFrameSource);
+	virtual ~FPixelStreamingInitializeFrameBuffer();
 
-    // Begin webrtc::VideoFrameBuffer interface
-    virtual webrtc::VideoFrameBuffer::Type type() const override
-    {
-        return webrtc::VideoFrameBuffer::Type::kNative;
-    }
+	virtual FPixelStreamingFrameBufferType GetFrameBufferType() const { return Initialize; }
 
 	virtual int width() const override;
-    virtual int height() const override;
+	virtual int height() const override;
 
-    virtual rtc::scoped_refptr<webrtc::I420BufferInterface> ToI420() override
-    {
-    	unimplemented();
-    	return nullptr;
-    }
+private:
+	FPixelStreamingFrameSource* FrameSource = nullptr;
+};
+
+class FPixelStreamingSimulcastFrameBuffer : public FPixelStreamingFrameBuffer
+{
+public:
+	FPixelStreamingSimulcastFrameBuffer(FPixelStreamingFrameSource* InFrameSource);
+	virtual ~FPixelStreamingSimulcastFrameBuffer();
 
-    virtual const webrtc::I420BufferInterface* GetI420() const override
-    {
-    	unimplemented();
-        return nullptr;
-    }
-    // End webrtc::VideoFrameBuffer interface
+	virtual FPixelStreamingFrameBufferType GetFrameBufferType() const { return Simulcast; }
+	FPixelStreamingLayerFrameSource* GetLayerFrameSource(int LayerIndex) const;
+
+    virtual int width() const override;
+    virtual int height() const override;
 
 private:
 	FPixelStreamingFrameSource* FrameSource = nullptr;
 };
 
-class FPixelStreamingLayerFrameBuffer : public webrtc::VideoFrameBuffer
+class FPixelStreamingLayerFrameBuffer : public FPixelStreamingFrameBuffer
 {
 public:
 	FPixelStreamingLayerFrameBuffer(FPixelStreamingLayerFrameSource* InLayerFrameSource);
-
 	virtual ~FPixelStreamingLayerFrameBuffer();
 
+	virtual FPixelStreamingFrameBufferType GetFrameBufferType() const { return Layer; }
 	FTexture2DRHIRef GetFrame() const;
 
-    // Begin webrtc::VideoFrameBuffer interface
-    virtual webrtc::VideoFrameBuffer::Type type() const override
-    {
-        return webrtc::VideoFrameBuffer::Type::kNative;
-    }
-
 	virtual int width() const override;
     virtual int height() const override;
 
-    virtual rtc::scoped_refptr<webrtc::I420BufferInterface> ToI420() override
-    {
-    	unimplemented();
-    	return nullptr;
-    }
-
-    virtual const webrtc::I420BufferInterface* GetI420() const override
-    {
-    	unimplemented();
-        return nullptr;
-    }
-    // End webrtc::VideoFrameBuffer interface
-
 private:
 	FPixelStreamingLayerFrameSource* LayerFrameSource = nullptr;
 };
diff --git a/Engine/Plugins/Media/PixelStreaming/Source/PixelStreaming/Private/PixelStreamingFramePump.cpp b/Engine/Plugins/Media/PixelStreaming/Source/PixelStreaming/Private/PixelStreamingFramePump.cpp
--- a/Engine/Plugins/Media/PixelStreaming/Source/PixelStreaming/Private/PixelStreamingFramePump.cpp
+++ b/Engine/Plugins/Media/PixelStreaming/Source/PixelStreaming/Private/PixelStreamingFramePump.cpp
@@ -5,10 +5,11 @@
 #include "PixelStreamingSettings.h"
 #include "PixelStreamingFrameSource.h"
 #include "PixelStreamingFrameBuffer.h"
+#include "PlayerSession.h"
 
 FPixelStreamingFramePump::FPixelStreamingFramePump(FPixelStreamingFrameSource* InFrameSource)
 :FrameSource(InFrameSource)
-,QualityControllerEvent(FPlatformProcess::GetSynchEventFromPool(false))
+,PlayersChangedEvent(FPlatformProcess::GetSynchEventFromPool(false))
 ,NextPumpEvent(FPlatformProcess::GetSynchEventFromPool(false))
 {
 	const int32 FPS = PixelStreamingSettings::CVarPixelStreamingWebRTCFps.GetValueOnAnyThread();
@@ -19,30 +20,53 @@ FPixelStreamingFramePump::FPixelStreamingFramePump(FPixelStreamingFrameSource* I
 FPixelStreamingFramePump::~FPixelStreamingFramePump()
 {
     bThreadRunning = false;
-    QualityControllerEvent->Trigger();
+    PlayersChangedEvent->Trigger();
     NextPumpEvent->Trigger();
     PumpThread->Join();
 }
 
-FPlayerVideoSource* FPixelStreamingFramePump::CreatePlayerVideoSource(FPlayerId PlayerId)
+FPlayerVideoSource* FPixelStreamingFramePump::CreatePlayerVideoSource(FPlayerId PlayerId, int Flags)
 {
-	FScopeLock Guard(&SourcesGuard);
-	PlayerVideoSources.Add(PlayerId, MakeUnique<FPlayerVideoSource>());
-	QualityControllerEvent->Trigger();
-	return PlayerVideoSources[PlayerId].Get();
+	FPlayerVideoSource* NewSource = nullptr;
+
+	if (PlayerVideoSources.Num() != 0)
+	{
+		FScopeLock Guard(&ExtraSourcesGuard);
+		ExtraVideoSources.Add(MakeUnique<FPlayerVideoSource>(PlayerId));
+		NewSource = ExtraVideoSources[ExtraVideoSources.Num() - 1].Get();
+	}
+	else
+	{
+		FScopeLock Guard(&SourcesGuard);
+		PlayerVideoSources.Add(PlayerId, MakeUnique<FPlayerVideoSource>(PlayerId));
+		NewSource = PlayerVideoSources[PlayerId].Get();
+	}
+
+	if ((Flags & FNewPlayerFlags::IsSFU) != 0)
+	{
+		ElevatedPlayerId = PlayerId;
+	}
+
+	PlayersChangedEvent->Trigger();
+	return NewSource;
 }
 
 void FPixelStreamingFramePump::RemovePlayerVideoSource(FPlayerId PlayerId)
 {
 	FScopeLock Guard(&SourcesGuard);
 	PlayerVideoSources.Remove(PlayerId);
+
+	if (PlayerId == ElevatedPlayerId)
+	{
+		ElevatedPlayerId = INVALID_PLAYER_ID;
+	}
 }
 
 void FPixelStreamingFramePump::SetQualityController(FPlayerId PlayerId)
 {
 	FScopeLock Guard(&SourcesGuard);
 	QualityControllerId = PlayerId;
-	QualityControllerEvent->Trigger();
+	PlayersChangedEvent->Trigger();
 }
 
 void FPixelStreamingFramePump::PumpLoop()
@@ -51,11 +75,46 @@ void FPixelStreamingFramePump::PumpLoop()
 
 	while (bThreadRunning)
 	{
-		bool IdleSleep = QualityControllerId == INVALID_PLAYER_ID;
-		if (!IdleSleep)
+		bool IdleSleep = true;
+
+		if (FrameSource->IsAvailable())
+		{
+			if (ExtraVideoSources.Num() > 0)
+			{
+				FScopeLock Guard0(&ExtraSourcesGuard);
+				FScopeLock Guard1(&SourcesGuard);
+
+				const int64 TimestampUs = rtc::TimeMicros();
+
+				// build a frame to pass to the quality controller source
+				rtc::scoped_refptr<FPixelStreamingInitializeFrameBuffer> EncoderFrameBuffer = new rtc::RefCountedObject<FPixelStreamingInitializeFrameBuffer>(FrameSource);
+
+				webrtc::VideoFrame Frame = webrtc::VideoFrame::Builder()
+					.set_video_frame_buffer(EncoderFrameBuffer)
+					.set_timestamp_us(TimestampUs)
+					.set_rotation(webrtc::VideoRotation::kVideoRotation_0)
+					.set_id(0)
+					.build();
+
+				for (auto& ExtraSource : ExtraVideoSources)
+				{
+					ExtraSource->OnFrameReady(Frame);
+					if (ExtraSource->IsInitialised())
+					{
+						PlayerVideoSources.Add(ExtraSource->GetPlayerId(), MoveTemp(ExtraSource));
+					}
+				}
+
+				ExtraVideoSources.RemoveAll([](auto& ExtraSource) { return !ExtraSource; });
+			}
+		}
+
+		const FPlayerId PumpPlayerId = (ElevatedPlayerId != INVALID_PLAYER_ID) ? ElevatedPlayerId : QualityControllerId;
+
+		if (PumpPlayerId != INVALID_PLAYER_ID)
 		{
 			FScopeLock Guard(&SourcesGuard); // lock here so the quality controller cant be removed after the check
-			if (PlayerVideoSources.Contains(QualityControllerId))
+			if (PlayerVideoSources.Contains(PumpPlayerId))
 			{
 				if (FrameSource->IsAvailable()) // dont trigger a sleep if this fails since we dont trigger the event for this changing
 				{
@@ -73,19 +132,15 @@ void FPixelStreamingFramePump::PumpLoop()
 						.build();
 					
 					// only send to quality controller (if it still exists. it might have been removed since the last check)
-					PlayerVideoSources[QualityControllerId]->OnFrameReady(Frame);
+					PlayerVideoSources[PumpPlayerId]->OnFrameReady(Frame);
 				}
-			}
-			else
-			{
-				// Quality controller doesn't exist. Idle.
-				IdleSleep = true;
+				IdleSleep = false;
 			}
 		}
 
 		if (IdleSleep)
 		{
-			QualityControllerEvent->Wait();
+			PlayersChangedEvent->Wait();
 		}
 		else
 		{
diff --git a/Engine/Plugins/Media/PixelStreaming/Source/PixelStreaming/Private/PixelStreamingFramePump.h b/Engine/Plugins/Media/PixelStreaming/Source/PixelStreaming/Private/PixelStreamingFramePump.h
--- a/Engine/Plugins/Media/PixelStreaming/Source/PixelStreaming/Private/PixelStreamingFramePump.h
+++ b/Engine/Plugins/Media/PixelStreaming/Source/PixelStreaming/Private/PixelStreamingFramePump.h
@@ -15,20 +15,25 @@ public:
 	FPixelStreamingFramePump(FPixelStreamingFrameSource* InFrameSource);
 	~FPixelStreamingFramePump();
 
-	FPlayerVideoSource* CreatePlayerVideoSource(FPlayerId PlayerId);
+	FPlayerVideoSource* CreatePlayerVideoSource(FPlayerId PlayerId, int Flags);
 	void RemovePlayerVideoSource(FPlayerId PlayerId);
 	void SetQualityController(FPlayerId PlayerId);
 
 private:
 	FCriticalSection SourcesGuard;
 	TMap<FPlayerId, TUniquePtr<FPlayerVideoSource>> PlayerVideoSources;
+
+	FCriticalSection ExtraSourcesGuard;
+	TArray<TUniquePtr<FPlayerVideoSource>> ExtraVideoSources;
+
+	FPlayerId ElevatedPlayerId = INVALID_PLAYER_ID;		// higher priority than quality controller
 	FPlayerId QualityControllerId = INVALID_PLAYER_ID;
 
 	FPixelStreamingFrameSource* FrameSource = nullptr;
 
 	TUniquePtr<FThread> PumpThread;
 	bool bThreadRunning = true;
-	FEvent* QualityControllerEvent;
+	FEvent* PlayersChangedEvent;
 	FEvent* NextPumpEvent;
 	double FrameDeltaMs = 1000.0 / 60.0; // default to 60fps
 	int32 NextFrameId = 0;
diff --git a/Engine/Plugins/Media/PixelStreaming/Source/PixelStreaming/Private/PixelStreamingRealEncoder.cpp b/Engine/Plugins/Media/PixelStreaming/Source/PixelStreaming/Private/PixelStreamingRealEncoder.cpp
--- a/Engine/Plugins/Media/PixelStreaming/Source/PixelStreaming/Private/PixelStreamingRealEncoder.cpp
+++ b/Engine/Plugins/Media/PixelStreaming/Source/PixelStreaming/Private/PixelStreamingRealEncoder.cpp
@@ -150,8 +150,13 @@ FPixelStreamingRealEncoder::~FPixelStreamingRealEncoder()
 
 void FPixelStreamingRealEncoder::Encode(const webrtc::VideoFrame& WebRTCFrame, bool Keyframe)
 {
-	FPixelStreamingLayerFrameBuffer* FrameBuffer = static_cast<FPixelStreamingLayerFrameBuffer*>(WebRTCFrame.video_frame_buffer().get());
-	FTexture2DRHIRef SourceTexture = FrameBuffer->GetFrame();
+	FTexture2DRHIRef SourceTexture;
+
+	FPixelStreamingFrameBuffer* FrameBuffer = static_cast<FPixelStreamingFrameBuffer*>(WebRTCFrame.video_frame_buffer().get());
+	check(FrameBuffer->GetFrameBufferType() == FPixelStreamingFrameBufferType::Layer);
+	FPixelStreamingLayerFrameBuffer* LayerFrameBuffer = static_cast<FPixelStreamingLayerFrameBuffer*>(FrameBuffer);
+	SourceTexture = LayerFrameBuffer->GetFrame();
+
 	if (SourceTexture)
 	{
 		AVEncoder::FVideoEncoderInputFrame* EncoderInputFrame = FrameFactory.GetFrameAndSetTexture(EncoderConfig.Width, EncoderConfig.Height, SourceTexture);
@@ -162,7 +167,8 @@ void FPixelStreamingRealEncoder::Encode(const webrtc::VideoFrame& WebRTCFrame, b
 			EncoderInputFrame->SetFrameID(WebRTCFrame.id());
 
 			AVEncoder::FVideoEncoder::FEncodeOptions Options;
-			Options.bForceKeyFrame = Keyframe;
+			Options.bForceKeyFrame = Keyframe || ForceNextKeyframe;
+			ForceNextKeyframe = false;
 
 			Encoder->Encode(EncoderInputFrame, Options);
 		}
diff --git a/Engine/Plugins/Media/PixelStreaming/Source/PixelStreaming/Private/PixelStreamingRealEncoder.h b/Engine/Plugins/Media/PixelStreaming/Source/PixelStreaming/Private/PixelStreamingRealEncoder.h
--- a/Engine/Plugins/Media/PixelStreaming/Source/PixelStreaming/Private/PixelStreamingRealEncoder.h
+++ b/Engine/Plugins/Media/PixelStreaming/Source/PixelStreaming/Private/PixelStreamingRealEncoder.h
@@ -14,6 +14,8 @@ public:
 
 	uint64 GetId() const { return Id; }
 
+	void SetForceNextKeyframe() { ForceNextKeyframe = true; }
+
 	void Encode(const webrtc::VideoFrame& WebRTCFrame, bool Keyframe);
 
 	AVEncoder::FVideoEncoder::FLayerConfig& GetCurrentConfig() { return EncoderConfig; }
@@ -25,4 +27,5 @@ private:
 	FPixelStreamingEncoderFrameFactory FrameFactory;
 	TUniquePtr<AVEncoder::FVideoEncoder> Encoder;
 	AVEncoder::FVideoEncoder::FLayerConfig EncoderConfig;
+	bool ForceNextKeyframe = false;
 };
diff --git a/Engine/Plugins/Media/PixelStreaming/Source/PixelStreaming/Private/PixelStreamingSimulcastEncoderAdapter.cpp b/Engine/Plugins/Media/PixelStreaming/Source/PixelStreaming/Private/PixelStreamingSimulcastEncoderAdapter.cpp
--- a/Engine/Plugins/Media/PixelStreaming/Source/PixelStreaming/Private/PixelStreamingSimulcastEncoderAdapter.cpp
+++ b/Engine/Plugins/Media/PixelStreaming/Source/PixelStreaming/Private/PixelStreamingSimulcastEncoderAdapter.cpp
@@ -1,704 +1,484 @@
 // Copyright Epic Games, Inc. All Rights Reserved.
 #include "PixelStreamingSimulcastEncoderAdapter.h"
-
-//#include <stdio.h>
-//#include <string.h>
-//
-//#include <algorithm>
-//#include <cstdint>
-//#include <string>
-//#include <utility>
-
-//#include "api/scoped_refptr.h"
-//#include "api/video/i420_buffer.h"
-//#include "api/video/video_codec_constants.h"
-//#include "api/video/video_frame_buffer.h"
-//#include "api/video/video_rotation.h"
-//#include "api/video_codecs/video_encoder.h"
-//#include "api/video_codecs/video_encoder_factory.h"
-//#include "api/video_codecs/video_encoder_software_fallback_wrapper.h"
-//#include "media/base/video_common.h"
-//#include "modules/video_coding/include/video_error_codes.h"
-//#include "modules/video_coding/utility/simulcast_rate_allocator.h"
-//#include "rtc_base/atomic_ops.h"
-//#include "rtc_base/checks.h"
-//#include "rtc_base/experiments/rate_control_settings.h"
-//#include "rtc_base/logging.h"
-//#include "system_wrappers/include/field_trial.h"
-
 #include "PixelStreamingFrameBuffer.h"
 
-namespace {
-
-#if PLATFORM_WINDOWS
-  #define sscanf sscanf_s
-#endif
-
-const unsigned int kDefaultMinQp = 2;
-const unsigned int kDefaultMaxQp = 56;
-// Max qp for lowest spatial resolution when doing simulcast.
-const unsigned int kLowestResMaxQp = 45;
-
-absl::optional<unsigned int> GetScreenshareBoostedQpValue() {
-  std::string experiment_group =
-      webrtc::field_trial::FindFullName("WebRTC-BoostedScreenshareQp");
-  unsigned int qp;
-  if (sscanf(experiment_group.c_str(), "%u", &qp) != 1)
-    return absl::nullopt;
-  qp = std::min(qp, 63u);
-  qp = std::max(qp, 1u);
-  return qp;
-}
-
-uint32_t SumStreamMaxBitrate(int streams, const webrtc::VideoCodec& codec) {
-  uint32_t bitrate_sum = 0;
-  for (int i = 0; i < streams; ++i) {
-    bitrate_sum += codec.simulcastStream[i].maxBitrate;
-  }
-  return bitrate_sum;
-}
-
-int NumberOfStreams(const webrtc::VideoCodec& codec) {
-  int streams =
-      codec.numberOfSimulcastStreams < 1 ? 1 : codec.numberOfSimulcastStreams;
-  uint32_t simulcast_max_bitrate = SumStreamMaxBitrate(streams, codec);
-  if (simulcast_max_bitrate == 0) {
-    streams = 1;
-  }
-  return streams;
-}
+namespace
+{
+    uint32_t SumStreamMaxBitrate(int streams, const webrtc::VideoCodec& codec)
+    {
+        uint32_t bitrate_sum = 0;
+        for (int i = 0; i < streams; ++i)
+        {
+            bitrate_sum += codec.simulcastStream[i].maxBitrate;
+        }
+        return bitrate_sum;
+    }
 
-int NumActiveStreams(const webrtc::VideoCodec& codec) {
-  int num_configured_streams = NumberOfStreams(codec);
-  int num_active_streams = 0;
-  for (int i = 0; i < num_configured_streams; ++i) {
-    if (codec.simulcastStream[i].active) {
-      ++num_active_streams;
+    int NumberOfStreams(const webrtc::VideoCodec& codec)
+    {
+        int streams = codec.numberOfSimulcastStreams < 1 ? 1 : codec.numberOfSimulcastStreams;
+        uint32_t simulcast_max_bitrate = SumStreamMaxBitrate(streams, codec);
+        if (simulcast_max_bitrate == 0)
+        {
+            streams = 1;
+        }
+        return streams;
     }
-  }
-  return num_active_streams;
-}
 
-int VerifyCodec(const webrtc::VideoCodec* inst) {
-  if (inst == nullptr) {
-    return WEBRTC_VIDEO_CODEC_ERR_PARAMETER;
-  }
-  if (inst->maxFramerate < 1) {
-    return WEBRTC_VIDEO_CODEC_ERR_PARAMETER;
-  }
-  // allow zero to represent an unspecified maxBitRate
-  if (inst->maxBitrate > 0 && inst->startBitrate > inst->maxBitrate) {
-    return WEBRTC_VIDEO_CODEC_ERR_PARAMETER;
-  }
-  if (inst->width <= 1 || inst->height <= 1) {
-    return WEBRTC_VIDEO_CODEC_ERR_PARAMETER;
-  }
-  if (inst->codecType == webrtc::kVideoCodecVP8 &&
-      inst->VP8().automaticResizeOn && NumActiveStreams(*inst) > 1) {
-    return WEBRTC_VIDEO_CODEC_ERR_PARAMETER;
-  }
-  return WEBRTC_VIDEO_CODEC_OK;
-}
+    int NumActiveStreams(const webrtc::VideoCodec& codec)
+    {
+        int num_configured_streams = NumberOfStreams(codec);
+        int num_active_streams = 0;
+        for (int i = 0; i < num_configured_streams; ++i)
+        {
+            if (codec.simulcastStream[i].active)
+            {
+                ++num_active_streams;
+            }
+        }
+        return num_active_streams;
+    }
 
-bool StreamResolutionCompare(const webrtc::SimulcastStream& a,
-                             const webrtc::SimulcastStream& b) {
-  return std::tie(a.height, a.width, a.maxBitrate, a.maxFramerate) <
-         std::tie(b.height, b.width, b.maxBitrate, b.maxFramerate);
-}
+    bool StreamResolutionCompare(const webrtc::SimulcastStream& a, const webrtc::SimulcastStream& b)
+    {
+        return std::tie(a.height, a.width, a.maxBitrate, a.maxFramerate) < std::tie(b.height, b.width, b.maxBitrate, b.maxFramerate);
+    }
 
-// An EncodedImageCallback implementation that forwards on calls to a
-// PixelStreamingSimulcastEncoderAdapter, but with the stream index it's registered with as
-// the first parameter to Encoded.
-class AdapterEncodedImageCallback : public webrtc::EncodedImageCallback {
- public:
-  AdapterEncodedImageCallback(PixelStreamingSimulcastEncoderAdapter* adapter,
-                              size_t stream_idx)
-      : adapter_(adapter), stream_idx_(stream_idx) {}
-
-  EncodedImageCallback::Result OnEncodedImage(
-      const webrtc::EncodedImage& encoded_image,
-      const webrtc::CodecSpecificInfo* codec_specific_info,
-      const webrtc::RTPFragmentationHeader* fragmentation) override {
-    return adapter_->OnEncodedImage(stream_idx_, encoded_image,
-                                    codec_specific_info, fragmentation);
-  }
-
- private:
-  PixelStreamingSimulcastEncoderAdapter* const adapter_;
-  const size_t stream_idx_;
-};
+	void PopulateStreamCodec(const webrtc::VideoCodec& codec_settings, int stream_index, uint32_t start_bitrate_kbps, webrtc::VideoCodec* stream_codec)
+	{
+		*stream_codec = codec_settings;
+
+		// Stream specific settings.
+		stream_codec->numberOfSimulcastStreams = 0;
+		stream_codec->width = codec_settings.simulcastStream[stream_index].width;
+		stream_codec->height = codec_settings.simulcastStream[stream_index].height;
+		stream_codec->maxBitrate = codec_settings.simulcastStream[stream_index].maxBitrate;
+		stream_codec->minBitrate = codec_settings.simulcastStream[stream_index].minBitrate;
+		stream_codec->maxFramerate = codec_settings.simulcastStream[stream_index].maxFramerate;
+		stream_codec->qpMax = codec_settings.simulcastStream[stream_index].qpMax;
+		stream_codec->active = codec_settings.simulcastStream[stream_index].active;
+		if (codec_settings.codecType == webrtc::kVideoCodecH264)
+		{
+			stream_codec->H264()->numberOfTemporalLayers = codec_settings.simulcastStream[stream_index].numberOfTemporalLayers;
+		}
+		stream_codec->startBitrate = start_bitrate_kbps;
+	}
+
+    // An EncodedImageCallback implementation that forwards on calls to a
+    // PixelStreamingSimulcastEncoderAdapter, but with the stream index it's registered with as
+    // the first parameter to Encoded.
+    class AdapterEncodedImageCallback : public webrtc::EncodedImageCallback
+    {
+    public:
+        AdapterEncodedImageCallback(PixelStreamingSimulcastEncoderAdapter* adapter,
+            size_t stream_idx)
+            : adapter_(adapter), stream_idx_(stream_idx)
+        {
+        }
+
+        EncodedImageCallback::Result OnEncodedImage(
+            const webrtc::EncodedImage& encoded_image,
+            const webrtc::CodecSpecificInfo* codec_specific_info,
+            const webrtc::RTPFragmentationHeader* fragmentation) override
+        {
+            return adapter_->OnEncodedImage(stream_idx_, encoded_image,
+                codec_specific_info, fragmentation);
+        }
+
+    private:
+        PixelStreamingSimulcastEncoderAdapter* const adapter_;
+        const size_t stream_idx_;
+    };
 }  // namespace
 
-PixelStreamingSimulcastEncoderAdapter::PixelStreamingSimulcastEncoderAdapter(webrtc::VideoEncoderFactory* factory,
-                                                 const webrtc::SdpVideoFormat& format)
-    : PixelStreamingSimulcastEncoderAdapter(factory, nullptr, format) {}
-
-PixelStreamingSimulcastEncoderAdapter::PixelStreamingSimulcastEncoderAdapter(
-    webrtc::VideoEncoderFactory* primary_factory,
-    webrtc::VideoEncoderFactory* fallback_factory,
-    const webrtc::SdpVideoFormat& format)
-    : inited_(0),
-      primary_encoder_factory_(primary_factory),
-      fallback_encoder_factory_(fallback_factory),
-      video_format_(format),
-      encoded_complete_callback_(nullptr),
-      experimental_boosted_screenshare_qp_(GetScreenshareBoostedQpValue()),
-      boost_base_layer_quality_(webrtc::RateControlSettings::ParseFromFieldTrials()
-                                    .Vp8BoostBaseLayerQuality()),
-      prefer_temporal_support_on_base_layer_(webrtc::field_trial::IsEnabled(
-          "WebRTC-Video-PreferTemporalSupportOnBaseLayer")) {
-  RTC_DCHECK(primary_factory);
-
-  // The adapter is typically created on the worker thread, but operated on
-  // the encoder task queue.
-  encoder_queue_.Detach();
-
-  memset(&codec_, 0, sizeof(webrtc::VideoCodec));
-}
-
-PixelStreamingSimulcastEncoderAdapter::~PixelStreamingSimulcastEncoderAdapter() {
-  RTC_DCHECK(!Initialized());
-  DestroyStoredEncoders();
+PixelStreamingSimulcastEncoderAdapter::PixelStreamingSimulcastEncoderAdapter(webrtc::VideoEncoderFactory* primary_factory, const webrtc::SdpVideoFormat& format)
+    : Initialized(false)
+    , PrimaryEncoderFactory(primary_factory)
+    , VideoFormat(format)
+    , EncodedCompleteCallback(nullptr)
+{
+    check(PrimaryEncoderFactory);
+    memset(&CurrentCodec, 0, sizeof(webrtc::VideoCodec));
 }
 
-void PixelStreamingSimulcastEncoderAdapter::SetFecControllerOverride(
-    webrtc::FecControllerOverride* fec_controller_override) {
-  // Ignored.
+PixelStreamingSimulcastEncoderAdapter::~PixelStreamingSimulcastEncoderAdapter()
+{
+    check(!IsInitialized());
 }
 
-int PixelStreamingSimulcastEncoderAdapter::Release() {
-  RTC_DCHECK_RUN_ON(&encoder_queue_);
-
-  while (!streaminfos_.empty()) {
-    std::unique_ptr<VideoEncoder> encoder =
-        std::move(streaminfos_.back().encoder);
-    // Even though it seems very unlikely, there are no guarantees that the
-    // encoder will not call back after being Release()'d. Therefore, we first
-    // disable the callbacks here.
-    encoder->RegisterEncodeCompleteCallback(nullptr);
-    encoder->Release();
-    streaminfos_.pop_back();  // Deletes callback adapter.
-    stored_encoders_.push(std::move(encoder));
-  }
-
-  // It's legal to move the encoder to another queue now.
-  encoder_queue_.Detach();
-
-  rtc::AtomicOps::ReleaseStore(&inited_, 0);
-
-  return WEBRTC_VIDEO_CODEC_OK;
-}
-
-// TODO(eladalon): s/inst/codec_settings/g.
-int PixelStreamingSimulcastEncoderAdapter::InitEncode(
-    const webrtc::VideoCodec* inst,
-    const webrtc::VideoEncoder::Settings& settings) {
-  RTC_DCHECK_RUN_ON(&encoder_queue_);
-
-  if (settings.number_of_cores < 1) {
-    return WEBRTC_VIDEO_CODEC_ERR_PARAMETER;
-  }
-
-  int ret = VerifyCodec(inst);
-  if (ret < 0) {
-    return ret;
-  }
-
-  ret = Release();
-  if (ret < 0) {
-    return ret;
-  }
-
-  int number_of_streams = NumberOfStreams(*inst);
-  RTC_DCHECK_LE(number_of_streams, webrtc::kMaxSimulcastStreams);
-  bool doing_simulcast_using_adapter = (number_of_streams > 1);
-  int num_active_streams = NumActiveStreams(*inst);
-
-  codec_ = *inst;
-  webrtc::SimulcastRateAllocator rate_allocator(codec_);
-  webrtc::VideoBitrateAllocation allocation =
-      rate_allocator.Allocate(webrtc::VideoBitrateAllocationParameters(
-          codec_.startBitrate * 1000, codec_.maxFramerate));
-  std::vector<uint32_t> start_bitrates;
-  for (int i = 0; i < webrtc::kMaxSimulcastStreams; ++i) {
-    uint32_t stream_bitrate = allocation.GetSpatialLayerSum(i) / 1000;
-    start_bitrates.push_back(stream_bitrate);
-  }
-
-  // Create |number_of_streams| of encoder instances and init them.
-  const auto minmax = std::minmax_element(
-      std::begin(codec_.simulcastStream),
-      std::begin(codec_.simulcastStream) + number_of_streams,
-      StreamResolutionCompare);
-  const auto lowest_resolution_stream_index =
-      std::distance(std::begin(codec_.simulcastStream), minmax.first);
-  const auto highest_resolution_stream_index =
-      std::distance(std::begin(codec_.simulcastStream), minmax.second);
-
-  RTC_DCHECK_LT(lowest_resolution_stream_index, number_of_streams);
-  RTC_DCHECK_LT(highest_resolution_stream_index, number_of_streams);
-
-  const webrtc::SdpVideoFormat format(
-      codec_.codecType == webrtc::kVideoCodecVP8 ? "VP8" : "H264",
-      video_format_.parameters);
-
-  for (int i = 0; i < number_of_streams; ++i) {
-    // If an existing encoder instance exists, reuse it.
-    // TODO(brandtr): Set initial RTP state (e.g., picture_id/tl0_pic_idx) here,
-    // when we start storing that state outside the encoder wrappers.
-    std::unique_ptr<VideoEncoder> encoder;
-    if (!stored_encoders_.empty()) {
-      encoder = std::move(stored_encoders_.top());
-      stored_encoders_.pop();
-    } else {
-      encoder = primary_encoder_factory_->CreateVideoEncoder(format);
-      if (fallback_encoder_factory_ != nullptr) {
-        encoder = CreateVideoEncoderSoftwareFallbackWrapper(
-            fallback_encoder_factory_->CreateVideoEncoder(format),
-            std::move(encoder),
-            i == lowest_resolution_stream_index &&
-                prefer_temporal_support_on_base_layer_);
-      }
-    }
-
-    bool encoder_initialized = false;
-    if (doing_simulcast_using_adapter && i == 0 &&
-        encoder->GetEncoderInfo().supports_simulcast) {
-      ret = encoder->InitEncode(&codec_, settings);
-      if (ret < 0) {
+int PixelStreamingSimulcastEncoderAdapter::Release()
+{
+    while (!StreamInfos.empty())
+    {
+        std::unique_ptr<VideoEncoder> encoder = std::move(StreamInfos.back().encoder);
+        // Even though it seems very unlikely, there are no guarantees that the
+        // encoder will not call back after being Release()'d. Therefore, we first
+        // disable the callbacks here.
+        encoder->RegisterEncodeCompleteCallback(nullptr);
         encoder->Release();
-      } else {
-        doing_simulcast_using_adapter = false;
-        number_of_streams = 1;
-        encoder_initialized = true;
-      }
-    }
-
-    webrtc::VideoCodec stream_codec;
-    uint32_t start_bitrate_kbps = start_bitrates[i];
-    const bool send_stream = doing_simulcast_using_adapter
-                                 ? start_bitrate_kbps > 0
-                                 : num_active_streams > 0;
-    if (!doing_simulcast_using_adapter) {
-      stream_codec = codec_;
-      stream_codec.numberOfSimulcastStreams =
-          std::max<uint8_t>(1, stream_codec.numberOfSimulcastStreams);
-    } else {
-      // Cap start bitrate to the min bitrate in order to avoid strange codec
-      // behavior. Since sending will be false, this should not matter.
-      StreamResolution stream_resolution =
-          i == highest_resolution_stream_index
-              ? StreamResolution::HIGHEST
-              : i == lowest_resolution_stream_index ? StreamResolution::LOWEST
-                                                    : StreamResolution::OTHER;
-
-      start_bitrate_kbps =
-          std::max(codec_.simulcastStream[i].minBitrate, start_bitrate_kbps);
-      PopulateStreamCodec(codec_, i, start_bitrate_kbps, stream_resolution,
-                          &stream_codec);
-    }
-
-    // TODO(ronghuawu): Remove once this is handled in LibvpxVp8Encoder.
-    if (stream_codec.qpMax < kDefaultMinQp) {
-      stream_codec.qpMax = kDefaultMaxQp;
-    }
-
-    if (!encoder_initialized) {
-      ret = encoder->InitEncode(&stream_codec, settings);
-      if (ret < 0) {
-        // Explicitly destroy the current encoder; because we haven't registered
-        // a StreamInfo for it yet, Release won't do anything about it.
-        encoder.reset();
-        Release();
-        return ret;
-      }
-    }
-
-    if (!doing_simulcast_using_adapter) {
-      // Without simulcast, just pass through the encoder info from the one
-      // active encoder.
-      encoder->RegisterEncodeCompleteCallback(encoded_complete_callback_);
-      streaminfos_.emplace_back(
-          std::move(encoder), nullptr,
-          std::make_unique<webrtc::FramerateController>(stream_codec.maxFramerate),
-          stream_codec.width, stream_codec.height, send_stream);
-    } else {
-      std::unique_ptr<webrtc::EncodedImageCallback> callback(
-          new AdapterEncodedImageCallback(this, i));
-      encoder->RegisterEncodeCompleteCallback(callback.get());
-      streaminfos_.emplace_back(
-          std::move(encoder), std::move(callback),
-          std::make_unique<webrtc::FramerateController>(stream_codec.maxFramerate),
-          stream_codec.width, stream_codec.height, send_stream);
+		StreamInfos.pop_back();  // Deletes callback adapter.
     }
-  }
 
-  // To save memory, don't store encoders that we don't use.
-  DestroyStoredEncoders();
+	Initialized = false;
 
-  rtc::AtomicOps::ReleaseStore(&inited_, 1);
-
-  return WEBRTC_VIDEO_CODEC_OK;
+    return WEBRTC_VIDEO_CODEC_OK;
 }
 
-int PixelStreamingSimulcastEncoderAdapter::Encode(
-    const webrtc::VideoFrame& input_image,
-    const std::vector<webrtc::VideoFrameType>* frame_types) {
-  RTC_DCHECK_RUN_ON(&encoder_queue_);
-
-  if (!Initialized()) {
-    return WEBRTC_VIDEO_CODEC_UNINITIALIZED;
-  }
-  if (encoded_complete_callback_ == nullptr) {
-    return WEBRTC_VIDEO_CODEC_UNINITIALIZED;
-  }
-
-  // All active streams should generate a key frame if
-  // a key frame is requested by any stream.
-  bool send_key_frame = false;
-  if (frame_types) {
-    for (size_t i = 0; i < frame_types->size(); ++i) {
-      if (frame_types->at(i) == webrtc::VideoFrameType::kVideoFrameKey) {
-        send_key_frame = true;
-        break;
-      }
-    }
-  }
-  for (size_t stream_idx = 0; stream_idx < streaminfos_.size(); ++stream_idx) {
-    if (streaminfos_[stream_idx].key_frame_request &&
-        streaminfos_[stream_idx].send_stream) {
-      send_key_frame = true;
-      break;
-    }
-  }
-
-  // Temporary thay may hold the result of texture to i420 buffer conversion.
-  rtc::scoped_refptr<webrtc::I420BufferInterface> src_buffer;
-  int src_width = input_image.width();
-  int src_height = input_image.height();
-  for (size_t stream_idx = 0; stream_idx < streaminfos_.size(); ++stream_idx) {
-    // Don't encode frames in resolutions that we don't intend to send.
-    if (!streaminfos_[stream_idx].send_stream) {
-      continue;
-    }
+int PixelStreamingSimulcastEncoderAdapter::InitEncode(const webrtc::VideoCodec* codec_settings, const webrtc::VideoEncoder::Settings& settings)
+{
+    int number_of_streams = NumberOfStreams(*codec_settings);
+    int num_active_streams_ = NumActiveStreams(*codec_settings);
+
+	CurrentCodec = *codec_settings;
+
+	const auto minmax = std::minmax_element(std::begin(CurrentCodec.simulcastStream), std::begin(CurrentCodec.simulcastStream) + number_of_streams, StreamResolutionCompare);
+	LowestResolutionStreamIndex = std::distance(std::begin(CurrentCodec.simulcastStream), minmax.first);
+	HighestResolutionStreamIndex = std::distance(std::begin(CurrentCodec.simulcastStream), minmax.second);
+
+    const webrtc::SdpVideoFormat format(CurrentCodec.codecType == webrtc::kVideoCodecVP8 ? "VP8" : "H264", VideoFormat.parameters);
+
+	if (num_active_streams_ == 1)
+	{
+		// with one stream we just proxy the pixelstreaming encoder
+		std::unique_ptr<VideoEncoder> encoder = PrimaryEncoderFactory->CreateVideoEncoder(format);
+		int ret = encoder->InitEncode(&CurrentCodec, settings);
+		if (ret < 0)
+		{
+			// Explicitly destroy the current encoder; because we haven't registered
+			// a StreamInfo for it yet, Release won't do anything about it.
+			encoder.reset();
+			Release();
+			return ret;
+		}
+
+		std::unique_ptr<webrtc::EncodedImageCallback> callback(new AdapterEncodedImageCallback(this, 0));
+		encoder->RegisterEncodeCompleteCallback(callback.get());
+		StreamInfos.emplace_back(
+			std::move(encoder), std::move(callback),
+			std::make_unique<webrtc::FramerateController>(CurrentCodec.maxFramerate),
+			CurrentCodec.width, CurrentCodec.height, true);
+	}
+	else
+	{
+		webrtc::SimulcastRateAllocator rate_allocator(CurrentCodec);
+		webrtc::VideoBitrateAllocation allocation = rate_allocator.Allocate(webrtc::VideoBitrateAllocationParameters(CurrentCodec.startBitrate * 1000, CurrentCodec.maxFramerate));
+		std::vector<uint32_t> start_bitrates;
+		for (int i = 0; i < webrtc::kMaxSimulcastStreams; ++i)
+		{
+			uint32_t stream_bitrate = allocation.GetSpatialLayerSum(i) / 1000;
+			start_bitrates.push_back(stream_bitrate);
+		}
+
+		for (int i = 0; i < number_of_streams; ++i)
+		{
+			webrtc::VideoCodec stream_codec;
+			uint32_t start_bitrate_kbps = std::max(CurrentCodec.simulcastStream[i].minBitrate, start_bitrates[i]);
+			PopulateStreamCodec(CurrentCodec, i, start_bitrate_kbps, &stream_codec);
+
+			std::unique_ptr<VideoEncoder> encoder = PrimaryEncoderFactory->CreateVideoEncoder(format);
+			int ret = encoder->InitEncode(&stream_codec, settings);
+			if (ret < 0)
+			{
+				// Explicitly destroy the current encoder; because we haven't registered
+				// a StreamInfo for it yet, Release won't do anything about it.
+				encoder.reset();
+				Release();
+				return ret;
+			}
+
+			std::unique_ptr<webrtc::EncodedImageCallback> callback(new AdapterEncodedImageCallback(this, i));
+			encoder->RegisterEncodeCompleteCallback(callback.get());
+			StreamInfos.emplace_back(
+				std::move(encoder), std::move(callback),
+				std::make_unique<webrtc::FramerateController>(stream_codec.maxFramerate),
+				stream_codec.width, stream_codec.height, true);
+		}
+	}
+
+	Initialized = true;
+
+    return WEBRTC_VIDEO_CODEC_OK;
+}
 
+int PixelStreamingSimulcastEncoderAdapter::EncodeStream(const webrtc::VideoFrame& input_image, FPixelStreamingLayerFrameSource* layer_frame_source, size_t stream_idx, bool send_key_frame)
+{
     // grab the simulcast frame source, extract the frame source for this layer and wrap that in a new frame buffer
-    FPixelStreamingSimulcastFrameBuffer* SimulcastFrameBuffer = static_cast<FPixelStreamingSimulcastFrameBuffer*>(input_image.video_frame_buffer().get());
-    FPixelStreamingLayerFrameSource* LayerFrameSource = SimulcastFrameBuffer->GetLayerFrameSource(stream_idx);
-    rtc::scoped_refptr<FPixelStreamingLayerFrameBuffer> LayerFrameBuffer = new rtc::RefCountedObject<FPixelStreamingLayerFrameBuffer>(LayerFrameSource);
+    rtc::scoped_refptr<FPixelStreamingLayerFrameBuffer> layer_frame_buffer = new rtc::RefCountedObject<FPixelStreamingLayerFrameBuffer>(layer_frame_source);
     webrtc::VideoFrame new_frame(input_image);
-    new_frame.set_video_frame_buffer(LayerFrameBuffer);
+    new_frame.set_video_frame_buffer(layer_frame_buffer);
 
-    const uint32_t frame_timestamp_ms =
-        1000 * new_frame.timestamp() / 90000;  // kVideoPayloadTypeFrequency;
+    const uint32_t frame_timestamp_ms =	1000 * new_frame.timestamp() / 90000;
 
     // If adapter is passed through and only one sw encoder does simulcast,
     // frame types for all streams should be passed to the encoder unchanged.
     // Otherwise a single per-encoder frame type is passed.
-    std::vector<webrtc::VideoFrameType> stream_frame_types(
-        streaminfos_.size() == 1 ? NumberOfStreams(codec_) : 1);
-    if (send_key_frame) {
-      std::fill(stream_frame_types.begin(), stream_frame_types.end(),
-                webrtc::VideoFrameType::kVideoFrameKey);
-      streaminfos_[stream_idx].key_frame_request = false;
-    } else {
-      if (streaminfos_[stream_idx].framerate_controller->DropFrame(
-              frame_timestamp_ms)) {
-        continue;
-      }
-      std::fill(stream_frame_types.begin(), stream_frame_types.end(),
-                webrtc::VideoFrameType::kVideoFrameDelta);
+    std::vector<webrtc::VideoFrameType> stream_frame_types(StreamInfos.size() == 1 ? NumberOfStreams(CurrentCodec) : 1);
+    if (send_key_frame)
+    {
+        std::fill(stream_frame_types.begin(), stream_frame_types.end(), webrtc::VideoFrameType::kVideoFrameKey);
+		StreamInfos[stream_idx].key_frame_request = false;
     }
-    streaminfos_[stream_idx].framerate_controller->AddFrame(frame_timestamp_ms);
-
-    int dst_width = streaminfos_[stream_idx].width;
-    int dst_height = streaminfos_[stream_idx].height;
-    // If scaling isn't required, because the input resolution
-    // matches the destination or the input image is empty (e.g.
-    // a keyframe request for encoders with internal camera
-    // sources) or the source image has a native handle, pass the image on
-    // directly. Otherwise, we'll scale it to match what the encoder expects
-    // (below).
-    // For texture frames, the underlying encoder is expected to be able to
-    // correctly sample/scale the source texture.
-    // TODO(perkj): ensure that works going forward, and figure out how this
-    // affects webrtc:5683.
-    if ((dst_width == src_width && dst_height == src_height) ||
-        (new_frame.video_frame_buffer()->type() ==
-             webrtc::VideoFrameBuffer::Type::kNative &&
-         streaminfos_[stream_idx]
-             .encoder->GetEncoderInfo()
-             .supports_native_handle)) {
-
-      int ret = streaminfos_[stream_idx].encoder->Encode(new_frame,
-                                                         &stream_frame_types);
-      if (ret != WEBRTC_VIDEO_CODEC_OK) {
-        return ret;
-      }
-    } else {
-      if (src_buffer == nullptr) {
-        src_buffer = input_image.video_frame_buffer()->ToI420();
-      }
-      rtc::scoped_refptr<webrtc::I420Buffer> dst_buffer =
-          webrtc::I420Buffer::Create(dst_width, dst_height);
-
-      dst_buffer->ScaleFrom(*src_buffer);
-
-      // UpdateRect is not propagated to lower simulcast layers currently.
-      // TODO(ilnik): Consider scaling UpdateRect together with the buffer.
-      webrtc::VideoFrame frame(input_image);
-      frame.set_video_frame_buffer(dst_buffer);
-      frame.set_rotation(webrtc::kVideoRotation_0);
-      frame.set_update_rect(
-          webrtc::VideoFrame::UpdateRect{0, 0, frame.width(), frame.height()});
-      int ret =
-          streaminfos_[stream_idx].encoder->Encode(frame, &stream_frame_types);
-      if (ret != WEBRTC_VIDEO_CODEC_OK) {
-        return ret;
-      }
+    else
+    {
+        if (StreamInfos[stream_idx].framerate_controller->DropFrame(frame_timestamp_ms))
+        {
+            return WEBRTC_VIDEO_CODEC_OK;
+        }
+        std::fill(stream_frame_types.begin(), stream_frame_types.end(),	webrtc::VideoFrameType::kVideoFrameDelta);
     }
-  }
+	StreamInfos[stream_idx].framerate_controller->AddFrame(frame_timestamp_ms);
 
-  return WEBRTC_VIDEO_CODEC_OK;
+    return StreamInfos[stream_idx].encoder->Encode(new_frame, &stream_frame_types);
 }
 
-int PixelStreamingSimulcastEncoderAdapter::RegisterEncodeCompleteCallback(
-    webrtc::EncodedImageCallback* callback) {
-  RTC_DCHECK_RUN_ON(&encoder_queue_);
-  encoded_complete_callback_ = callback;
-  if (streaminfos_.size() == 1) {
-    streaminfos_[0].encoder->RegisterEncodeCompleteCallback(callback);
-  }
-  return WEBRTC_VIDEO_CODEC_OK;
-}
+int PixelStreamingSimulcastEncoderAdapter::Encode(const webrtc::VideoFrame& input_image, const std::vector<webrtc::VideoFrameType>* frame_types)
+{
+    if (!IsInitialized())
+    {
+        return WEBRTC_VIDEO_CODEC_UNINITIALIZED;
+    }
 
-void PixelStreamingSimulcastEncoderAdapter::SetRates(
-    const RateControlParameters& parameters) {
-  RTC_DCHECK_RUN_ON(&encoder_queue_);
+    if (EncodedCompleteCallback == nullptr)
+    {
+        return WEBRTC_VIDEO_CODEC_UNINITIALIZED;
+    }
 
-  if (!Initialized()) {
-    RTC_LOG(LS_WARNING) << "SetRates while not initialized";
-    return;
-  }
+	// ignore any init frames
+	// NOTE: It seems most of the time they never even get here because the internals of WebRTC decide to drop it. It seems it thinks
+	// the encoder is paused or something.
+    const FPixelStreamingFrameBuffer* frame_buffer = static_cast<FPixelStreamingFrameBuffer*>(input_image.video_frame_buffer().get());
+    if (frame_buffer->GetFrameBufferType() == FPixelStreamingFrameBufferType::Initialize)
+    {
+        return WEBRTC_VIDEO_CODEC_OK;
+    }
 
-  if (parameters.framerate_fps < 1.0) {
-    RTC_LOG(LS_WARNING) << "Invalid framerate: " << parameters.framerate_fps;
-    return;
-  }
+	// separate out our pixelstreaming image sources
+    check(frame_buffer->GetFrameBufferType() == FPixelStreamingFrameBufferType::Simulcast);
+    const FPixelStreamingSimulcastFrameBuffer* simulcast_frame_buffer = static_cast<const FPixelStreamingSimulcastFrameBuffer*>(frame_buffer);
+
+    // All active streams should generate a key frame if
+    // a key frame is requested by any stream.
+    bool send_key_frame = false;
+    if (frame_types)
+    {
+        for (size_t i = 0; i < frame_types->size(); ++i)
+        {
+            if (frame_types->at(i) == webrtc::VideoFrameType::kVideoFrameKey)
+            {
+                send_key_frame = true;
+                break;
+            }
+        }
+    }
 
-  codec_.maxFramerate = static_cast<uint32_t>(parameters.framerate_fps + 0.5);
+    for (size_t stream_idx = 0; stream_idx < StreamInfos.size(); ++stream_idx)
+    {
+        if (StreamInfos[stream_idx].key_frame_request &&
+			StreamInfos[stream_idx].send_stream)
+        {
+            send_key_frame = true;
+            break;
+        }
+    }
 
-  if (streaminfos_.size() == 1) {
-    // Not doing simulcast.
-    streaminfos_[0].encoder->SetRates(parameters);
-    return;
-  }
+    for (size_t stream_idx = 0; stream_idx < StreamInfos.size(); ++stream_idx)
+    {
+        // Don't encode frames in resolutions that we don't intend to send.
+        if (!StreamInfos[stream_idx].send_stream)
+        {
+            continue;
+        }
+
+		FPixelStreamingLayerFrameSource* layer_frame_source = simulcast_frame_buffer->GetLayerFrameSource(StreamInfos.size() == 1 ? HighestResolutionStreamIndex : stream_idx);
+        int ret = EncodeStream(input_image, layer_frame_source, stream_idx, send_key_frame);
+        if (ret != WEBRTC_VIDEO_CODEC_OK)
+        {
+            return ret;
+        }
+    }
+	return WEBRTC_VIDEO_CODEC_OK;
+}
 
-  for (size_t stream_idx = 0; stream_idx < streaminfos_.size(); ++stream_idx) {
-    uint32_t stream_bitrate_kbps =
-        parameters.bitrate.GetSpatialLayerSum(stream_idx) / 1000;
+int PixelStreamingSimulcastEncoderAdapter::RegisterEncodeCompleteCallback(webrtc::EncodedImageCallback* callback)
+{
+	EncodedCompleteCallback = callback;
+    return WEBRTC_VIDEO_CODEC_OK;
+}
 
-    // Need a key frame if we have not sent this stream before.
-    if (stream_bitrate_kbps > 0 && !streaminfos_[stream_idx].send_stream) {
-      streaminfos_[stream_idx].key_frame_request = true;
-    }
-    streaminfos_[stream_idx].send_stream = stream_bitrate_kbps > 0;
-
-    // Slice the temporal layers out of the full allocation and pass it on to
-    // the encoder handling the current simulcast stream.
-    RateControlParameters stream_parameters = parameters;
-    stream_parameters.bitrate = webrtc::VideoBitrateAllocation();
-    for (int i = 0; i < webrtc::kMaxTemporalStreams; ++i) {
-      if (parameters.bitrate.HasBitrate(stream_idx, i)) {
-        stream_parameters.bitrate.SetBitrate(
-            0, i, parameters.bitrate.GetBitrate(stream_idx, i));
-      }
+void PixelStreamingSimulcastEncoderAdapter::SetRates(const RateControlParameters& parameters)
+{
+    if (!IsInitialized())
+    {
+        RTC_LOG(LS_WARNING) << "SetRates while not initialized";
+        return;
     }
 
-    // Assign link allocation proportionally to spatial layer allocation.
-    if (!parameters.bandwidth_allocation.IsZero() &&
-        parameters.bitrate.get_sum_bps() > 0) {
-      stream_parameters.bandwidth_allocation =
-          webrtc::DataRate::BitsPerSec((parameters.bandwidth_allocation.bps() *
-                                stream_parameters.bitrate.get_sum_bps()) /
-                               parameters.bitrate.get_sum_bps());
-      // Make sure we don't allocate bandwidth lower than target bitrate.
-      if (stream_parameters.bandwidth_allocation.bps() <
-          stream_parameters.bitrate.get_sum_bps()) {
-        stream_parameters.bandwidth_allocation =
-            webrtc::DataRate::BitsPerSec(stream_parameters.bitrate.get_sum_bps());
-      }
+    if (parameters.framerate_fps < 1.0)
+    {
+        RTC_LOG(LS_WARNING) << "Invalid framerate: " << parameters.framerate_fps;
+        return;
     }
 
-    stream_parameters.framerate_fps = std::min<double>(
-        parameters.framerate_fps,
-        streaminfos_[stream_idx].framerate_controller->GetTargetRate());
+	CurrentCodec.maxFramerate = static_cast<uint32_t>(parameters.framerate_fps + 0.5);
 
-    streaminfos_[stream_idx].encoder->SetRates(stream_parameters);
-  }
+    if (StreamInfos.size() == 1)
+    {
+        // Not doing simulcast.
+		StreamInfos[0].encoder->SetRates(parameters);
+        return;
+    }
+
+    for (size_t stream_idx = 0; stream_idx < StreamInfos.size(); ++stream_idx)
+    {
+        uint32_t stream_bitrate_kbps = parameters.bitrate.GetSpatialLayerSum(stream_idx) / 1000;
+
+        // Need a key frame if we have not sent this stream before.
+        if (stream_bitrate_kbps > 0 && !StreamInfos[stream_idx].send_stream)
+        {
+			StreamInfos[stream_idx].key_frame_request = true;
+        }
+		StreamInfos[stream_idx].send_stream = stream_bitrate_kbps > 0;
+
+        // Slice the temporal layers out of the full allocation and pass it on to
+        // the encoder handling the current simulcast stream.
+        RateControlParameters stream_parameters = parameters;
+        stream_parameters.bitrate = webrtc::VideoBitrateAllocation();
+        for (int i = 0; i < webrtc::kMaxTemporalStreams; ++i)
+        {
+            if (parameters.bitrate.HasBitrate(stream_idx, i))
+            {
+                stream_parameters.bitrate.SetBitrate(0, i, parameters.bitrate.GetBitrate(stream_idx, i));
+            }
+        }
+
+        // Assign link allocation proportionally to spatial layer allocation.
+        if (!parameters.bandwidth_allocation.IsZero() && parameters.bitrate.get_sum_bps() > 0)
+        {
+            stream_parameters.bandwidth_allocation =
+                webrtc::DataRate::BitsPerSec((parameters.bandwidth_allocation.bps() *
+                    stream_parameters.bitrate.get_sum_bps()) /
+                    parameters.bitrate.get_sum_bps());
+            // Make sure we don't allocate bandwidth lower than target bitrate.
+            if (stream_parameters.bandwidth_allocation.bps() < stream_parameters.bitrate.get_sum_bps())
+            {
+                stream_parameters.bandwidth_allocation =
+                    webrtc::DataRate::BitsPerSec(stream_parameters.bitrate.get_sum_bps());
+            }
+        }
+
+        stream_parameters.framerate_fps = std::min<double>(parameters.framerate_fps, StreamInfos[stream_idx].framerate_controller->GetTargetRate());
+
+		StreamInfos[stream_idx].encoder->SetRates(stream_parameters);
+    }
 }
 
-void PixelStreamingSimulcastEncoderAdapter::OnPacketLossRateUpdate(float packet_loss_rate) {
-  for (StreamInfo& info : streaminfos_) {
-    info.encoder->OnPacketLossRateUpdate(packet_loss_rate);
-  }
+void PixelStreamingSimulcastEncoderAdapter::OnPacketLossRateUpdate(float packet_loss_rate)
+{
+    for (StreamInfo& info : StreamInfos)
+    {
+        info.encoder->OnPacketLossRateUpdate(packet_loss_rate);
+    }
 }
 
-void PixelStreamingSimulcastEncoderAdapter::OnRttUpdate(int64_t rtt_ms) {
-  for (StreamInfo& info : streaminfos_) {
-    info.encoder->OnRttUpdate(rtt_ms);
-  }
+void PixelStreamingSimulcastEncoderAdapter::OnRttUpdate(int64_t rtt_ms)
+{
+    for (StreamInfo& info : StreamInfos)
+    {
+        info.encoder->OnRttUpdate(rtt_ms);
+    }
 }
 
-void PixelStreamingSimulcastEncoderAdapter::OnLossNotification(
-    const LossNotification& loss_notification) {
-  for (StreamInfo& info : streaminfos_) {
-    info.encoder->OnLossNotification(loss_notification);
-  }
+void PixelStreamingSimulcastEncoderAdapter::OnLossNotification(const LossNotification& loss_notification)
+{
+    for (StreamInfo& info : StreamInfos)
+    {
+        info.encoder->OnLossNotification(loss_notification);
+    }
 }
 
-// TODO(brandtr): Add task checker to this member function, when all encoder
-// callbacks are coming in on the encoder queue.
 webrtc::EncodedImageCallback::Result PixelStreamingSimulcastEncoderAdapter::OnEncodedImage(
     size_t stream_idx,
     const webrtc::EncodedImage& encodedImage,
     const webrtc::CodecSpecificInfo* codecSpecificInfo,
-    const webrtc::RTPFragmentationHeader* fragmentation) {
-  webrtc::EncodedImage stream_image(encodedImage);
-  webrtc::CodecSpecificInfo stream_codec_specific = *codecSpecificInfo;
+    const webrtc::RTPFragmentationHeader* fragmentation)
+{
+    webrtc::EncodedImage stream_image(encodedImage);
+    webrtc::CodecSpecificInfo stream_codec_specific = *codecSpecificInfo;
 
-  stream_image.SetSpatialIndex(stream_idx);
+    stream_image.SetSpatialIndex(stream_idx);
 
-  return encoded_complete_callback_->OnEncodedImage(
-      stream_image, &stream_codec_specific, fragmentation);
+    return EncodedCompleteCallback->OnEncodedImage(stream_image, &stream_codec_specific, fragmentation);
 }
 
-void PixelStreamingSimulcastEncoderAdapter::PopulateStreamCodec(
-    const webrtc::VideoCodec& inst,
-    int stream_index,
-    uint32_t start_bitrate_kbps,
-    StreamResolution stream_resolution,
-    webrtc::VideoCodec* stream_codec) {
-  *stream_codec = inst;
-
-  // Stream specific settings.
-  stream_codec->numberOfSimulcastStreams = 0;
-  stream_codec->width = inst.simulcastStream[stream_index].width;
-  stream_codec->height = inst.simulcastStream[stream_index].height;
-  stream_codec->maxBitrate = inst.simulcastStream[stream_index].maxBitrate;
-  stream_codec->minBitrate = inst.simulcastStream[stream_index].minBitrate;
-  stream_codec->maxFramerate = inst.simulcastStream[stream_index].maxFramerate;
-  stream_codec->qpMax = inst.simulcastStream[stream_index].qpMax;
-  stream_codec->active = inst.simulcastStream[stream_index].active;
-  // Settings that are based on stream/resolution.
-  if (stream_resolution == StreamResolution::LOWEST) {
-    // Settings for lowest spatial resolutions.
-    if (inst.mode == webrtc::VideoCodecMode::kScreensharing) {
-      if (experimental_boosted_screenshare_qp_) {
-        stream_codec->qpMax = *experimental_boosted_screenshare_qp_;
-      }
-    } else if (boost_base_layer_quality_) {
-      stream_codec->qpMax = kLowestResMaxQp;
-    }
-  }
-  if (inst.codecType == webrtc::kVideoCodecVP8) {
-    stream_codec->VP8()->numberOfTemporalLayers =
-        inst.simulcastStream[stream_index].numberOfTemporalLayers;
-    if (stream_resolution != StreamResolution::HIGHEST) {
-      // For resolutions below CIF, set the codec |complexity| parameter to
-      // kComplexityHigher, which maps to cpu_used = -4.
-      int pixels_per_frame = stream_codec->width * stream_codec->height;
-      if (pixels_per_frame < 352 * 288) {
-        stream_codec->VP8()->complexity =
-            webrtc::VideoCodecComplexity::kComplexityHigher;
-      }
-      // Turn off denoising for all streams but the highest resolution.
-      stream_codec->VP8()->denoisingOn = false;
-    }
-  } else if (inst.codecType == webrtc::kVideoCodecH264) {
-    stream_codec->H264()->numberOfTemporalLayers =
-        inst.simulcastStream[stream_index].numberOfTemporalLayers;
-  }
-  // TODO(ronghuawu): what to do with targetBitrate.
-
-  stream_codec->startBitrate = start_bitrate_kbps;
+bool PixelStreamingSimulcastEncoderAdapter::IsInitialized() const
+{
+	return Initialized;
 }
 
-bool PixelStreamingSimulcastEncoderAdapter::Initialized() const {
-  return rtc::AtomicOps::AcquireLoad(&inited_) == 1;
-}
-
-void PixelStreamingSimulcastEncoderAdapter::DestroyStoredEncoders() {
-  while (!stored_encoders_.empty()) {
-    stored_encoders_.pop();
-  }
-}
+webrtc::VideoEncoder::EncoderInfo PixelStreamingSimulcastEncoderAdapter::GetEncoderInfo() const
+{
+    if (StreamInfos.size() == 1)
+    {
+        // Not using simulcast adapting functionality, just pass through.
+        return StreamInfos[0].encoder->GetEncoderInfo();
+    }
 
-webrtc::VideoEncoder::EncoderInfo PixelStreamingSimulcastEncoderAdapter::GetEncoderInfo() const {
-  if (streaminfos_.size() == 1) {
-    // Not using simulcast adapting functionality, just pass through.
-    return streaminfos_[0].encoder->GetEncoderInfo();
-  }
-
-  VideoEncoder::EncoderInfo encoder_info;
-  encoder_info.implementation_name = "PixelStreamingSimulcastEncoderAdapter";
-  encoder_info.requested_resolution_alignment = 1;
-  encoder_info.supports_native_handle = true;
-  encoder_info.scaling_settings.thresholds = absl::nullopt;
-  if (streaminfos_.empty()) {
-    return encoder_info;
-  }
-
-  encoder_info.scaling_settings = VideoEncoder::ScalingSettings::kOff;
-  int num_active_streams = NumActiveStreams(codec_);
-
-  for (size_t i = 0; i < streaminfos_.size(); ++i) {
-    VideoEncoder::EncoderInfo encoder_impl_info =
-        streaminfos_[i].encoder->GetEncoderInfo();
-
-    if (i == 0) {
-      // Encoder name indicates names of all sub-encoders.
-      encoder_info.implementation_name += " (";
-      encoder_info.implementation_name += encoder_impl_info.implementation_name;
-
-      encoder_info.supports_native_handle =
-          encoder_impl_info.supports_native_handle;
-      encoder_info.has_trusted_rate_controller =
-          encoder_impl_info.has_trusted_rate_controller;
-      encoder_info.is_hardware_accelerated =
-          encoder_impl_info.is_hardware_accelerated;
-      encoder_info.has_internal_source = encoder_impl_info.has_internal_source;
-    } else {
-      encoder_info.implementation_name += ", ";
-      encoder_info.implementation_name += encoder_impl_info.implementation_name;
-
-      // Native handle supported if any encoder supports it.
-      encoder_info.supports_native_handle |=
-          encoder_impl_info.supports_native_handle;
-
-      // Trusted rate controller only if all encoders have it.
-      encoder_info.has_trusted_rate_controller &=
-          encoder_impl_info.has_trusted_rate_controller;
-
-      // Uses hardware support if any of the encoders uses it.
-      // For example, if we are having issues with down-scaling due to
-      // pipelining delay in HW encoders we need higher encoder usage
-      // thresholds in CPU adaptation.
-      encoder_info.is_hardware_accelerated |=
-          encoder_impl_info.is_hardware_accelerated;
-
-      // Has internal source only if all encoders have it.
-      encoder_info.has_internal_source &= encoder_impl_info.has_internal_source;
+    VideoEncoder::EncoderInfo encoder_info;
+    encoder_info.implementation_name = "PixelStreamingSimulcastEncoderAdapter";
+    encoder_info.requested_resolution_alignment = 1;
+    encoder_info.supports_native_handle = true;
+    encoder_info.scaling_settings.thresholds = absl::nullopt;
+    if (StreamInfos.empty())
+    {
+        return encoder_info;
     }
-    encoder_info.fps_allocation[i] = encoder_impl_info.fps_allocation[0];
-    encoder_info.requested_resolution_alignment = cricket::LeastCommonMultiple(
-        encoder_info.requested_resolution_alignment,
-        encoder_impl_info.requested_resolution_alignment);
-    if (num_active_streams == 1 && codec_.simulcastStream[i].active) {
-      encoder_info.scaling_settings = encoder_impl_info.scaling_settings;
+
+    encoder_info.scaling_settings = VideoEncoder::ScalingSettings::kOff;
+    int num_active_streams = NumActiveStreams(CurrentCodec);
+
+    for (size_t i = 0; i < StreamInfos.size(); ++i)
+    {
+        VideoEncoder::EncoderInfo encoder_impl_info = StreamInfos[i].encoder->GetEncoderInfo();
+
+        if (i == 0)
+        {
+            // Encoder name indicates names of all sub-encoders.
+            encoder_info.implementation_name += " (";
+            encoder_info.implementation_name += encoder_impl_info.implementation_name;
+
+            encoder_info.supports_native_handle = encoder_impl_info.supports_native_handle;
+            encoder_info.has_trusted_rate_controller = encoder_impl_info.has_trusted_rate_controller;
+            encoder_info.is_hardware_accelerated = encoder_impl_info.is_hardware_accelerated;
+            encoder_info.has_internal_source = encoder_impl_info.has_internal_source;
+        }
+        else
+        {
+            encoder_info.implementation_name += ", ";
+            encoder_info.implementation_name += encoder_impl_info.implementation_name;
+
+            // Native handle supported if any encoder supports it.
+            encoder_info.supports_native_handle |= encoder_impl_info.supports_native_handle;
+
+            // Trusted rate controller only if all encoders have it.
+            encoder_info.has_trusted_rate_controller &= encoder_impl_info.has_trusted_rate_controller;
+
+            // Uses hardware support if any of the encoders uses it.
+            // For example, if we are having issues with down-scaling due to
+            // pipelining delay in HW encoders we need higher encoder usage
+            // thresholds in CPU adaptation.
+            encoder_info.is_hardware_accelerated |= encoder_impl_info.is_hardware_accelerated;
+
+            // Has internal source only if all encoders have it.
+            encoder_info.has_internal_source &= encoder_impl_info.has_internal_source;
+        }
+        encoder_info.fps_allocation[i] = encoder_impl_info.fps_allocation[0];
+        encoder_info.requested_resolution_alignment = cricket::LeastCommonMultiple(encoder_info.requested_resolution_alignment, encoder_impl_info.requested_resolution_alignment);
+        if (num_active_streams == 1 && CurrentCodec.simulcastStream[i].active)
+        {
+            encoder_info.scaling_settings = encoder_impl_info.scaling_settings;
+        }
     }
-  }
-  encoder_info.implementation_name += ")";
+    encoder_info.implementation_name += ")";
 
-  return encoder_info;
+    return encoder_info;
 }
diff --git a/Engine/Plugins/Media/PixelStreaming/Source/PixelStreaming/Private/PixelStreamingSimulcastEncoderAdapter.h b/Engine/Plugins/Media/PixelStreaming/Source/PixelStreaming/Private/PixelStreamingSimulcastEncoderAdapter.h
--- a/Engine/Plugins/Media/PixelStreaming/Source/PixelStreaming/Private/PixelStreamingSimulcastEncoderAdapter.h
+++ b/Engine/Plugins/Media/PixelStreaming/Source/PixelStreaming/Private/PixelStreamingSimulcastEncoderAdapter.h
@@ -2,132 +2,83 @@
 #pragma once
 
 #include "CoreMinimal.h"
-
 #include "WebRTCIncludes.h"
-//#include <memory>
-
-//#include <string>
-//#include <utility>
-//#include <vector>
-
-//#include "absl/types/optional.h"
-//#include "api/fec_controller_override.h"
-//#include "api/video_codecs/sdp_video_format.h"
-//#include "api/video_codecs/video_encoder.h"
-//#include "modules/video_coding/include/video_codec_interface.h"
-//#include "modules/video_coding/utility/framerate_controller.h"
-//#include "rtc_base/atomic_ops.h"
-//#include "rtc_base/synchronization/sequence_checker.h"
-//#include "rtc_base/system/rtc_export.h"
 
-namespace webrtc {
-
-class SimulcastRateAllocator;
-class VideoEncoderFactory;
+class FPixelStreamingLayerFrameSource;
 
+namespace webrtc
+{
+    class SimulcastRateAllocator;
+    class VideoEncoderFactory;
 }
 
-
-// This is duplicated from webrtc::simulcast_encoder_adapter
-// Something about allocating within the library was causing memory corruption and
-// so duplicating it here makes sure its allocating the correct way.
-
-class PixelStreamingSimulcastEncoderAdapter : public webrtc::VideoEncoder {
- public:
-  // TODO(bugs.webrtc.org/11000): Remove when downstream usage is gone.
-  PixelStreamingSimulcastEncoderAdapter(webrtc::VideoEncoderFactory* primarty_factory,
-                          const webrtc::SdpVideoFormat& format);
-  // |primary_factory| produces the first-choice encoders to use.
-  // |fallback_factory|, if non-null, is used to create fallback encoder that
-  // will be used if InitEncode() fails for the primary encoder.
-  PixelStreamingSimulcastEncoderAdapter(webrtc::VideoEncoderFactory* primary_factory,
-                          webrtc::VideoEncoderFactory* fallback_factory,
-                          const webrtc::SdpVideoFormat& format);
-  ~PixelStreamingSimulcastEncoderAdapter() override;
-
-  // Implements VideoEncoder.
-  void SetFecControllerOverride(
-      webrtc::FecControllerOverride* fec_controller_override) override;
-  int Release() override;
-  int InitEncode(const webrtc::VideoCodec* codec_settings,
-                 const webrtc::VideoEncoder::Settings& settings) override;
-  int Encode(const webrtc::VideoFrame& input_image,
-             const std::vector<webrtc::VideoFrameType>* frame_types) override;
-  int RegisterEncodeCompleteCallback(webrtc::EncodedImageCallback* callback) override;
-  void SetRates(const RateControlParameters& parameters) override;
-  void OnPacketLossRateUpdate(float packet_loss_rate) override;
-  void OnRttUpdate(int64_t rtt_ms) override;
-  void OnLossNotification(const LossNotification& loss_notification) override;
-
-  // Eventual handler for the contained encoders' EncodedImageCallbacks, but
-  // called from an internal helper that also knows the correct stream
-  // index.
-  webrtc::EncodedImageCallback::Result OnEncodedImage(
-      size_t stream_idx,
-      const webrtc::EncodedImage& encoded_image,
-      const webrtc::CodecSpecificInfo* codec_specific_info,
-      const webrtc::RTPFragmentationHeader* fragmentation);
-
-  EncoderInfo GetEncoderInfo() const override;
-
- private:
-  struct StreamInfo {
-    StreamInfo(std::unique_ptr<webrtc::VideoEncoder> encoder,
-               std::unique_ptr<webrtc::EncodedImageCallback> callback,
-               std::unique_ptr<webrtc::FramerateController> framerate_controller,
-               uint16_t width,
-               uint16_t height,
-               bool send_stream)
-        : encoder(std::move(encoder)),
-          callback(std::move(callback)),
-          framerate_controller(std::move(framerate_controller)),
-          width(width),
-          height(height),
-          key_frame_request(false),
-          send_stream(send_stream) {}
-    std::unique_ptr<webrtc::VideoEncoder> encoder;
-    std::unique_ptr<webrtc::EncodedImageCallback> callback;
-    std::unique_ptr<webrtc::FramerateController> framerate_controller;
-    uint16_t width;
-    uint16_t height;
-    bool key_frame_request;
-    bool send_stream;
-  };
-
-  enum class StreamResolution {
-    OTHER,
-    HIGHEST,
-    LOWEST,
-  };
-
-  // Populate the codec settings for each simulcast stream.
-  void PopulateStreamCodec(const webrtc::VideoCodec& inst,
-                           int stream_index,
-                           uint32_t start_bitrate_kbps,
-                           StreamResolution stream_resolution,
-                           webrtc::VideoCodec* stream_codec);
-
-  bool Initialized() const;
-
-  void DestroyStoredEncoders();
-
-  volatile int inited_;  // Accessed atomically.
-  webrtc::VideoEncoderFactory* const primary_encoder_factory_;
-  webrtc::VideoEncoderFactory* const fallback_encoder_factory_;
-  const webrtc::SdpVideoFormat video_format_;
-  webrtc::VideoCodec codec_;
-  std::vector<StreamInfo> streaminfos_;
-  webrtc::EncodedImageCallback* encoded_complete_callback_;
-
-  // Used for checking the single-threaded access of the encoder interface.
-  webrtc::SequenceChecker encoder_queue_;
-
-  // Store encoders in between calls to Release and InitEncode, so they don't
-  // have to be recreated. Remaining encoders are destroyed by the destructor.
-  std::stack<std::unique_ptr<webrtc::VideoEncoder>> stored_encoders_;
-
-  const absl::optional<unsigned int> experimental_boosted_screenshare_qp_;
-  const bool boost_base_layer_quality_;
-  const bool prefer_temporal_support_on_base_layer_;
+// This is highly modified version of webrtc::simulcast_encoder_adapter
+
+class PixelStreamingSimulcastEncoderAdapter : public webrtc::VideoEncoder
+{
+public:
+    PixelStreamingSimulcastEncoderAdapter(webrtc::VideoEncoderFactory* primarty_factory, const webrtc::SdpVideoFormat& format);
+    ~PixelStreamingSimulcastEncoderAdapter() override;
+
+    // Implements VideoEncoder.
+    int Release() override;
+    int InitEncode(const webrtc::VideoCodec* codec_settings, const webrtc::VideoEncoder::Settings& settings) override;
+    int Encode(const webrtc::VideoFrame& input_image, const std::vector<webrtc::VideoFrameType>* frame_types) override;
+    int RegisterEncodeCompleteCallback(webrtc::EncodedImageCallback* callback) override;
+
+    void SetRates(const RateControlParameters& parameters) override;
+    void OnPacketLossRateUpdate(float packet_loss_rate) override;
+    void OnRttUpdate(int64_t rtt_ms) override;
+    void OnLossNotification(const LossNotification& loss_notification) override;
+
+    // Eventual handler for the contained encoders' EncodedImageCallbacks, but
+    // called from an internal helper that also knows the correct stream
+    // index.
+    webrtc::EncodedImageCallback::Result OnEncodedImage(
+        size_t stream_idx,
+        const webrtc::EncodedImage& encoded_image,
+        const webrtc::CodecSpecificInfo* codec_specific_info,
+        const webrtc::RTPFragmentationHeader* fragmentation);
+
+    EncoderInfo GetEncoderInfo() const override;
+
+private:
+    struct StreamInfo
+    {
+        StreamInfo(std::unique_ptr<webrtc::VideoEncoder> encoder,
+            std::unique_ptr<webrtc::EncodedImageCallback> callback,
+            std::unique_ptr<webrtc::FramerateController> framerate_controller,
+            uint16_t width,
+            uint16_t height,
+            bool send_stream)
+            : encoder(std::move(encoder))
+            , callback(std::move(callback))
+            , framerate_controller(std::move(framerate_controller))
+            , width(width)
+            , height(height)
+            , key_frame_request(false)
+            , send_stream(send_stream)
+        {
+        }
+        std::unique_ptr<webrtc::VideoEncoder> encoder;
+        std::unique_ptr<webrtc::EncodedImageCallback> callback;
+        std::unique_ptr<webrtc::FramerateController> framerate_controller;
+        uint16_t width;
+        uint16_t height;
+        bool key_frame_request;
+        bool send_stream;
+    };
+
+    bool IsInitialized() const;
+
+	TAtomic<bool> Initialized;
+    webrtc::VideoEncoderFactory* const PrimaryEncoderFactory;
+    const webrtc::SdpVideoFormat VideoFormat;
+    webrtc::VideoCodec CurrentCodec;
+    std::vector<StreamInfo> StreamInfos;
+    webrtc::EncodedImageCallback* EncodedCompleteCallback;
+
+	int LowestResolutionStreamIndex;
+	int HighestResolutionStreamIndex;
+    int EncodeStream(const webrtc::VideoFrame& input_image, FPixelStreamingLayerFrameSource* layer_frame_source, size_t stream_idx, bool send_key_frame);
 };
-
diff --git a/Engine/Plugins/Media/PixelStreaming/Source/PixelStreaming/Private/PixelStreamingVideoEncoder.cpp b/Engine/Plugins/Media/PixelStreaming/Source/PixelStreaming/Private/PixelStreamingVideoEncoder.cpp
--- a/Engine/Plugins/Media/PixelStreaming/Source/PixelStreaming/Private/PixelStreamingVideoEncoder.cpp
+++ b/Engine/Plugins/Media/PixelStreaming/Source/PixelStreaming/Private/PixelStreamingVideoEncoder.cpp
@@ -63,10 +63,9 @@ int32 FPixelStreamingVideoEncoder::Encode(webrtc::VideoFrame const& frame, std::
 	}
 
 	bool Keyframe = false;
-	if ((frame_types && (*frame_types)[0] == webrtc::VideoFrameType::kVideoFrameKey) || ForceNextKeyframe)
+	if ((frame_types && (*frame_types)[0] == webrtc::VideoFrameType::kVideoFrameKey))
 	{
 		Keyframe = true;
-		ForceNextKeyframe = false;
 	}
 
 	FPixelStreamingRealEncoder* Encoder = Factory.GetHardwareEncoder(HardwareEncoderId);
@@ -172,7 +171,20 @@ void FPixelStreamingVideoEncoder::SendEncodedImage(uint64 SourceEncoderId, webrt
 		return;
 	}
 
-	// Dump H264 frames to file for debugging if CVar is turned on.
+	if (FirstKeyframeCountdown > 0)
+	{
+		// ideally we want to make the first frame of new peers a keyframe but we
+		// dont know when webrtc will decide to start sending out frames. this is
+		// the next best option. its a count beause when it was the very next frame
+		// it still didnt seem to work. delaying it a few frames seems to have worked.
+		--FirstKeyframeCountdown;
+		if (FirstKeyframeCountdown == 0)
+		{
+			Factory.ForceKeyFrame();
+		}
+	}
+
+	// Dump H264 frames to file for debugging if CVar is turned 5on.
 	if (PixelStreamingSettings::CVarPixelStreamingDebugDumpFrame.GetValueOnAnyThread())
 	{
 		static IFileHandle* FileHandle = nullptr;
@@ -193,18 +205,3 @@ void FPixelStreamingVideoEncoder::SendEncodedImage(uint64 SourceEncoderId, webrt
 		OnEncodedImageCallback->OnEncodedImage(encoded_image, codec_specific_info, fragmentation);
 	}
 }
-
-bool FPixelStreamingVideoEncoder::IsRegisteredWithWebRTC()
-{
-	return OnEncodedImageCallback != nullptr;
-}
-
-int32_t FPixelStreamingVideoEncoder::GetSmoothedAverageQP() const
-{
-	return 0;
-	// if(this->Context == nullptr)
-	// {
-	// 	return -1;
-	// }
-	// return (int32_t)this->Context->SmoothedAvgQP.Get();
-}
diff --git a/Engine/Plugins/Media/PixelStreaming/Source/PixelStreaming/Private/PixelStreamingVideoEncoder.h b/Engine/Plugins/Media/PixelStreaming/Source/PixelStreaming/Private/PixelStreamingVideoEncoder.h
--- a/Engine/Plugins/Media/PixelStreaming/Source/PixelStreaming/Private/PixelStreamingVideoEncoder.h
+++ b/Engine/Plugins/Media/PixelStreaming/Source/PixelStreaming/Private/PixelStreamingVideoEncoder.h
@@ -9,8 +9,6 @@
 #include "Misc/Optional.h"
 
 class IPixelStreamingSessions;
-class FPlayerSession;
-struct FEncoderContext;
 class FPixelStreamingVideoEncoderFactory;
 
 // Implementation that is a WebRTC video encoder that allows us tie to our actually underlying non-WebRTC video encoder.
@@ -36,11 +34,6 @@ public:
 	// End WebRTC Interface.
 
 	void SendEncodedImage(uint64 SourceEncoderId, webrtc::EncodedImage const& encoded_image, webrtc::CodecSpecificInfo const* codec_specific_info, webrtc::RTPFragmentationHeader const* fragmentation);
-	bool IsRegisteredWithWebRTC();
-
-	void ForceKeyFrame() { ForceNextKeyframe = true; }
-	int32_t GetSmoothedAverageQP() const;
-	
 
 private:
 	void UpdateConfig();
@@ -58,9 +51,11 @@ private:
 
 	webrtc::EncodedImageCallback* OnEncodedImageCallback = nullptr;
 
-	bool ForceNextKeyframe = true;
-
 	// WebRTC may request a bitrate/framerate change using SetRates(), we only respect this if this encoder is actually encoding
 	// so we use this optional object to store a rate change and act upon it when this encoder does its next call to Encode().
 	TOptional<RateControlParameters> PendingRateChange;
+
+	// Used to send an initial keyframe
+	// see notes in SendEncodedImage implementation
+	int FirstKeyframeCountdown = 2;
 };
diff --git a/Engine/Plugins/Media/PixelStreaming/Source/PixelStreaming/Private/PlayerSession.h b/Engine/Plugins/Media/PixelStreaming/Source/PixelStreaming/Private/PlayerSession.h
--- a/Engine/Plugins/Media/PixelStreaming/Source/PixelStreaming/Private/PlayerSession.h
+++ b/Engine/Plugins/Media/PixelStreaming/Source/PixelStreaming/Private/PlayerSession.h
@@ -9,6 +9,13 @@
 #include "PlayerId.h"
 #include "PixelStreamingDataChannelObserver.h"
 
+enum FNewPlayerFlags
+{
+	None = 0,
+	SupportsDataChannel = 1,
+	IsSFU = 2,
+};
+
 class FSignallingServerConnection;
 class IPixelStreamingSessions;
 
diff --git a/Engine/Plugins/Media/PixelStreaming/Source/PixelStreaming/Private/PlayerVideoSource.cpp b/Engine/Plugins/Media/PixelStreaming/Source/PixelStreaming/Private/PlayerVideoSource.cpp
--- a/Engine/Plugins/Media/PixelStreaming/Source/PixelStreaming/Private/PlayerVideoSource.cpp
+++ b/Engine/Plugins/Media/PixelStreaming/Source/PixelStreaming/Private/PlayerVideoSource.cpp
@@ -1,7 +1,8 @@
 // Copyright Epic Games, Inc. All Rights Reserved.
 #include "PlayerVideoSource.h"
 
-FPlayerVideoSource::FPlayerVideoSource()
+FPlayerVideoSource::FPlayerVideoSource(FPlayerId InPlayerId)
+	:PlayerId(InPlayerId)
 {
 	CurrentState = webrtc::MediaSourceInterface::SourceState::kInitializing;
 }
@@ -28,6 +29,8 @@ void FPlayerVideoSource::OnFrameReady(const webrtc::VideoFrame& Frame)
 	}
 
 	OnFrame(Frame);
+
+	Initialised = true;
 }
 
 bool FPlayerVideoSource::AdaptCaptureFrame(const int64 TimestampUs, FIntPoint Resolution)
diff --git a/Engine/Plugins/Media/PixelStreaming/Source/PixelStreaming/Private/PlayerVideoSource.h b/Engine/Plugins/Media/PixelStreaming/Source/PixelStreaming/Private/PlayerVideoSource.h
--- a/Engine/Plugins/Media/PixelStreaming/Source/PixelStreaming/Private/PlayerVideoSource.h
+++ b/Engine/Plugins/Media/PixelStreaming/Source/PixelStreaming/Private/PlayerVideoSource.h
@@ -3,14 +3,18 @@
 #pragma once
 
 #include "WebRTCIncludes.h"
+#include "PlayerId.h"
 
 // each connected player has a single source. it is fed by FVideoSource and feeds the encoder
 class FPlayerVideoSource : public rtc::AdaptedVideoTrackSource
 {
 public:
-	FPlayerVideoSource();
+	FPlayerVideoSource(FPlayerId InPlayerId);
 	virtual ~FPlayerVideoSource();
 
+	FPlayerId GetPlayerId() const {	return PlayerId; }
+	bool IsInitialised() const { return Initialised; }
+
 	void AddRef() const override
 	{
 		FPlatformAtomics::InterlockedIncrement(&RefCount);
@@ -54,4 +58,6 @@ private:
 	webrtc::MediaSourceInterface::SourceState CurrentState;
 
 	mutable int32 RefCount;
+	FPlayerId PlayerId;
+	bool Initialised = false;
 };
diff --git a/Engine/Plugins/Media/PixelStreaming/Source/PixelStreaming/Private/SignallingServerConnection.cpp b/Engine/Plugins/Media/PixelStreaming/Source/PixelStreaming/Private/SignallingServerConnection.cpp
--- a/Engine/Plugins/Media/PixelStreaming/Source/PixelStreaming/Private/SignallingServerConnection.cpp
+++ b/Engine/Plugins/Media/PixelStreaming/Source/PixelStreaming/Private/SignallingServerConnection.cpp
@@ -14,6 +14,7 @@
 #include "PixelStreamingSettings.h"
 #include "TimerManager.h"
 #include "PixelStreamerDelegates.h"
+#include "PlayerSession.h"
 
 DECLARE_LOG_CATEGORY_EXTERN(LogPixelStreamingSS, Log, VeryVerbose);
 DEFINE_LOG_CATEGORY(LogPixelStreamingSS);
@@ -521,8 +522,10 @@ void FSignallingServerConnection::OnPlayerConnected(const FJsonObjectPtr& Json)
 	{
 		HANDLE_SS_ERROR(TEXT("Failed to get `playerId` from `join` message\n%s"), *ToString(Json));
 	}
-	const bool SupportsDataChannel = Json->GetBoolField(TEXT("dataChannel"));
-	Observer.OnPlayerConnected(PlayerId, SupportsDataChannel);
+	int Flags = 0;
+	Flags |= Json->GetBoolField(TEXT("dataChannel")) ? FNewPlayerFlags::SupportsDataChannel : 0;
+	Flags |= Json->GetBoolField(TEXT("sfu")) ? FNewPlayerFlags::IsSFU : 0;
+	Observer.OnPlayerConnected(PlayerId, Flags);
 }
 
 void FSignallingServerConnection::OnPlayerDisconnected(const FJsonObjectPtr& Json)
diff --git a/Engine/Plugins/Media/PixelStreaming/Source/PixelStreaming/Private/SignallingServerConnection.h b/Engine/Plugins/Media/PixelStreaming/Source/PixelStreaming/Private/SignallingServerConnection.h
--- a/Engine/Plugins/Media/PixelStreaming/Source/PixelStreaming/Private/SignallingServerConnection.h
+++ b/Engine/Plugins/Media/PixelStreaming/Source/PixelStreaming/Private/SignallingServerConnection.h
@@ -30,7 +30,7 @@ public:
 	// Streamer-only
 	virtual void OnRemoteIceCandidate(FPlayerId PlayerId, const std::string& SdpMid, int SdpMLineIndex, const std::string& Sdp)
 	{ unimplemented(); }
-	virtual void OnPlayerConnected(FPlayerId PlayerId, bool SupportsDataChannel)
+	virtual void OnPlayerConnected(FPlayerId PlayerId, int Flags)
 	{ unimplemented(); }
 	virtual void OnPlayerDisconnected(FPlayerId PlayerId)
 	{ unimplemented(); }
diff --git a/Engine/Plugins/Media/PixelStreaming/Source/PixelStreaming/Private/Streamer.cpp b/Engine/Plugins/Media/PixelStreaming/Source/PixelStreaming/Private/Streamer.cpp
--- a/Engine/Plugins/Media/PixelStreaming/Source/PixelStreaming/Private/Streamer.cpp
+++ b/Engine/Plugins/Media/PixelStreaming/Source/PixelStreaming/Private/Streamer.cpp
@@ -64,7 +64,6 @@ void FStreamer::StartWebRtcSignallingThread()
 	PeerConnectionConfig = {};
 
 	std::unique_ptr<FPixelStreamingSimulcastEncoderFactory> videoEncoderFactory = std::make_unique<FPixelStreamingSimulcastEncoderFactory>(this);
-	//std::unique_ptr<FPixelStreamingVideoEncoderFactory> videoEncoderFactory = std::make_unique<FPixelStreamingVideoEncoderFactory>(this);
 	VideoEncoderFactory = videoEncoderFactory->GetRealFactory();
 
 	bool bUseLegacyAudioDeviceModule = PixelStreamingSettings::CVarPixelStreamingWebRTCUseLegacyAudioDevice.GetValueOnAnyThread();
@@ -157,7 +156,7 @@ void FStreamer::OnDataChannelOpen(FPlayerId PlayerId, webrtc::DataChannelInterfa
 	}
 }
 
-webrtc::PeerConnectionInterface* FStreamer::CreateSession(FPlayerId PlayerId, bool SupportsDataChannel)
+webrtc::PeerConnectionInterface* FStreamer::CreateSession(FPlayerId PlayerId, int Flags)
 {
 	webrtc::PeerConnectionInterface* PeerConnection = PlayerSessions.CreatePlayerSession(PlayerId, 
 		PeerConnectionFactory, 
@@ -166,7 +165,7 @@ webrtc::PeerConnectionInterface* FStreamer::CreateSession(FPlayerId PlayerId, bo
 	
 	if (PeerConnection != nullptr)
 	{
-		if (SupportsDataChannel)
+		if ((Flags & FNewPlayerFlags::SupportsDataChannel) != 0)
 		{
 			webrtc::DataChannelInit DataChannelConfig;
 			DataChannelConfig.reliable = true;
@@ -181,7 +180,7 @@ webrtc::PeerConnectionInterface* FStreamer::CreateSession(FPlayerId PlayerId, bo
 			DataChannelObserver->Register(DataChannel);
 		}
 
-		AddStreams(PlayerId, PeerConnection);
+		AddStreams(PlayerId, PeerConnection, Flags);
 	}
 
 	return PeerConnection;
@@ -191,7 +190,7 @@ void FStreamer::OnSessionDescription(FPlayerId PlayerId, webrtc::SdpType Type, c
 {
 	if (Type == webrtc::SdpType::kOffer)
 	{
-		if (webrtc::PeerConnectionInterface* PeerConnection = CreateSession(PlayerId, true))
+		if (webrtc::PeerConnectionInterface* PeerConnection = CreateSession(PlayerId, FNewPlayerFlags::SupportsDataChannel))
 		{
 			webrtc::SdpParseError Error;
 			std::unique_ptr<webrtc::SessionDescriptionInterface> SessionDesc = webrtc::CreateSessionDescription(Type, to_string(Sdp), &Error);
@@ -357,10 +356,10 @@ void FStreamer::OnRemoteIceCandidate(FPlayerId PlayerId, const std::string& SdpM
 	this->PlayerSessions.OnRemoteIceCandidate(PlayerId, SdpMid, SdpMLineIndex, Sdp);
 }
 
-void FStreamer::OnPlayerConnected(FPlayerId PlayerId, bool SupportsDataChannel)
+void FStreamer::OnPlayerConnected(FPlayerId PlayerId, int Flags)
 {
 	// create peer connection
-	if (webrtc::PeerConnectionInterface* PlayerSession = CreateSession(PlayerId, SupportsDataChannel))
+	if (webrtc::PeerConnectionInterface* PlayerSession = CreateSession(PlayerId, Flags))
 	{
 		// make them send only
 		std::vector<rtc::scoped_refptr<webrtc::RtpTransceiverInterface>> Transceivers = PlayerSession->GetTransceivers();
@@ -482,16 +481,16 @@ void FStreamer::DeletePlayerSession(FPlayerId PlayerId)
 	}
 }
 
-void FStreamer::AddStreams(FPlayerId PlayerId, webrtc::PeerConnectionInterface* PeerConnection)
+void FStreamer::AddStreams(FPlayerId PlayerId, webrtc::PeerConnectionInterface* PeerConnection, int Flags)
 {
 	check(PeerConnection);
 
 	bool bSyncVideoAndAudio = !PixelStreamingSettings::CVarPixelStreamingWebRTCDisableAudioSync.GetValueOnAnyThread();
 	
-	FString const AudioStreamId = bSyncVideoAndAudio ? TEXT("pixelstreaming_av_stream_id") : TEXT("pixelstreaming_audio_stream_id");
-	FString const VideoStreamId = bSyncVideoAndAudio ? TEXT("pixelstreaming_av_stream_id") : TEXT("pixelstreaming_video_stream_id");
-	FString const AudioTrackLabel = TEXT("pixelstreaming_audio_track_label");
-	FString const VideoTrackLabel = TEXT("pixelstreaming_video_track_label");
+	const FString AudioStreamId = bSyncVideoAndAudio ? TEXT("pixelstreaming_av_stream_id") : TEXT("pixelstreaming_audio_stream_id");
+	const FString VideoStreamId = bSyncVideoAndAudio ? TEXT("pixelstreaming_av_stream_id") : TEXT("pixelstreaming_video_stream_id");
+	const FString AudioTrackLabel = TEXT("pixelstreaming_audio_track_label");
+	const FString VideoTrackLabel = TEXT("pixelstreaming_video_track_label");
 
 	if (!PeerConnection->GetSenders().empty())
 	{
@@ -501,14 +500,14 @@ void FStreamer::AddStreams(FPlayerId PlayerId, webrtc::PeerConnectionInterface*
 	// Use PeerConnection's transceiver API to add create audio/video tracks will correct directionality.
 	// These tracks are only thin wrappers around the underlying sources (the sources are shared among all peer's tracks).
 	// As per the WebRTC source: "The same source can be used by multiple VideoTracks."
-	this->SetupVideoTrack(PlayerId, PeerConnection, VideoStreamId, VideoTrackLabel);
+	this->SetupVideoTrack(PlayerId, PeerConnection, VideoStreamId, VideoTrackLabel, Flags);
 	this->SetupAudioTrack(PeerConnection, AudioStreamId, AudioTrackLabel);
 }
 
-void FStreamer::SetupVideoTrack(FPlayerId PlayerId, webrtc::PeerConnectionInterface* PeerConnection, FString const VideoStreamId, FString const VideoTrackLabel)
+void FStreamer::SetupVideoTrack(FPlayerId PlayerId, webrtc::PeerConnectionInterface* PeerConnection, const FString VideoStreamId, const FString VideoTrackLabel, int Flags)
 {
 	// Create a video source for this player
-	FPlayerVideoSource* PlayerVideoSource = FramePump.CreatePlayerVideoSource(PlayerId);
+	FPlayerVideoSource* PlayerVideoSource = FramePump.CreatePlayerVideoSource(PlayerId, Flags);
 
 	// Create video track
 	rtc::scoped_refptr<webrtc::VideoTrackInterface> VideoTrack = PeerConnectionFactory->CreateVideoTrack(TCHAR_TO_UTF8(*VideoTrackLabel), PlayerVideoSource);
diff --git a/Engine/Plugins/Media/PixelStreaming/Source/PixelStreaming/Private/Streamer.h b/Engine/Plugins/Media/PixelStreaming/Source/PixelStreaming/Private/Streamer.h
--- a/Engine/Plugins/Media/PixelStreaming/Source/PixelStreaming/Private/Streamer.h
+++ b/Engine/Plugins/Media/PixelStreaming/Source/PixelStreaming/Private/Streamer.h
@@ -57,7 +57,7 @@ public:
 
 private:
 
-	webrtc::PeerConnectionInterface* CreateSession(FPlayerId PlayerId, bool SupportsDataChannel);
+	webrtc::PeerConnectionInterface* CreateSession(FPlayerId PlayerId, int Flags);
 	webrtc::AudioProcessing* SetupAudioProcessingModule();
 
 	// Procedure for WebRTC inter-thread communication
@@ -68,7 +68,7 @@ private:
 	virtual void OnConfig(const webrtc::PeerConnectionInterface::RTCConfiguration& Config) override;
 	virtual void OnSessionDescription(FPlayerId PlayerId, webrtc::SdpType Type, const FString& Sdp) override;
 	virtual void OnRemoteIceCandidate(FPlayerId PlayerId, const std::string& SdpMid, int SdpMLineIndex, const std::string& Sdp) override;
-	virtual void OnPlayerConnected(FPlayerId PlayerId, bool SupportsDataChannel) override;
+	virtual void OnPlayerConnected(FPlayerId PlayerId, int Flags) override;
 	virtual void OnPlayerDisconnected(FPlayerId PlayerId) override;
 	virtual void OnSignallingServerDisconnected() override;
 	
@@ -76,9 +76,9 @@ private:
 	void ModifyAudioTransceiverDirection(webrtc::PeerConnectionInterface* PeerConnection);
 	void DeletePlayerSession(FPlayerId PlayerId);
 	void DeleteAllPlayerSessions();
-	void AddStreams(FPlayerId PlayerId, webrtc::PeerConnectionInterface* PeerConnection);
-	void SetupVideoTrack(FPlayerId PlayerId, webrtc::PeerConnectionInterface* PeerConnection, FString const VideoStreamId, FString const VideoTrackLabel);
-	void SetupAudioTrack(webrtc::PeerConnectionInterface* PeerConnection, FString const AudioStreamId, FString const AudioTrackLabel);
+	void AddStreams(FPlayerId PlayerId, webrtc::PeerConnectionInterface* PeerConnection, int Flags);
+	void SetupVideoTrack(FPlayerId PlayerId, webrtc::PeerConnectionInterface* PeerConnection, FString const VideoStreamId, FString const VideoTrackLabel, int Flags);
+	void SetupAudioTrack(webrtc::PeerConnectionInterface* PeerConnection, const FString AudioStreamId, const FString AudioTrackLabel);
 	void HandleOffer(FPlayerId PlayerId, webrtc::PeerConnectionInterface* PeerConnection, TUniquePtr<webrtc::SessionDescriptionInterface> Sdp);
 	void OnDataChannelOpen(FPlayerId PlayerId, webrtc::DataChannelInterface* DataChannel);
 	void OnQualityControllerChanged(FPlayerId PlayerId);
diff --git a/Engine/Plugins/Media/PixelStreaming/Source/PixelStreaming/Private/WebRTCIncludes.h b/Engine/Plugins/Media/PixelStreaming/Source/PixelStreaming/Private/WebRTCIncludes.h
--- a/Engine/Plugins/Media/PixelStreaming/Source/PixelStreaming/Private/WebRTCIncludes.h
+++ b/Engine/Plugins/Media/PixelStreaming/Source/PixelStreaming/Private/WebRTCIncludes.h
@@ -41,24 +41,6 @@ THIRD_PARTY_INCLUDES_START
 #include "api/video/i420_buffer.h"
 #include "api/video/video_sink_interface.h"
 
-#if PLATFORM_WINDOWS
-// some code in the simulcast adaptor has trouble finding the windows versions of these
-// functions. these are here to redirect them properly
-template<class T>
-T InterlockedIncrement(volatile T* i) {
-	return ::_InterlockedIncrement(i);
-}
-template<class T>
-T InterlockedDecrement(volatile T* i) {
-	return ::_InterlockedDecrement(i);
-}
-template<class T, class U>
-T InterlockedCompareExchange(volatile T* i, U old_value, U new_value) {
-	return ::_InterlockedCompareExchange(i, new_value, old_value);
-}
-#endif
-
-#include "rtc_base/atomic_ops.h"
 #include "rtc_base/thread.h"
 #include "rtc_base/logging.h"
 #include "rtc_base/ssl_adapter.h"
