Merging //UE5/Dev-TensorWorks to Release-5.0 (//UE5/Release-5.0)
- Fixing new player flags conflict on linux

diff --git a/Engine/Plugins/Media/PixelStreaming/Source/PixelStreaming/Private/PixelStreamingFramePump.cpp b/Engine/Plugins/Media/PixelStreaming/Source/PixelStreaming/Private/PixelStreamingFramePump.cpp
--- a/Engine/Plugins/Media/PixelStreaming/Source/PixelStreaming/Private/PixelStreamingFramePump.cpp
+++ b/Engine/Plugins/Media/PixelStreaming/Source/PixelStreaming/Private/PixelStreamingFramePump.cpp
@@ -42,7 +42,7 @@ FPlayerVideoSource* FPixelStreamingFramePump::CreatePlayerVideoSource(FPlayerId
 		NewSource = PlayerVideoSources[PlayerId].Get();
 	}
 
-	if ((Flags & FNewPlayerFlags::IsSFU) != 0)
+	if ((Flags & PSPFlag_IsSFU) != 0)
 	{
 		ElevatedPlayerId = PlayerId;
 	}
diff --git a/Engine/Plugins/Media/PixelStreaming/Source/PixelStreaming/Private/PlayerSession.h b/Engine/Plugins/Media/PixelStreaming/Source/PixelStreaming/Private/PlayerSession.h
--- a/Engine/Plugins/Media/PixelStreaming/Source/PixelStreaming/Private/PlayerSession.h
+++ b/Engine/Plugins/Media/PixelStreaming/Source/PixelStreaming/Private/PlayerSession.h
@@ -9,11 +9,11 @@
 #include "PlayerId.h"
 #include "PixelStreamingDataChannelObserver.h"
 
-enum FNewPlayerFlags
+enum FPixelStreamingPlayerFlags
 {
-	None = 0,
-	SupportsDataChannel = 1,
-	IsSFU = 2,
+	PSPFlag_None = 0,
+	PSPFlag_SupportsDataChannel = 1,
+	PSPFlag_IsSFU = 2,
 };
 
 class FSignallingServerConnection;
diff --git a/Engine/Plugins/Media/PixelStreaming/Source/PixelStreaming/Private/SignallingServerConnection.cpp b/Engine/Plugins/Media/PixelStreaming/Source/PixelStreaming/Private/SignallingServerConnection.cpp
--- a/Engine/Plugins/Media/PixelStreaming/Source/PixelStreaming/Private/SignallingServerConnection.cpp
+++ b/Engine/Plugins/Media/PixelStreaming/Source/PixelStreaming/Private/SignallingServerConnection.cpp
@@ -523,8 +523,8 @@ void FSignallingServerConnection::OnPlayerConnected(const FJsonObjectPtr& Json)
 		HANDLE_SS_ERROR(TEXT("Failed to get `playerId` from `join` message\n%s"), *ToString(Json));
 	}
 	int Flags = 0;
-	Flags |= Json->GetBoolField(TEXT("dataChannel")) ? FNewPlayerFlags::SupportsDataChannel : 0;
-	Flags |= Json->GetBoolField(TEXT("sfu")) ? FNewPlayerFlags::IsSFU : 0;
+	Flags |= Json->GetBoolField(TEXT("dataChannel")) ? PSPFlag_SupportsDataChannel : 0;
+	Flags |= Json->GetBoolField(TEXT("sfu")) ? PSPFlag_IsSFU : 0;
 	Observer.OnPlayerConnected(PlayerId, Flags);
 }
 
diff --git a/Engine/Plugins/Media/PixelStreaming/Source/PixelStreaming/Private/Streamer.cpp b/Engine/Plugins/Media/PixelStreaming/Source/PixelStreaming/Private/Streamer.cpp
--- a/Engine/Plugins/Media/PixelStreaming/Source/PixelStreaming/Private/Streamer.cpp
+++ b/Engine/Plugins/Media/PixelStreaming/Source/PixelStreaming/Private/Streamer.cpp
@@ -165,7 +165,7 @@ webrtc::PeerConnectionInterface* FStreamer::CreateSession(FPlayerId PlayerId, in
 	
 	if (PeerConnection != nullptr)
 	{
-		if ((Flags & FNewPlayerFlags::SupportsDataChannel) != 0)
+		if ((Flags & PSPFlag_SupportsDataChannel) != 0)
 		{
 			webrtc::DataChannelInit DataChannelConfig;
 			DataChannelConfig.reliable = true;
@@ -190,7 +190,7 @@ void FStreamer::OnSessionDescription(FPlayerId PlayerId, webrtc::SdpType Type, c
 {
 	if (Type == webrtc::SdpType::kOffer)
 	{
-		if (webrtc::PeerConnectionInterface* PeerConnection = CreateSession(PlayerId, FNewPlayerFlags::SupportsDataChannel))
+		if (webrtc::PeerConnectionInterface* PeerConnection = CreateSession(PlayerId, PSPFlag_SupportsDataChannel))
 		{
 			webrtc::SdpParseError Error;
 			std::unique_ptr<webrtc::SessionDescriptionInterface> SessionDesc = webrtc::CreateSessionDescription(Type, to_string(Sdp), &Error);
