Pixel Streaming plugin formatting pass

diff --git a/Engine/Plugins/Media/PixelStreaming/Source/PixelStreaming/Private/InputDevice.cpp b/Engine/Plugins/Media/PixelStreaming/Source/PixelStreaming/Private/InputDevice.cpp
--- a/Engine/Plugins/Media/PixelStreaming/Source/PixelStreaming/Private/InputDevice.cpp
+++ b/Engine/Plugins/Media/PixelStreaming/Source/PixelStreaming/Private/InputDevice.cpp
@@ -37,12 +37,12 @@ public:
 	virtual ~FCursor() = default;
 	virtual FVector2D GetPosition() const override { return Position; }
 	virtual void SetPosition(const int32 X, const int32 Y) override { Position = FVector2D(X, Y); };
-	virtual void SetType(const EMouseCursor::Type InNewCursor) override {};
+	virtual void SetType(const EMouseCursor::Type InNewCursor) override{};
 	virtual EMouseCursor::Type GetType() const override { return EMouseCursor::Type::Default; };
-	virtual void GetSize(int32& Width, int32& Height) const override {};
-	virtual void Show(bool bShow) override {};
-	virtual void Lock(const RECT* const Bounds) override {};
-	virtual void SetTypeShape(EMouseCursor::Type InCursorType, void* CursorHandle) override {};
+	virtual void GetSize(int32& Width, int32& Height) const override{};
+	virtual void Show(bool bShow) override{};
+	virtual void Lock(const RECT* const Bounds) override{};
+	virtual void SetTypeShape(EMouseCursor::Type InCursorType, void* CursorHandle) override{};
 
 private:
 	/** The cursor position sent across with mouse events. */
@@ -72,18 +72,18 @@ public:
 	 * Functions passed directly to the wrapped application.
 	 */
 
-	virtual void SetMessageHandler(const TSharedRef< FGenericApplicationMessageHandler >& InMessageHandler) { WrappedApplication->SetMessageHandler(InMessageHandler); }
+	virtual void SetMessageHandler(const TSharedRef<FGenericApplicationMessageHandler>& InMessageHandler) { WrappedApplication->SetMessageHandler(InMessageHandler); }
 	virtual void PollGameDeviceState(const float TimeDelta) { WrappedApplication->PollGameDeviceState(TimeDelta); }
 	virtual void PumpMessages(const float TimeDelta) { WrappedApplication->PumpMessages(TimeDelta); }
 	virtual void ProcessDeferredEvents(const float TimeDelta) { WrappedApplication->ProcessDeferredEvents(TimeDelta); }
 	virtual void Tick(const float TimeDelta) { WrappedApplication->Tick(TimeDelta); }
-	virtual TSharedRef< FGenericWindow > MakeWindow() { return WrappedApplication->MakeWindow(); }
-	virtual void InitializeWindow(const TSharedRef< FGenericWindow >& Window, const TSharedRef< FGenericWindowDefinition >& InDefinition, const TSharedPtr< FGenericWindow >& InParent, const bool bShowImmediately) { WrappedApplication->InitializeWindow(Window, InDefinition, InParent, bShowImmediately); }
-	virtual void SetCapture(const TSharedPtr< FGenericWindow >& InWindow) { WrappedApplication->SetCapture(InWindow); }
+	virtual TSharedRef<FGenericWindow> MakeWindow() { return WrappedApplication->MakeWindow(); }
+	virtual void InitializeWindow(const TSharedRef<FGenericWindow>& Window, const TSharedRef<FGenericWindowDefinition>& InDefinition, const TSharedPtr<FGenericWindow>& InParent, const bool bShowImmediately) { WrappedApplication->InitializeWindow(Window, InDefinition, InParent, bShowImmediately); }
+	virtual void SetCapture(const TSharedPtr<FGenericWindow>& InWindow) { WrappedApplication->SetCapture(InWindow); }
 	virtual void* GetCapture(void) const { return WrappedApplication->GetCapture(); }
 	virtual FModifierKeysState GetModifierKeys() const { return WrappedApplication->GetModifierKeys(); }
-	virtual TSharedPtr< FGenericWindow > GetWindowUnderCursor() { return WrappedApplication->GetWindowUnderCursor(); }
-	virtual void SetHighPrecisionMouseMode(const bool Enable, const TSharedPtr< FGenericWindow >& InWindow) { WrappedApplication->SetHighPrecisionMouseMode(Enable, InWindow); };
+	virtual TSharedPtr<FGenericWindow> GetWindowUnderCursor() { return WrappedApplication->GetWindowUnderCursor(); }
+	virtual void SetHighPrecisionMouseMode(const bool Enable, const TSharedPtr<FGenericWindow>& InWindow) { WrappedApplication->SetHighPrecisionMouseMode(Enable, InWindow); };
 	virtual bool IsUsingHighPrecisionMouseMode() const { return WrappedApplication->IsUsingHighPrecisionMouseMode(); }
 	virtual bool IsUsingTrackpad() const { return WrappedApplication->IsUsingTrackpad(); }
 	virtual bool IsGamepadAttached() const { return WrappedApplication->IsGamepadAttached(); }
@@ -173,367 +173,367 @@ void FInputDevice::Tick(float DeltaTime)
 	{
 		switch (Event.Event)
 		{
-		case EventType::UNDEFINED:
+			case EventType::UNDEFINED:
+			{
+				checkNoEntry();
+			}
+			break;
+			case EventType::KEY_DOWN:
+			{
+				uint8 JavaScriptKeyCode;
+				bool IsRepeat;
+				Event.GetKeyDown(JavaScriptKeyCode, IsRepeat);
+				const FKey* AgnosticKey = JavaScriptKeyCodeToFKey[JavaScriptKeyCode];
+				if (FilterKey(*AgnosticKey))
+				{
+					const uint32* KeyCodePtr;
+					const uint32* CharacterCodePtr;
+					FInputKeyManager::Get().GetCodesFromKey(*AgnosticKey, KeyCodePtr, CharacterCodePtr);
+					uint32 KeyCode = KeyCodePtr ? *KeyCodePtr : 0;
+					uint32 CharacterCode = CharacterCodePtr ? *CharacterCodePtr : 0;
+					MessageHandler->OnKeyDown(KeyCode, CharacterCode, IsRepeat);
+					UE_LOG(PixelStreamerInputDevice, Verbose, TEXT("KEY_DOWN: KeyCode = %d; CharacterCode = %d; IsRepeat = %s"), KeyCode, CharacterCode, IsRepeat ? TEXT("True") : TEXT("False"));
+				}
+			}
+			break;
+			case EventType::KEY_UP:
+			{
+				uint8 JavaScriptKeyCode;
+				Event.GetKeyUp(JavaScriptKeyCode);
+				const FKey* AgnosticKey = JavaScriptKeyCodeToFKey[JavaScriptKeyCode];
+				if (FilterKey(*AgnosticKey))
+				{
+					const uint32* KeyCodePtr;
+					const uint32* CharacterCodePtr;
+					FInputKeyManager::Get().GetCodesFromKey(*AgnosticKey, KeyCodePtr, CharacterCodePtr);
+					uint32 KeyCode = KeyCodePtr ? *KeyCodePtr : 0;
+					uint32 CharacterCode = CharacterCodePtr ? *CharacterCodePtr : 0;
+					MessageHandler->OnKeyUp(KeyCode, CharacterCode, false); // Key up events are never repeats.
+					UE_LOG(PixelStreamerInputDevice, Verbose, TEXT("KEY_UP: KeyCode = %d; CharacterCode = %d"), KeyCode, CharacterCode);
+				}
+			}
+			break;
+			case EventType::KEY_PRESS:
+			{
+				TCHAR UnicodeCharacter;
+				Event.GetCharacterCode(UnicodeCharacter);
+				MessageHandler->OnKeyChar(UnicodeCharacter, false); // Key press repeat not yet available but are not intrinsically used.
+				UE_LOG(PixelStreamerInputDevice, Verbose, TEXT("KEY_PRESSED: Character = '%c'"), UnicodeCharacter);
+			}
+			break;
+			case EventType::MOUSE_ENTER:
+			{
+				// Override application layer to special pixel streaming version.
+				FSlateApplication::Get().OverridePlatformApplication(PixelStreamerApplicationWrapper);
+				FSlateApplication::Get().OnCursorSet();
+
+				// Make sure the viewport is active.
+				FSlateApplication::Get().ProcessApplicationActivationEvent(true);
+
+				UE_LOG(PixelStreamerInputDevice, Verbose, TEXT("MOUSE_ENTER"));
+			}
+			break;
+			case EventType::MOUSE_LEAVE:
+			{
+				// Restore normal application layer.
+				FSlateApplication::Get().OverridePlatformApplication(PixelStreamerApplicationWrapper->WrappedApplication);
+
+				UE_LOG(PixelStreamerInputDevice, Verbose, TEXT("MOUSE_LEAVE"));
+			}
+			break;
+			case EventType::MOUSE_MOVE:
+			{
+				// Ensure the cursor will be updated so it can change appearance.
+				FSlateApplication::Get().OnCursorSet();
+
+				uint16 PosX;
+				uint16 PosY;
+				int16 DeltaX;
+				int16 DeltaY;
+				Event.GetMouseDelta(PosX, PosY, DeltaX, DeltaY);
+				FVector2D CursorPos = GEngine->GameViewport->GetWindow()->GetPositionInScreen() + FVector2D(PosX, PosY);
+				PixelStreamerApplicationWrapper->Cursor->SetPosition(CursorPos.X, CursorPos.Y);
+				MessageHandler->OnRawMouseMove(DeltaX, DeltaY);
+				UE_LOG(PixelStreamerInputDevice, VeryVerbose, TEXT("MOUSE_MOVE: Pos = (%d, %d); CursorPos = (%d, %d); Delta = (%d, %d)"), PosX, PosY, static_cast<int>(CursorPos.X), static_cast<int>(CursorPos.Y), DeltaX, DeltaY);
+			}
+			break;
+			case EventType::MOUSE_DOWN:
+			{
+				// If a user clicks on the application window and then clicks on the
+				// browser then this will move the focus away from the application
+				// window which will deactivate the application, so we need to check
+				// if we must reactivate the application.
+				if (!FSlateApplication::Get().IsActive())
+				{
+					FSlateApplication::Get().ProcessApplicationActivationEvent(true);
+				}
+
+				EMouseButtons::Type Button;
+				uint16 PosX;
+				uint16 PosY;
+				Event.GetMouseClick(Button, PosX, PosY);
+				FVector2D CursorPos = GEngine->GameViewport->GetWindow()->GetPositionInScreen() + FVector2D(PosX, PosY);
+				PixelStreamerApplicationWrapper->Cursor->SetPosition(CursorPos.X, CursorPos.Y);
+				MessageHandler->OnMouseDown(GEngine->GameViewport->GetWindow()->GetNativeWindow(), Button, CursorPos);
+				UE_LOG(PixelStreamerInputDevice, Verbose, TEXT("MOUSE_DOWN: Button = %d; Pos = (%d, %d); CursorPos = (%d, %d)"), Button, PosX, PosY, static_cast<int>(CursorPos.X), static_cast<int>(CursorPos.Y));
+
+				// The browser may be faking a mouse when touching so it will send
+				// over a mouse down event.
+				FindFocusedWidget();
+			}
+			break;
+			case EventType::MOUSE_UP:
+			{
+				EMouseButtons::Type Button;
+				uint16 PosX;
+				uint16 PosY;
+				Event.GetMouseClick(Button, PosX, PosY);
+				FVector2D CursorPos = GEngine->GameViewport->GetWindow()->GetPositionInScreen() + FVector2D(PosX, PosY);
+				PixelStreamerApplicationWrapper->Cursor->SetPosition(CursorPos.X, CursorPos.Y);
+				MessageHandler->OnMouseUp(Button);
+				UE_LOG(PixelStreamerInputDevice, Verbose, TEXT("MOUSE_UP: Button = %d; Pos = (%d, %d); CursorPos = (%d, %d)"), Button, PosX, PosY, static_cast<int>(CursorPos.X), static_cast<int>(CursorPos.Y));
+			}
+			break;
+			case EventType::MOUSE_WHEEL:
+			{
+				int16 Delta;
+				uint16 PosX;
+				uint16 PosY;
+				Event.GetMouseWheel(Delta, PosX, PosY);
+				const float SpinFactor = 1 / 120.0f;
+				FVector2D CursorPos = GEngine->GameViewport->GetWindow()->GetPositionInScreen() + FVector2D(PosX, PosY);
+				MessageHandler->OnMouseWheel(Delta * SpinFactor, CursorPos);
+				UE_LOG(PixelStreamerInputDevice, Verbose, TEXT("MOUSE_WHEEL: Delta = %d; Pos = (%d, %d); CursorPos = (%d, %d)"), Delta, PosX, PosY, static_cast<int>(CursorPos.X), static_cast<int>(CursorPos.Y));
+			}
+			break;
+			case EventType::TOUCH_START:
+			{
+				uint8 TouchIndex;
+				uint16 PosX;
+				uint16 PosY;
+				uint8 Force; // Force is between 0.0 and 1.0 so will need to unquantize from byte.
+				Event.GetTouch(TouchIndex, PosX, PosY, Force);
+				FVector2D CursorPos = GEngine->GameViewport->GetWindow()->GetPositionInScreen() + FVector2D(PosX, PosY);
+				MessageHandler->OnTouchStarted(GEngine->GameViewport->GetWindow()->GetNativeWindow(), CursorPos, Force / 255.0f, TouchIndex, 0); // TODO: ControllerId?
+				UE_LOG(PixelStreamerInputDevice, Verbose, TEXT("TOUCH_START: TouchIndex = %d; Pos = (%d, %d); CursorPos = (%d, %d); Force = %.3f"), TouchIndex, PosX, PosY, static_cast<int>(CursorPos.X), static_cast<int>(CursorPos.Y), Force / 255.0f);
+
+				// If the user starts a touch then they may be focusing an editable
+				// widget.
+				FindFocusedWidget();
+			}
+			break;
+			case EventType::TOUCH_END:
+			{
+				uint8 TouchIndex;
+				uint16 PosX;
+				uint16 PosY;
+				uint8 Force;
+				Event.GetTouch(TouchIndex, PosX, PosY, Force);
+				FVector2D CursorPos = GEngine->GameViewport->GetWindow()->GetPositionInScreen() + FVector2D(PosX, PosY);
+				MessageHandler->OnTouchEnded(CursorPos, TouchIndex, 0); // TODO: ControllerId?
+				UE_LOG(PixelStreamerInputDevice, Verbose, TEXT("TOUCH_END: TouchIndex = %d; Pos = (%d, %d); CursorPos = (%d, %d)"), TouchIndex, PosX, PosY, static_cast<int>(CursorPos.X), static_cast<int>(CursorPos.Y));
+			}
+			break;
+			case EventType::TOUCH_MOVE:
+			{
+				uint8 TouchIndex;
+				uint16 PosX;
+				uint16 PosY;
+				uint8 Force; // Force is between 0.0 and 1.0 so will need to unquantize from byte.
+				Event.GetTouch(TouchIndex, PosX, PosY, Force);
+				FVector2D CursorPos = GEngine->GameViewport->GetWindow()->GetPositionInScreen() + FVector2D(PosX, PosY);
+				MessageHandler->OnTouchMoved(CursorPos, Force / 255.0f, TouchIndex, 0); // TODO: ControllerId?
+				UE_LOG(PixelStreamerInputDevice, VeryVerbose, TEXT("TOUCH_MOVE: TouchIndex = %d; Pos = (%d, %d); CursorPos = (%d, %d); Force = %.3f"), TouchIndex, PosX, PosY, static_cast<int>(CursorPos.X), static_cast<int>(CursorPos.Y), Force / 255.0f);
+			}
+			break;
+			case EventType::GAMEPAD_PRESS:
+			{
+				uint8 ControllerId;
+				uint8 ButtonIndex;
+				bool bIsRepeat;
+				Event.GetGamepadButtonPressed(ControllerId, ButtonIndex, bIsRepeat);
+				FGamepadKeyNames::Type ControllerButton = ConvertButtonIndexToGamepadButton(ButtonIndex);
+				if (ControllerButton == FGamepadKeyNames::Invalid)
+				{
+					break;
+				}
+				MessageHandler->OnControllerButtonPressed(ControllerButton, (int32_t)ControllerId, bIsRepeat);
+				UE_LOG(PixelStreamerInputDevice, VeryVerbose, TEXT("GAMEPAD_PRESS: ControllerId = %d; ButtonIndex = %d; IsRepeat = %d"), ControllerId, ButtonIndex, bIsRepeat);
+			}
+			break;
+			case EventType::GAMEPAD_RELEASE:
+			{
+				uint8 ControllerId;
+				uint8 ButtonIndex;
+				Event.GetGamepadButtonReleased(ControllerId, ButtonIndex);
+				FGamepadKeyNames::Type ControllerButton = ConvertButtonIndexToGamepadButton(ButtonIndex);
+				if (ControllerButton == FGamepadKeyNames::Invalid)
+				{
+					break;
+				}
+				MessageHandler->OnControllerButtonReleased(ControllerButton, (int32_t)ControllerId, false);
+				UE_LOG(PixelStreamerInputDevice, VeryVerbose, TEXT("GAMEPAD_RELEASE: ControllerId = %d; ButtonIndex = %d;"), ControllerId, ButtonIndex);
+			}
+			break;
+			case EventType::GAMEPAD_ANALOG:
+			{
+				uint8 ControllerId;
+				uint8 AxisIndex;
+				float AnalogValue;
+				Event.GetGamepadAnalog(ControllerId, AxisIndex, AnalogValue);
+				FGamepadKeyNames::Type ControllerAxis = ConvertAxisIndexToGamepadAxis(AxisIndex);
+				if (ControllerAxis == FGamepadKeyNames::Invalid)
+				{
+					break;
+				}
+				MessageHandler->OnControllerAnalog(ControllerAxis, (int32_t)ControllerId, AnalogValue);
+				UE_LOG(PixelStreamerInputDevice, VeryVerbose, TEXT("GAMEPAD_ANALOG: ControllerId = %d; AxisIndex = %d; AnalogValue = %.4f;"), ControllerId, AxisIndex, AnalogValue);
+			}
+			break;
+			default:
+			{
+				UE_LOG(PixelStreamerInputDevice, Error, TEXT("Unknown Pixel Streaming event %d with word 0x%016llx"), static_cast<int>(Event.Event), Event.Data.Word);
+			}
+			break;
+		}
+	}
+
+	FString UIInteraction;
+	while (UIInteractions.Dequeue(UIInteraction))
+	{
+		for (UPixelStreamerInputComponent* InputComponent : this->PixelStreamingModule->GetInputComponents())
 		{
-			checkNoEntry();
+			InputComponent->OnInputEvent.Broadcast(UIInteraction);
+			UE_LOG(PixelStreamerInputDevice, Verbose, TEXT("UIInteraction = %s"), *UIInteraction);
 		}
-		break;
-		case EventType::KEY_DOWN:
+	}
+
+	FString Command;
+	while (Commands.Dequeue(Command))
+	{
+		UPixelStreamerInputComponent::OnCommand(Command);
+	}
+}
+
+FGamepadKeyNames::Type FInputDevice::ConvertAxisIndexToGamepadAxis(uint8 AnalogAxis)
+{
+	switch (AnalogAxis)
+	{
+		case 1:
 		{
-			uint8 JavaScriptKeyCode;
-			bool IsRepeat;
-			Event.GetKeyDown(JavaScriptKeyCode, IsRepeat);
-			const FKey* AgnosticKey = JavaScriptKeyCodeToFKey[JavaScriptKeyCode];
-			if (FilterKey(*AgnosticKey))
-			{
-				const uint32* KeyCodePtr;
-				const uint32* CharacterCodePtr;
-				FInputKeyManager::Get().GetCodesFromKey(*AgnosticKey, KeyCodePtr, CharacterCodePtr);
-				uint32 KeyCode = KeyCodePtr ? *KeyCodePtr : 0;
-				uint32 CharacterCode = CharacterCodePtr ? *CharacterCodePtr : 0;
-				MessageHandler->OnKeyDown(KeyCode, CharacterCode, IsRepeat);
-				UE_LOG(PixelStreamerInputDevice, Verbose, TEXT("KEY_DOWN: KeyCode = %d; CharacterCode = %d; IsRepeat = %s"), KeyCode, CharacterCode, IsRepeat ? TEXT("True") : TEXT("False"));
-			}
+			return FGamepadKeyNames::LeftAnalogX;
 		}
 		break;
-		case EventType::KEY_UP:
+		case 2:
 		{
-			uint8 JavaScriptKeyCode;
-			Event.GetKeyUp(JavaScriptKeyCode);
-			const FKey* AgnosticKey = JavaScriptKeyCodeToFKey[JavaScriptKeyCode];
-			if (FilterKey(*AgnosticKey))
-			{
-				const uint32* KeyCodePtr;
-				const uint32* CharacterCodePtr;
-				FInputKeyManager::Get().GetCodesFromKey(*AgnosticKey, KeyCodePtr, CharacterCodePtr);
-				uint32 KeyCode = KeyCodePtr ? *KeyCodePtr : 0;
-				uint32 CharacterCode = CharacterCodePtr ? *CharacterCodePtr : 0;
-				MessageHandler->OnKeyUp(KeyCode, CharacterCode, false);   // Key up events are never repeats.
-				UE_LOG(PixelStreamerInputDevice, Verbose, TEXT("KEY_UP: KeyCode = %d; CharacterCode = %d"), KeyCode, CharacterCode);
-			}
+			return FGamepadKeyNames::LeftAnalogY;
 		}
 		break;
-		case EventType::KEY_PRESS:
+		case 3:
 		{
-			TCHAR UnicodeCharacter;
-			Event.GetCharacterCode(UnicodeCharacter);
-			MessageHandler->OnKeyChar(UnicodeCharacter, false);   // Key press repeat not yet available but are not intrinsically used.
-			UE_LOG(PixelStreamerInputDevice, Verbose, TEXT("KEY_PRESSED: Character = '%c'"), UnicodeCharacter);
+			return FGamepadKeyNames::RightAnalogX;
 		}
 		break;
-		case EventType::MOUSE_ENTER:
+		case 4:
 		{
-			// Override application layer to special pixel streaming version.
-			FSlateApplication::Get().OverridePlatformApplication(PixelStreamerApplicationWrapper);
-			FSlateApplication::Get().OnCursorSet();
-
-			// Make sure the viewport is active.
-			FSlateApplication::Get().ProcessApplicationActivationEvent(true);
-
-			UE_LOG(PixelStreamerInputDevice, Verbose, TEXT("MOUSE_ENTER"));
+			return FGamepadKeyNames::RightAnalogY;
 		}
 		break;
-		case EventType::MOUSE_LEAVE:
+		case 5:
 		{
-			// Restore normal application layer.
-			FSlateApplication::Get().OverridePlatformApplication(PixelStreamerApplicationWrapper->WrappedApplication);
-
-			UE_LOG(PixelStreamerInputDevice, Verbose, TEXT("MOUSE_LEAVE"));
+			return FGamepadKeyNames::LeftTriggerAnalog;
 		}
 		break;
-		case EventType::MOUSE_MOVE:
+		case 6:
 		{
-			// Ensure the cursor will be updated so it can change appearance.
-			FSlateApplication::Get().OnCursorSet();
-
-			uint16 PosX;
-			uint16 PosY;
-			int16 DeltaX;
-			int16 DeltaY;
-			Event.GetMouseDelta(PosX, PosY, DeltaX, DeltaY);
-			FVector2D CursorPos = GEngine->GameViewport->GetWindow()->GetPositionInScreen() + FVector2D(PosX, PosY);
-			PixelStreamerApplicationWrapper->Cursor->SetPosition(CursorPos.X, CursorPos.Y);
-			MessageHandler->OnRawMouseMove(DeltaX, DeltaY);
-			UE_LOG(PixelStreamerInputDevice, VeryVerbose, TEXT("MOUSE_MOVE: Pos = (%d, %d); CursorPos = (%d, %d); Delta = (%d, %d)"), PosX, PosY, static_cast<int>(CursorPos.X), static_cast<int>(CursorPos.Y), DeltaX, DeltaY);
+			return FGamepadKeyNames::RightTriggerAnalog;
 		}
 		break;
-		case EventType::MOUSE_DOWN:
+		default:
 		{
-			// If a user clicks on the application window and then clicks on the
-			// browser then this will move the focus away from the application
-			// window which will deactivate the application, so we need to check
-			// if we must reactivate the application.
-			if (!FSlateApplication::Get().IsActive())
-			{
-				FSlateApplication::Get().ProcessApplicationActivationEvent(true);
-			}
-
-			EMouseButtons::Type Button;
-			uint16 PosX;
-			uint16 PosY;
-			Event.GetMouseClick(Button, PosX, PosY);
-			FVector2D CursorPos = GEngine->GameViewport->GetWindow()->GetPositionInScreen() + FVector2D(PosX, PosY);
-			PixelStreamerApplicationWrapper->Cursor->SetPosition(CursorPos.X, CursorPos.Y);
-			MessageHandler->OnMouseDown(GEngine->GameViewport->GetWindow()->GetNativeWindow(), Button, CursorPos);
-			UE_LOG(PixelStreamerInputDevice, Verbose, TEXT("MOUSE_DOWN: Button = %d; Pos = (%d, %d); CursorPos = (%d, %d)"), Button, PosX, PosY, static_cast<int>(CursorPos.X), static_cast<int>(CursorPos.Y));
+			return FGamepadKeyNames::Invalid;
+		}
+		break;
+	}
+}
 
-			// The browser may be faking a mouse when touching so it will send
-			// over a mouse down event.
-			FindFocusedWidget();
+FGamepadKeyNames::Type FInputDevice::ConvertButtonIndexToGamepadButton(uint8 ButtonIndex)
+{
+	switch (ButtonIndex)
+	{
+		case 0:
+		{
+			return FGamepadKeyNames::FaceButtonBottom;
+		}
+		case 1:
+		{
+			return FGamepadKeyNames::FaceButtonRight;
 		}
 		break;
-		case EventType::MOUSE_UP:
+		case 2:
 		{
-			EMouseButtons::Type Button;
-			uint16 PosX;
-			uint16 PosY;
-			Event.GetMouseClick(Button, PosX, PosY);
-			FVector2D CursorPos = GEngine->GameViewport->GetWindow()->GetPositionInScreen() + FVector2D(PosX, PosY);
-			PixelStreamerApplicationWrapper->Cursor->SetPosition(CursorPos.X, CursorPos.Y);
-			MessageHandler->OnMouseUp(Button);
-			UE_LOG(PixelStreamerInputDevice, Verbose, TEXT("MOUSE_UP: Button = %d; Pos = (%d, %d); CursorPos = (%d, %d)"), Button, PosX, PosY, static_cast<int>(CursorPos.X), static_cast<int>(CursorPos.Y));
+			return FGamepadKeyNames::FaceButtonLeft;
 		}
 		break;
-		case EventType::MOUSE_WHEEL:
+		case 3:
 		{
-			int16 Delta;
-			uint16 PosX;
-			uint16 PosY;
-			Event.GetMouseWheel(Delta, PosX, PosY);
-			const float SpinFactor = 1 / 120.0f;
-			FVector2D CursorPos = GEngine->GameViewport->GetWindow()->GetPositionInScreen() + FVector2D(PosX, PosY);
-			MessageHandler->OnMouseWheel(Delta * SpinFactor, CursorPos);
-			UE_LOG(PixelStreamerInputDevice, Verbose, TEXT("MOUSE_WHEEL: Delta = %d; Pos = (%d, %d); CursorPos = (%d, %d)"), Delta, PosX, PosY, static_cast<int>(CursorPos.X), static_cast<int>(CursorPos.Y));
+			return FGamepadKeyNames::FaceButtonTop;
 		}
 		break;
-		case EventType::TOUCH_START:
+		case 4:
 		{
-			uint8 TouchIndex;
-			uint16 PosX;
-			uint16 PosY;
-			uint8 Force;   // Force is between 0.0 and 1.0 so will need to unquantize from byte.
-			Event.GetTouch(TouchIndex, PosX, PosY, Force);
-			FVector2D CursorPos = GEngine->GameViewport->GetWindow()->GetPositionInScreen() + FVector2D(PosX, PosY);
-			MessageHandler->OnTouchStarted(GEngine->GameViewport->GetWindow()->GetNativeWindow(), CursorPos, Force / 255.0f, TouchIndex, 0);   // TODO: ControllerId?
-			UE_LOG(PixelStreamerInputDevice, Verbose, TEXT("TOUCH_START: TouchIndex = %d; Pos = (%d, %d); CursorPos = (%d, %d); Force = %.3f"), TouchIndex, PosX, PosY, static_cast<int>(CursorPos.X), static_cast<int>(CursorPos.Y), Force / 255.0f);
-
-			// If the user starts a touch then they may be focusing an editable
-			// widget.
-			FindFocusedWidget();
+			return FGamepadKeyNames::LeftShoulder;
 		}
 		break;
-		case EventType::TOUCH_END:
+		case 5:
 		{
-			uint8 TouchIndex;
-			uint16 PosX;
-			uint16 PosY;
-			uint8 Force;
-			Event.GetTouch(TouchIndex, PosX, PosY, Force);
-			FVector2D CursorPos = GEngine->GameViewport->GetWindow()->GetPositionInScreen() + FVector2D(PosX, PosY);
-			MessageHandler->OnTouchEnded(CursorPos, TouchIndex, 0);   // TODO: ControllerId?
-			UE_LOG(PixelStreamerInputDevice, Verbose, TEXT("TOUCH_END: TouchIndex = %d; Pos = (%d, %d); CursorPos = (%d, %d)"), TouchIndex, PosX, PosY, static_cast<int>(CursorPos.X), static_cast<int>(CursorPos.Y));
+			return FGamepadKeyNames::RightShoulder;
 		}
 		break;
-		case EventType::TOUCH_MOVE:
+		// Buttons 6 and 7 are mapped as analog axis as they are the triggers
+		case 8:
 		{
-			uint8 TouchIndex;
-			uint16 PosX;
-			uint16 PosY;
-			uint8 Force;   // Force is between 0.0 and 1.0 so will need to unquantize from byte.
-			Event.GetTouch(TouchIndex, PosX, PosY, Force);
-			FVector2D CursorPos = GEngine->GameViewport->GetWindow()->GetPositionInScreen() + FVector2D(PosX, PosY);
-			MessageHandler->OnTouchMoved(CursorPos, Force / 255.0f, TouchIndex, 0);   // TODO: ControllerId?
-			UE_LOG(PixelStreamerInputDevice, VeryVerbose, TEXT("TOUCH_MOVE: TouchIndex = %d; Pos = (%d, %d); CursorPos = (%d, %d); Force = %.3f"), TouchIndex, PosX, PosY, static_cast<int>(CursorPos.X), static_cast<int>(CursorPos.Y), Force / 255.0f);
+			return FGamepadKeyNames::SpecialLeft;
 		}
 		break;
-		case EventType::GAMEPAD_PRESS:
-		{
-			uint8 ControllerId;
-			uint8 ButtonIndex;
-			bool bIsRepeat;
-			Event.GetGamepadButtonPressed(ControllerId, ButtonIndex, bIsRepeat);
-			FGamepadKeyNames::Type ControllerButton = ConvertButtonIndexToGamepadButton(ButtonIndex);
-			if (ControllerButton == FGamepadKeyNames::Invalid)
-			{
-				break;
-			}
-			MessageHandler->OnControllerButtonPressed(ControllerButton, (int32_t)ControllerId, bIsRepeat);
-			UE_LOG(PixelStreamerInputDevice, VeryVerbose, TEXT("GAMEPAD_PRESS: ControllerId = %d; ButtonIndex = %d; IsRepeat = %d"), ControllerId, ButtonIndex, bIsRepeat);
+		case 9:
+		{
+			return FGamepadKeyNames::SpecialRight;
+		}
+		case 10:
+		{
+			return FGamepadKeyNames::LeftThumb;
 		}
 		break;
-		case EventType::GAMEPAD_RELEASE:
+		case 11:
 		{
-			uint8 ControllerId;
-			uint8 ButtonIndex;
-			Event.GetGamepadButtonReleased(ControllerId, ButtonIndex);
-			FGamepadKeyNames::Type ControllerButton = ConvertButtonIndexToGamepadButton(ButtonIndex);
-			if (ControllerButton == FGamepadKeyNames::Invalid)
-			{
-				break;
-			}
-			MessageHandler->OnControllerButtonReleased(ControllerButton, (int32_t)ControllerId, false);
-			UE_LOG(PixelStreamerInputDevice, VeryVerbose, TEXT("GAMEPAD_RELEASE: ControllerId = %d; ButtonIndex = %d;"), ControllerId, ButtonIndex);
+			return FGamepadKeyNames::RightThumb;
 		}
 		break;
-		case EventType::GAMEPAD_ANALOG:
-		{
-			uint8 ControllerId;
-			uint8 AxisIndex;
-			float AnalogValue;
-			Event.GetGamepadAnalog(ControllerId, AxisIndex, AnalogValue);
-			FGamepadKeyNames::Type ControllerAxis = ConvertAxisIndexToGamepadAxis(AxisIndex);
-			if (ControllerAxis == FGamepadKeyNames::Invalid)
-			{
-				break;
-			}
-			MessageHandler->OnControllerAnalog(ControllerAxis, (int32_t)ControllerId, AnalogValue);
-			UE_LOG(PixelStreamerInputDevice, VeryVerbose, TEXT("GAMEPAD_ANALOG: ControllerId = %d; AxisIndex = %d; AnalogValue = %.4f;"), ControllerId, AxisIndex, AnalogValue);
+		case 12:
+		{
+			return FGamepadKeyNames::DPadUp;
 		}
 		break;
-		default:
+		case 13:
+		{
+			return FGamepadKeyNames::DPadDown;
+		}
+		case 14:
 		{
-			UE_LOG(PixelStreamerInputDevice, Error, TEXT("Unknown Pixel Streaming event %d with word 0x%016llx"), static_cast<int>(Event.Event), Event.Data.Word);
+			return FGamepadKeyNames::DPadLeft;
 		}
 		break;
+		case 15:
+		{
+			return FGamepadKeyNames::DPadRight;
 		}
-	}
-
-	FString UIInteraction;
-	while (UIInteractions.Dequeue(UIInteraction))
-	{
-		for (UPixelStreamerInputComponent* InputComponent : this->PixelStreamingModule->GetInputComponents())
+		break;
+		default:
 		{
-			InputComponent->OnInputEvent.Broadcast(UIInteraction);
-			UE_LOG(PixelStreamerInputDevice, Verbose, TEXT("UIInteraction = %s"), *UIInteraction);
+			return FGamepadKeyNames::Invalid;
 		}
-	}
-
-	FString Command;
-	while (Commands.Dequeue(Command))
-	{
-		UPixelStreamerInputComponent::OnCommand(Command);
-	}
-}
-
-FGamepadKeyNames::Type FInputDevice::ConvertAxisIndexToGamepadAxis(uint8 AnalogAxis)
-{
-	switch (AnalogAxis)
-	{
-	case 1:
-	{
-		return FGamepadKeyNames::LeftAnalogX;
-	}
-	break;
-	case 2:
-	{
-		return FGamepadKeyNames::LeftAnalogY;
-	}
-	break;
-	case 3:
-	{
-		return FGamepadKeyNames::RightAnalogX;
-	}
-	break;
-	case 4:
-	{
-		return FGamepadKeyNames::RightAnalogY;
-	}
-	break;
-	case 5:
-	{
-		return FGamepadKeyNames::LeftTriggerAnalog;
-	}
-	break;
-	case 6:
-	{
-		return FGamepadKeyNames::RightTriggerAnalog;
-	}
-	break;
-	default:
-	{
-		return FGamepadKeyNames::Invalid;
-	}
-	break;
-	}
-}
-
-FGamepadKeyNames::Type FInputDevice::ConvertButtonIndexToGamepadButton(uint8 ButtonIndex)
-{
-	switch (ButtonIndex)
-	{
-	case 0:
-	{
-		return FGamepadKeyNames::FaceButtonBottom;
-	}
-	case 1:
-	{
-		return FGamepadKeyNames::FaceButtonRight;
-	}
-	break;
-	case 2:
-	{
-		return FGamepadKeyNames::FaceButtonLeft;
-	}
-	break;
-	case 3:
-	{
-		return FGamepadKeyNames::FaceButtonTop;
-	}
-	break;
-	case 4:
-	{
-		return FGamepadKeyNames::LeftShoulder;
-	}
-	break;
-	case 5:
-	{
-		return FGamepadKeyNames::RightShoulder;
-	}
-	break;
-	// Buttons 6 and 7 are mapped as analog axis as they are the triggers
-	case 8:
-	{
-		return FGamepadKeyNames::SpecialLeft;
-	}
-	break;
-	case 9:
-	{
-		return FGamepadKeyNames::SpecialRight;
-	}
-	case 10:
-	{
-		return FGamepadKeyNames::LeftThumb;
-	}
-	break;
-	case 11:
-	{
-		return FGamepadKeyNames::RightThumb;
-	}
-	break;
-	case 12:
-	{
-		return FGamepadKeyNames::DPadUp;
-	}
-	break;
-	case 13:
-	{
-		return FGamepadKeyNames::DPadDown;
-	}
-	case 14:
-	{
-		return FGamepadKeyNames::DPadLeft;
-	}
-	break;
-	case 15:
-	{
-		return FGamepadKeyNames::DPadRight;
-	}
-	break;
-	default:
-	{
-		return FGamepadKeyNames::Invalid;
-	}
-	break;
+		break;
 	}
 }
 
@@ -555,7 +555,7 @@ void FInputDevice::SetChannelValue(int32 ControllerId, FForceFeedbackChannelType
 {
 }
 
-void FInputDevice::SetChannelValues(int32 ControllerId, const FForceFeedbackValues &values)
+void FInputDevice::SetChannelValues(int32 ControllerId, const FForceFeedbackValues& values)
 {
 }
 
@@ -584,7 +584,7 @@ void FInputDevice::FindFocusedWidget()
 {
 	FSlateApplication::Get().ForEachUser([&](FSlateUser& User) {
 		TSharedPtr<SWidget> FocusedWidget = User.GetFocusedWidget();
-		
+
 		static FName SEditableTextType(TEXT("SEditableText"));
 		static FName SMultiLineEditableTextType(TEXT("SMultiLineEditableText"));
 		bool bEditable = FocusedWidget && (FocusedWidget->GetType() == SEditableTextType || FocusedWidget->GetType() == SMultiLineEditableTextType);
@@ -599,7 +599,7 @@ void FInputDevice::FindFocusedWidget()
 			TSharedPtr<FJsonObject> JsonObject = MakeShareable(new FJsonObject);
 			JsonObject->SetStringField(TEXT("command"), TEXT("onScreenKeyboard"));
 			JsonObject->SetBoolField(TEXT("showOnScreenKeyboard"), bEditable);
-			
+
 			if (bEditable)
 			{
 				Pos = Pos - GEngine->GameViewport->GetWindow()->GetPositionInScreen();
@@ -652,11 +652,11 @@ namespace
 	*/
 	struct FTouch
 	{
-		uint16 PosX;   // X position of finger.
-		uint16 PosY;   // Y position of finger.
-		uint8 TouchIndex;   // Index of finger for tracking multi-touch events.
-		uint8 Force;   // Amount of pressure being applied by the finger.
-		uint8 Valid;	// 1 if the touch was within bounds.
+		uint16 PosX;	  // X position of finger.
+		uint16 PosY;	  // Y position of finger.
+		uint8 TouchIndex; // Index of finger for tracking multi-touch events.
+		uint8 Force;	  // Amount of pressure being applied by the finger.
+		uint8 Valid;	  // 1 if the touch was within bounds.
 	};
 
 	using FKeyCodeType = uint8;
@@ -714,9 +714,24 @@ namespace
 		return String;
 	}
 
-	enum class KeyState { Alt = 1 << 0,	Ctrl = 1 << 1, Shift = 1 << 2 };
-	enum class MouseButtonState	{ Left = 1 << 0, Right = 1 << 1, Middle = 1 << 2, Button4 = 1 << 3, Button5 = 1 << 4, Button6 = 1 << 5, Button7 = 1 << 6, Button8 = 1 << 7 };
-}
+	enum class KeyState
+	{
+		Alt = 1 << 0,
+		Ctrl = 1 << 1,
+		Shift = 1 << 2
+	};
+	enum class MouseButtonState
+	{
+		Left = 1 << 0,
+		Right = 1 << 1,
+		Middle = 1 << 2,
+		Button4 = 1 << 3,
+		Button5 = 1 << 4,
+		Button6 = 1 << 5,
+		Button7 = 1 << 6,
+		Button8 = 1 << 7
+	};
+} // namespace
 
 void FInputDevice::OnMessage(const webrtc::DataBuffer& Buffer)
 {
@@ -729,218 +744,218 @@ void FInputDevice::OnMessage(const webrtc::DataBuffer& Buffer)
 
 	switch (MsgType)
 	{
-	case EToStreamerMsg::UIInteraction:
-	{
-		FString Descriptor = PixelStreamingProtocol::ParseString(Buffer, FInputDevice::GetMessageHeaderOffset());
-		UE_LOG(PixelStreamerInput, Verbose, TEXT("UIInteraction: %s"), *Descriptor);
-		ProcessUIInteraction(Descriptor);
-		break;
-	}
-	case EToStreamerMsg::Command:
-	{
-		FString Descriptor = PixelStreamingProtocol::ParseString(Buffer, FInputDevice::GetMessageHeaderOffset());
-		UE_LOG(PixelStreamerInput, Verbose, TEXT("Command: %s"), *Descriptor);
-		ProcessCommand(Descriptor);
-		break;
-	}
-	case EToStreamerMsg::KeyDown:
-	{
-		GET(FKeyCodeType, KeyCode);
-		GET(FRepeatType, Repeat);
-		checkf(Size == 0, TEXT("%d"), Size);
-		UE_LOG(PixelStreamerInput, Verbose, TEXT("key down: %d, repeat: %d"), KeyCode, Repeat);
-
-		FEvent KeyDownEvent(EventType::KEY_DOWN);
-		KeyDownEvent.SetKeyDown(KeyCode, Repeat != 0);
-		ProcessEvent(KeyDownEvent);
-		break;
-	}
-	case EToStreamerMsg::KeyUp:
-	{
-		GET(FKeyCodeType, KeyCode);
-		checkf(Size == 0, TEXT("%d"), Size);
-		UE_LOG(PixelStreamerInput, Verbose, TEXT("key up: %d"), KeyCode);
-
-		FEvent KeyUpEvent(EventType::KEY_UP);
-		KeyUpEvent.SetKeyUp(KeyCode);
-		ProcessEvent(KeyUpEvent);
-		break;
-	}
-	case EToStreamerMsg::KeyPress:
-	{
-		GET(FCharacterType, Character);
-		checkf(Size == 0, TEXT("%d"), Size);
-		UE_LOG(PixelStreamerInput, Verbose, TEXT("key press: '%c'"), Character);
+		case EToStreamerMsg::UIInteraction:
+		{
+			FString Descriptor = PixelStreamingProtocol::ParseString(Buffer, FInputDevice::GetMessageHeaderOffset());
+			UE_LOG(PixelStreamerInput, Verbose, TEXT("UIInteraction: %s"), *Descriptor);
+			ProcessUIInteraction(Descriptor);
+			break;
+		}
+		case EToStreamerMsg::Command:
+		{
+			FString Descriptor = PixelStreamingProtocol::ParseString(Buffer, FInputDevice::GetMessageHeaderOffset());
+			UE_LOG(PixelStreamerInput, Verbose, TEXT("Command: %s"), *Descriptor);
+			ProcessCommand(Descriptor);
+			break;
+		}
+		case EToStreamerMsg::KeyDown:
+		{
+			GET(FKeyCodeType, KeyCode);
+			GET(FRepeatType, Repeat);
+			checkf(Size == 0, TEXT("%d"), Size);
+			UE_LOG(PixelStreamerInput, Verbose, TEXT("key down: %d, repeat: %d"), KeyCode, Repeat);
+
+			FEvent KeyDownEvent(EventType::KEY_DOWN);
+			KeyDownEvent.SetKeyDown(KeyCode, Repeat != 0);
+			ProcessEvent(KeyDownEvent);
+			break;
+		}
+		case EToStreamerMsg::KeyUp:
+		{
+			GET(FKeyCodeType, KeyCode);
+			checkf(Size == 0, TEXT("%d"), Size);
+			UE_LOG(PixelStreamerInput, Verbose, TEXT("key up: %d"), KeyCode);
+
+			FEvent KeyUpEvent(EventType::KEY_UP);
+			KeyUpEvent.SetKeyUp(KeyCode);
+			ProcessEvent(KeyUpEvent);
+			break;
+		}
+		case EToStreamerMsg::KeyPress:
+		{
+			GET(FCharacterType, Character);
+			checkf(Size == 0, TEXT("%d"), Size);
+			UE_LOG(PixelStreamerInput, Verbose, TEXT("key press: '%c'"), Character);
+
+			FEvent KeyPressEvent(EventType::KEY_PRESS);
+			KeyPressEvent.SetCharCode(Character);
+			ProcessEvent(KeyPressEvent);
+			break;
+		}
+		case EToStreamerMsg::MouseEnter:
+		{
+			checkf(Size == 0, TEXT("%d"), Size);
+			UE_LOG(PixelStreamerInput, Verbose, TEXT("mouseEnter"));
+			ProcessEvent(FEvent(FInputDevice::EventType::MOUSE_ENTER));
+			break;
+		}
+		case EToStreamerMsg::MouseLeave:
+		{
+			checkf(Size == 0, TEXT("%d"), Size);
+			ProcessEvent(FEvent(FInputDevice::EventType::MOUSE_LEAVE));
+			UE_LOG(PixelStreamerInput, Verbose, TEXT("mouseLeave"));
+			break;
+		}
+		case EToStreamerMsg::MouseDown:
+		{
+			GET(FButtonType, Button);
+			GET(FPosType, PosX);
+			GET(FPosType, PosY);
+			checkf(Size == 0, TEXT("%d"), Size);
+			UE_LOG(PixelStreamerInput, Verbose, TEXT("mouseDown at (%d, %d), button %d"), PosX, PosY, Button);
 
-		FEvent KeyPressEvent(EventType::KEY_PRESS);
-		KeyPressEvent.SetCharCode(Character);
-		ProcessEvent(KeyPressEvent);
-		break;
-	}
-	case EToStreamerMsg::MouseEnter:
-	{
-		checkf(Size == 0, TEXT("%d"), Size);
-		UE_LOG(PixelStreamerInput, Verbose, TEXT("mouseEnter"));
-		ProcessEvent(FEvent(FInputDevice::EventType::MOUSE_ENTER));
-		break;
-	}
-	case EToStreamerMsg::MouseLeave:
-	{
-		checkf(Size == 0, TEXT("%d"), Size);
-		ProcessEvent(FEvent(FInputDevice::EventType::MOUSE_LEAVE));
-		UE_LOG(PixelStreamerInput, Verbose, TEXT("mouseLeave"));
-		break;
-	}
-	case EToStreamerMsg::MouseDown:
-	{
-		GET(FButtonType, Button);
-		GET(FPosType, PosX);
-		GET(FPosType, PosY);
-		checkf(Size == 0, TEXT("%d"), Size);
-		UE_LOG(PixelStreamerInput, Verbose, TEXT("mouseDown at (%d, %d), button %d"), PosX, PosY, Button);
+			UnquantizeAndDenormalize(PosX, PosY);
 
-		UnquantizeAndDenormalize(PosX, PosY);
+			FEvent MouseDownEvent(EventType::MOUSE_DOWN);
+			MouseDownEvent.SetMouseClick(Button, PosX, PosY);
+			ProcessEvent(MouseDownEvent);
+			break;
+		}
+		case EToStreamerMsg::MouseUp:
+		{
+			GET(FButtonType, Button);
+			GET(FPosType, PosX);
+			GET(FPosType, PosY);
+			checkf(Size == 0, TEXT("%d"), Size);
+			UE_LOG(PixelStreamerInput, Verbose, TEXT("mouseUp at (%d, %d), button %d"), PosX, PosY, Button);
 
-		FEvent MouseDownEvent(EventType::MOUSE_DOWN);
-		MouseDownEvent.SetMouseClick(Button, PosX, PosY);
-		ProcessEvent(MouseDownEvent);
-		break;
-	}
-	case EToStreamerMsg::MouseUp:
-	{
-		GET(FButtonType, Button);
-		GET(FPosType, PosX);
-		GET(FPosType, PosY);
-		checkf(Size == 0, TEXT("%d"), Size);
-		UE_LOG(PixelStreamerInput, Verbose, TEXT("mouseUp at (%d, %d), button %d"), PosX, PosY, Button);
+			UnquantizeAndDenormalize(PosX, PosY);
 
-		UnquantizeAndDenormalize(PosX, PosY);
+			FEvent MouseDownEvent(EventType::MOUSE_UP);
+			MouseDownEvent.SetMouseClick(Button, PosX, PosY);
+			ProcessEvent(MouseDownEvent);
+			break;
+		}
+		case EToStreamerMsg::MouseMove:
+		{
+			GET(FPosType, PosX);
+			GET(FPosType, PosY);
+			GET(FDeltaType, DeltaX);
+			GET(FDeltaType, DeltaY);
+			checkf(Size == 0, TEXT("%d"), Size);
+			UE_LOG(PixelStreamerInput, Verbose, TEXT("mouseMove to (%d, %d), delta (%d, %d)"), PosX, PosY, DeltaX, DeltaY);
 
-		FEvent MouseDownEvent(EventType::MOUSE_UP);
-		MouseDownEvent.SetMouseClick(Button, PosX, PosY);
-		ProcessEvent(MouseDownEvent);
-		break;
-	}
-	case EToStreamerMsg::MouseMove:
-	{
-		GET(FPosType, PosX);
-		GET(FPosType, PosY);
-		GET(FDeltaType, DeltaX);
-		GET(FDeltaType, DeltaY);
-		checkf(Size == 0, TEXT("%d"), Size);
-		UE_LOG(PixelStreamerInput, Verbose, TEXT("mouseMove to (%d, %d), delta (%d, %d)"), PosX, PosY, DeltaX, DeltaY);
-
-		UnquantizeAndDenormalize(PosX, PosY);
-		UnquantizeAndDenormalize(DeltaX, DeltaY);
-
-		FEvent MouseMoveEvent(EventType::MOUSE_MOVE);
-		MouseMoveEvent.SetMouseDelta(PosX, PosY, DeltaX, DeltaY);
-		ProcessEvent(MouseMoveEvent);
-		break;
-	}
-	case EToStreamerMsg::MouseWheel:
-	{
-		GET(FDeltaType, Delta);
-		GET(FPosType, PosX);
-		GET(FPosType, PosY);
-		checkf(Size == 0, TEXT("%d"), Size);
-		UE_LOG(PixelStreamerInput, Verbose, TEXT("mouseWheel, delta %d"), Delta);
+			UnquantizeAndDenormalize(PosX, PosY);
+			UnquantizeAndDenormalize(DeltaX, DeltaY);
 
-		UnquantizeAndDenormalize(PosX, PosY);
+			FEvent MouseMoveEvent(EventType::MOUSE_MOVE);
+			MouseMoveEvent.SetMouseDelta(PosX, PosY, DeltaX, DeltaY);
+			ProcessEvent(MouseMoveEvent);
+			break;
+		}
+		case EToStreamerMsg::MouseWheel:
+		{
+			GET(FDeltaType, Delta);
+			GET(FPosType, PosX);
+			GET(FPosType, PosY);
+			checkf(Size == 0, TEXT("%d"), Size);
+			UE_LOG(PixelStreamerInput, Verbose, TEXT("mouseWheel, delta %d"), Delta);
 
-		FEvent MouseWheelEvent(EventType::MOUSE_WHEEL);
-		MouseWheelEvent.SetMouseWheel(Delta, PosX, PosY);
-		ProcessEvent(MouseWheelEvent);
-		break;
-	}
-	case EToStreamerMsg::TouchStart:
-	{
-		FTouchesType Touches = GetTouches(Data, Size);
-		checkf(Size == 0, TEXT("%d"), Size);
-		UE_LOG(PixelStreamerInput, Verbose, TEXT("TouchStart: %s"), *TouchesToString(Touches));
+			UnquantizeAndDenormalize(PosX, PosY);
 
-		for (const FTouch& Touch : Touches)
+			FEvent MouseWheelEvent(EventType::MOUSE_WHEEL);
+			MouseWheelEvent.SetMouseWheel(Delta, PosX, PosY);
+			ProcessEvent(MouseWheelEvent);
+			break;
+		}
+		case EToStreamerMsg::TouchStart:
 		{
-			if (Touch.Valid)
+			FTouchesType Touches = GetTouches(Data, Size);
+			checkf(Size == 0, TEXT("%d"), Size);
+			UE_LOG(PixelStreamerInput, Verbose, TEXT("TouchStart: %s"), *TouchesToString(Touches));
+
+			for (const FTouch& Touch : Touches)
 			{
-				FEvent TouchStartEvent(EventType::TOUCH_START);
-				TouchStartEvent.SetTouch(Touch.TouchIndex, Touch.PosX, Touch.PosY, Touch.Force);
-				ProcessEvent(TouchStartEvent);
+				if (Touch.Valid)
+				{
+					FEvent TouchStartEvent(EventType::TOUCH_START);
+					TouchStartEvent.SetTouch(Touch.TouchIndex, Touch.PosX, Touch.PosY, Touch.Force);
+					ProcessEvent(TouchStartEvent);
+				}
 			}
+			break;
 		}
-		break;
-	}
-	case EToStreamerMsg::TouchEnd:
-	{
-		FTouchesType Touches = GetTouches(Data, Size);
-		checkf(Size == 0, TEXT("%d"), Size);
-		UE_LOG(PixelStreamerInput, Verbose, TEXT("TouchEnd: %s"), *TouchesToString(Touches));
-
-		for (const FTouch& Touch : Touches)
+		case EToStreamerMsg::TouchEnd:
 		{
-			// Always allowing the "up" events regardless of in or outside the valid region so
-			// states aren't stuck "down". Might want to uncomment this if it causes other issues.
-			//if (Touch.Valid)
+			FTouchesType Touches = GetTouches(Data, Size);
+			checkf(Size == 0, TEXT("%d"), Size);
+			UE_LOG(PixelStreamerInput, Verbose, TEXT("TouchEnd: %s"), *TouchesToString(Touches));
+
+			for (const FTouch& Touch : Touches)
 			{
-				FEvent TouchEndEvent(EventType::TOUCH_END);
-				TouchEndEvent.SetTouch(Touch.TouchIndex, Touch.PosX, Touch.PosY, Touch.Force);
-				ProcessEvent(TouchEndEvent);
+				// Always allowing the "up" events regardless of in or outside the valid region so
+				// states aren't stuck "down". Might want to uncomment this if it causes other issues.
+				//if (Touch.Valid)
+				{
+					FEvent TouchEndEvent(EventType::TOUCH_END);
+					TouchEndEvent.SetTouch(Touch.TouchIndex, Touch.PosX, Touch.PosY, Touch.Force);
+					ProcessEvent(TouchEndEvent);
+				}
 			}
+			break;
 		}
-		break;
-	}
-	case EToStreamerMsg::TouchMove:
-	{
-		FTouchesType Touches = GetTouches(Data, Size);
-		checkf(Size == 0, TEXT("%d"), Size);
-		UE_LOG(PixelStreamerInput, Verbose, TEXT("TouchMove: %s"), *TouchesToString(Touches));
-
-		for (const FTouch& Touch : Touches)
+		case EToStreamerMsg::TouchMove:
 		{
-			if (Touch.Valid)
+			FTouchesType Touches = GetTouches(Data, Size);
+			checkf(Size == 0, TEXT("%d"), Size);
+			UE_LOG(PixelStreamerInput, Verbose, TEXT("TouchMove: %s"), *TouchesToString(Touches));
+
+			for (const FTouch& Touch : Touches)
 			{
-				FEvent TouchMoveEvent(EventType::TOUCH_MOVE);
-				TouchMoveEvent.SetTouch(Touch.TouchIndex, Touch.PosX, Touch.PosY, Touch.Force);
-				ProcessEvent(TouchMoveEvent);
+				if (Touch.Valid)
+				{
+					FEvent TouchMoveEvent(EventType::TOUCH_MOVE);
+					TouchMoveEvent.SetTouch(Touch.TouchIndex, Touch.PosX, Touch.PosY, Touch.Force);
+					ProcessEvent(TouchMoveEvent);
+				}
 			}
+			break;
+		}
+		case EToStreamerMsg::GamepadButtonPressed:
+		{
+			GET(FControllerIndex, ControllerIndex);
+			GET(FControllerButtonIndex, ButtonIndex);
+			GET(FRepeatType, Repeat);
+			checkf(Size == 0, TEXT("%d"), Size);
+
+			FEvent GamepadAnalogEvent(EventType::GAMEPAD_PRESS);
+			GamepadAnalogEvent.SetGamepadButtonPressed(ControllerIndex, ButtonIndex, Repeat != 0);
+			ProcessEvent(GamepadAnalogEvent);
+			break;
+		}
+		case EToStreamerMsg::GamepadButtonReleased:
+		{
+			GET(FControllerIndex, ControllerIndex);
+			GET(FControllerButtonIndex, ButtonIndex);
+			checkf(Size == 0, TEXT("%d"), Size);
+
+			FEvent GamepadAnalogEvent(EventType::GAMEPAD_RELEASE);
+			GamepadAnalogEvent.SetGamepadButtonReleased(ControllerIndex, ButtonIndex);
+			ProcessEvent(GamepadAnalogEvent);
+			break;
+		}
+		case EToStreamerMsg::GamepadAnalog:
+		{
+			GET(FControllerIndex, ControllerIndex);
+			GET(FControllerAxis, AxisIndex);
+			GET(FControllerAnalog, AnalogValue);
+			checkf(Size == 0, TEXT("%d"), Size);
+
+			FEvent GamepadAnalogEvent(EventType::GAMEPAD_ANALOG);
+			GamepadAnalogEvent.SetGamepadAnalog(ControllerIndex, AxisIndex, (float)AnalogValue);
+			ProcessEvent(GamepadAnalogEvent);
+			break;
 		}
-		break;
-	}
-	case EToStreamerMsg::GamepadButtonPressed:
-	{
-		GET(FControllerIndex, ControllerIndex);
-		GET(FControllerButtonIndex, ButtonIndex);
-		GET(FRepeatType, Repeat);
-		checkf(Size == 0, TEXT("%d"), Size);
-
-		FEvent GamepadAnalogEvent(EventType::GAMEPAD_PRESS);
-		GamepadAnalogEvent.SetGamepadButtonPressed(ControllerIndex, ButtonIndex, Repeat != 0);
-		ProcessEvent(GamepadAnalogEvent);
-		break;
-	}
-	case EToStreamerMsg::GamepadButtonReleased:
-	{
-		GET(FControllerIndex, ControllerIndex);
-		GET(FControllerButtonIndex, ButtonIndex);
-		checkf(Size == 0, TEXT("%d"), Size);
-
-		FEvent GamepadAnalogEvent(EventType::GAMEPAD_RELEASE);
-		GamepadAnalogEvent.SetGamepadButtonReleased(ControllerIndex, ButtonIndex);
-		ProcessEvent(GamepadAnalogEvent);
-		break;
-	}
-	case EToStreamerMsg::GamepadAnalog:
-	{
-		GET(FControllerIndex, ControllerIndex);
-		GET(FControllerAxis, AxisIndex);
-		GET(FControllerAnalog, AnalogValue);
-		checkf(Size == 0, TEXT("%d"), Size);
-
-		FEvent GamepadAnalogEvent(EventType::GAMEPAD_ANALOG);
-		GamepadAnalogEvent.SetGamepadAnalog(ControllerIndex, AxisIndex, (float)AnalogValue);
-		ProcessEvent(GamepadAnalogEvent);
-		break;
-	}
 	}
 }
 
@@ -953,36 +968,36 @@ void FInputDevice::FEvent::GetMouseClick(EMouseButtons::Type& OutButton, uint16&
 	uint8 Button = Data.MouseButton.Button;
 	switch (Button)
 	{
-	case 0:
-	{
-		OutButton = EMouseButtons::Left;
-	}
-	break;
-	case 1:
-	{
-		OutButton = EMouseButtons::Middle;
-	}
-	break;
-	case 2:
-	{
-		OutButton = EMouseButtons::Right;
-	}
-	break;
-	case 3:
-	{
-		OutButton = EMouseButtons::Thumb01;
-	}
-	break;
-	case 4:
-	{
-		OutButton = EMouseButtons::Thumb02;
-	}
-	break;
-	default:
-	{
-		UE_LOG(PixelStreamerInputDevice, Error, TEXT("Unknown Pixel Streaming mouse click with button %d and word 0x%016llx"), Button, Data.Word);
-	}
-	break;
+		case 0:
+		{
+			OutButton = EMouseButtons::Left;
+		}
+		break;
+		case 1:
+		{
+			OutButton = EMouseButtons::Middle;
+		}
+		break;
+		case 2:
+		{
+			OutButton = EMouseButtons::Right;
+		}
+		break;
+		case 3:
+		{
+			OutButton = EMouseButtons::Thumb01;
+		}
+		break;
+		case 4:
+		{
+			OutButton = EMouseButtons::Thumb02;
+		}
+		break;
+		default:
+		{
+			UE_LOG(PixelStreamerInputDevice, Error, TEXT("Unknown Pixel Streaming mouse click with button %d and word 0x%016llx"), Button, Data.Word);
+		}
+		break;
 	}
 	OutPosX = Data.MouseButton.PosX;
 	OutPosY = Data.MouseButton.PosY;
diff --git a/Engine/Plugins/Media/PixelStreaming/Source/PixelStreaming/Private/PixelStreamingAudioDeviceModule.h b/Engine/Plugins/Media/PixelStreaming/Source/PixelStreaming/Private/PixelStreamingAudioDeviceModule.h
--- a/Engine/Plugins/Media/PixelStreaming/Source/PixelStreaming/Private/PixelStreamingAudioDeviceModule.h
+++ b/Engine/Plugins/Media/PixelStreaming/Source/PixelStreaming/Private/PixelStreamingAudioDeviceModule.h
@@ -14,52 +14,52 @@
 class FSubmixCapturer : public ISubmixBufferListener
 {
 
-	public:
+public:
+	// This magic number is the max volume used in webrtc fake audio device module.
+	static const uint32_t MaxVolumeLevel = 14392;
+
+	FSubmixCapturer()
+		: bInitialised(false)
+		, bCapturing(false)
+		, TargetSampleRate(48000)
+		, TargetNumChannels(2)
+		, bReportedSampleRateMismatch(false)
+		, AudioCallback(nullptr)
+		, VolumeLevel(FSubmixCapturer::MaxVolumeLevel)
+		, RecordingBuffer()
+		, CriticalSection()
+	{
+	}
 
-		// This magic number is the max volume used in webrtc fake audio device module.
-		static const uint32_t MaxVolumeLevel = 14392;	
-
-		FSubmixCapturer() 
-		:	bInitialised(false)
-		,	bCapturing(false)
-		,	TargetSampleRate(48000)
-		,	TargetNumChannels(2)
-		,	bReportedSampleRateMismatch(false)
-		,	AudioCallback(nullptr)
-		,	VolumeLevel(FSubmixCapturer::MaxVolumeLevel) 
-		,	RecordingBuffer()
-		,	CriticalSection()
-		{}
-
-		virtual ~FSubmixCapturer() = default;
-
-		bool Init();
-		bool IsInitialised() const;
-		bool IsCapturing() const;
-		void Uninitialise();
-		bool StartCapturing();
-		bool EndCapturing();
-		uint32_t GetVolume() const;
-		void SetVolume(uint32_t NewVolume);
-
-		// ISubmixBufferListener interface
-		void OnNewSubmixBuffer(const USoundSubmix* OwningSubmix, float* AudioData, int32 NumSamples, int32 NumChannels, const int32 SampleRate, double AudioClock) override;
-
-		void RegisterAudioCallback(webrtc::AudioTransport* AudioCb);
+	virtual ~FSubmixCapturer() = default;
 
-	private:
-		int32 GetSamplesPerDurationSecs(float InSeconds) const;
+	bool Init();
+	bool IsInitialised() const;
+	bool IsCapturing() const;
+	void Uninitialise();
+	bool StartCapturing();
+	bool EndCapturing();
+	uint32_t GetVolume() const;
+	void SetVolume(uint32_t NewVolume);
 
-	private:
-		FThreadSafeBool bInitialised;
-		FThreadSafeBool bCapturing;
-		uint32 TargetSampleRate;
-		uint8 TargetNumChannels;
-		bool bReportedSampleRateMismatch;
-		webrtc::AudioTransport* AudioCallback;
-		uint32_t VolumeLevel;
-		TArray<int16_t> RecordingBuffer;
-		FCriticalSection CriticalSection; // One thread captures audio from UE, the other controls the state of the capturer.
+	// ISubmixBufferListener interface
+	void OnNewSubmixBuffer(const USoundSubmix* OwningSubmix, float* AudioData, int32 NumSamples, int32 NumChannels, const int32 SampleRate, double AudioClock) override;
+
+	void RegisterAudioCallback(webrtc::AudioTransport* AudioCb);
+
+private:
+	int32 GetSamplesPerDurationSecs(float InSeconds) const;
+
+private:
+	FThreadSafeBool bInitialised;
+	FThreadSafeBool bCapturing;
+	uint32 TargetSampleRate;
+	uint8 TargetNumChannels;
+	bool bReportedSampleRateMismatch;
+	webrtc::AudioTransport* AudioCallback;
+	uint32_t VolumeLevel;
+	TArray<int16_t> RecordingBuffer;
+	FCriticalSection CriticalSection; // One thread captures audio from UE, the other controls the state of the capturer.
 };
 
 // Requests audio from WebRTC at a regular interval (10ms)
@@ -68,73 +68,71 @@ class FSubmixCapturer : public ISubmixBufferListener
 // there is no demand for audio and does not populate the sinks.
 class FAudioPlayoutRequester
 {
+public:
+	class Runnable : public FRunnable
+	{
 	public:
+		Runnable(TFunction<void()> RequestPlayoutFunc)
+			: bIsRunning(false)
+			, LastAudioRequestTimeMs(0)
+			, RequestPlayoutFunc(RequestPlayoutFunc) {}
+		// Begin FRunnable interface.
+		virtual bool Init() override;
+		virtual uint32 Run() override;
+		virtual void Stop() override;
+		virtual void Exit() override;
+		// End FRunnable interface
+	private:
+		bool bIsRunning;
+		int64_t LastAudioRequestTimeMs;
+		TFunction<void()> RequestPlayoutFunc;
+	};
+
+	FAudioPlayoutRequester()
+		: bIsPlayoutInitialised(false)
+		, bIsPlaying(false)
+		, SampleRate(48000) //webrtc will mix all sources into one big buffer with this sample rate
+		, NumChannels(2)	//webrtc will mix all sources into one big buffer with this many channels
+		, RequesterRunnable()
+		, RequesterThread()
+		, AudioCallback(nullptr)
+		, PlayoutCriticalSection()
+		, PlayoutBuffer(){};
+
+	void InitPlayout();
+	void StartPlayout();
+	void StopPlayout();
+	bool Playing() const;
+	bool PlayoutIsInitialized() const;
+	void Uninitialise();
+	void RegisterAudioCallback(webrtc::AudioTransport* AudioCallback);
 
-		class Runnable : public FRunnable
-		{	
-			public:
-				Runnable(TFunction<void()> RequestPlayoutFunc)
-				: bIsRunning(false)
-				, LastAudioRequestTimeMs(0)
-				, RequestPlayoutFunc(RequestPlayoutFunc){}
-				// Begin FRunnable interface.
-				virtual bool Init() override;
-				virtual uint32 Run() override;
-				virtual void Stop() override;
-				virtual void Exit() override;
-				// End FRunnable interface
-			private:
-				bool bIsRunning;
-				int64_t LastAudioRequestTimeMs;
-				TFunction<void()> RequestPlayoutFunc;
-		};
-
-		FAudioPlayoutRequester() 
-			: bIsPlayoutInitialised(false)
-			, bIsPlaying(false)
-			, SampleRate(48000) //webrtc will mix all sources into one big buffer with this sample rate
-			, NumChannels(2)    //webrtc will mix all sources into one big buffer with this many channels
-			, RequesterRunnable()
-			, RequesterThread()
-			, AudioCallback(nullptr)
-			, PlayoutCriticalSection()
-			, PlayoutBuffer(){};
-
-		void InitPlayout();
-		void StartPlayout();
-		void StopPlayout();
-	    bool Playing() const;
-		bool PlayoutIsInitialized() const;
-		void Uninitialise();
-		void RegisterAudioCallback(webrtc::AudioTransport* AudioCallback);
+public:
+	static int16_t const RequestIntervalMs = 10;
 
-	public:
-		static int16_t const RequestIntervalMs = 10;
-		
-	private:
-		FThreadSafeBool bIsPlayoutInitialised;
-		FThreadSafeBool bIsPlaying;
-		uint32 SampleRate;
-		uint8 NumChannels;
-		TUniquePtr<FAudioPlayoutRequester::Runnable> RequesterRunnable;
-		TUniquePtr<FRunnableThread> RequesterThread;
-		webrtc::AudioTransport* AudioCallback;
-		FCriticalSection PlayoutCriticalSection;
-		TArray<int16_t> PlayoutBuffer;
+private:
+	FThreadSafeBool bIsPlayoutInitialised;
+	FThreadSafeBool bIsPlaying;
+	uint32 SampleRate;
+	uint8 NumChannels;
+	TUniquePtr<FAudioPlayoutRequester::Runnable> RequesterRunnable;
+	TUniquePtr<FRunnableThread> RequesterThread;
+	webrtc::AudioTransport* AudioCallback;
+	FCriticalSection PlayoutCriticalSection;
+	TArray<int16_t> PlayoutBuffer;
 };
 
 // A custom audio device module for WebRTC.
-// It lets us receive WebRTC audio in UE 
+// It lets us receive WebRTC audio in UE
 // and also transmit UE audio to WebRTC.
-class FPixelStreamingAudioDeviceModule
-	: public webrtc::AudioDeviceModule
+class FPixelStreamingAudioDeviceModule : public webrtc::AudioDeviceModule
 {
 
 public:
 	FPixelStreamingAudioDeviceModule()
-	:	bInitialized(false)
-	,	Capturer()
-	,	Requester(){};
+		: bInitialized(false)
+		, Capturer()
+		, Requester(){};
 
 	virtual ~FPixelStreamingAudioDeviceModule() = default;
 
@@ -144,7 +142,6 @@ private:
 	FAudioPlayoutRequester Requester;
 
 private:
-
 	//
 	// webrtc::AudioDeviceModule interface
 	//
@@ -310,5 +307,4 @@ private:
 	{
 		return -1;
 	}
-
 };
\ No newline at end of file
diff --git a/Engine/Plugins/Media/PixelStreaming/Source/PixelStreaming/Private/PixelStreamingEncoderFactory.cpp b/Engine/Plugins/Media/PixelStreaming/Source/PixelStreaming/Private/PixelStreamingEncoderFactory.cpp
--- a/Engine/Plugins/Media/PixelStreaming/Source/PixelStreaming/Private/PixelStreamingEncoderFactory.cpp
+++ b/Engine/Plugins/Media/PixelStreaming/Source/PixelStreaming/Private/PixelStreamingEncoderFactory.cpp
@@ -12,7 +12,7 @@
 #include "PixelStreamingSimulcastEncoderAdapter.h"
 
 FPixelStreamingSimulcastEncoderFactory::FPixelStreamingSimulcastEncoderFactory(IPixelStreamingSessions* InPixelStreamingSessions)
-	:RealFactory(new FPixelStreamingVideoEncoderFactory(InPixelStreamingSessions))
+	: RealFactory(new FPixelStreamingVideoEncoderFactory(InPixelStreamingSessions))
 {
 }
 
@@ -111,7 +111,7 @@ void FPixelStreamingVideoEncoderFactory::ForceKeyFrame()
 {
 	FScopeLock Lock(&ActiveEncodersGuard);
 	// Go through each encoder and send our encoded image to its callback
-	for(auto& KeyAndEncoder : HardwareEncoders)
+	for (auto& KeyAndEncoder : HardwareEncoders)
 	{
 		KeyAndEncoder.Value->SetForceNextKeyframe();
 	}
@@ -123,7 +123,7 @@ namespace
 	{
 		return static_cast<uint64>(Width) << 32 | static_cast<uint64>(Height);
 	}
-}
+} // namespace
 
 FPixelStreamingVideoEncoderFactory::FHardwareEncoderId FPixelStreamingVideoEncoderFactory::GetOrCreateHardwareEncoder(int Width, int Height, int MaxBitrate, int TargetBitrate, int MaxFramerate)
 {
diff --git a/Engine/Plugins/Media/PixelStreaming/Source/PixelStreaming/Private/PixelStreamingModule.cpp b/Engine/Plugins/Media/PixelStreaming/Source/PixelStreaming/Private/PixelStreamingModule.cpp
--- a/Engine/Plugins/Media/PixelStreaming/Source/PixelStreaming/Private/PixelStreamingModule.cpp
+++ b/Engine/Plugins/Media/PixelStreaming/Source/PixelStreaming/Private/PixelStreamingModule.cpp
@@ -18,9 +18,9 @@
 #include "Utils.h"
 
 #if PLATFORM_WINDOWS || PLATFORM_XBOXONE
-#include "Windows/WindowsHWrapper.h"
+	#include "Windows/WindowsHWrapper.h"
 #elif PLATFORM_LINUX
-#include "CudaModule.h"
+	#include "CudaModule.h"
 #endif
 
 #include "RenderingThread.h"
@@ -40,7 +40,7 @@
 #include "Engine/Engine.h"
 
 #if !UE_BUILD_SHIPPING
-#	include "DrawDebugHelpers.h"
+	#include "DrawDebugHelpers.h"
 #endif
 
 DEFINE_LOG_CATEGORY(PixelStreaming);
@@ -49,8 +49,8 @@ IPixelStreamingModule* FPixelStreamingModule::PixelStreamingModule = nullptr;
 
 namespace
 {
-	
-	#if PLATFORM_WINDOWS || PLATFORM_XBOXONE
+
+#if PLATFORM_WINDOWS || PLATFORM_XBOXONE
 	// required for WMF video decoding
 	// some Windows versions don't have Media Foundation preinstalled. We configure MF DLLs as delay-loaded and load them manually here
 	// checking the result and avoiding error message box if failed
@@ -72,8 +72,8 @@ namespace
 				&& FPlatformProcess::GetDllHandle(TEXT("msmpeg2adec.dll"));
 		}
 	}
-	#endif
-}
+#endif
+} // namespace
 
 void FPixelStreamingModule::InitStreamer()
 {
@@ -81,12 +81,12 @@ void FPixelStreamingModule::InitStreamer()
 	FParse::Value(FCommandLine::Get(), TEXT("PixelStreamingID="), StreamerId);
 
 	FString SignallingServerURL;
-	if (!FParse::Value(FCommandLine::Get(), TEXT("PixelStreamingURL="), SignallingServerURL)) {
+	if (!FParse::Value(FCommandLine::Get(), TEXT("PixelStreamingURL="), SignallingServerURL))
+	{
 
 		FString SignallingServerIP;
 		uint16 SignallingServerPort = 8888;
-		if (!FParse::Value(FCommandLine::Get(), TEXT("PixelStreamingIP="), SignallingServerIP) ||
-			!FParse::Value(FCommandLine::Get(), TEXT("PixelStreamingPort="), SignallingServerPort))
+		if (!FParse::Value(FCommandLine::Get(), TEXT("PixelStreamingIP="), SignallingServerIP) || !FParse::Value(FCommandLine::Get(), TEXT("PixelStreamingPort="), SignallingServerPort))
 		{
 			UE_LOG(PixelStreamer, Log, TEXT("PixelStreaming is disabled, provide `PixelStreamingIP` and `PixelStreamingPort` cmd-args to enable it"));
 			return;
@@ -142,7 +142,7 @@ void FPixelStreamingModule::InitStreamer()
 	verify(FModuleManager::Get().LoadModule(FName("ImageWrapper")));
 
 	Streamer = MakeUnique<FStreamer>(SignallingServerURL, StreamerId);
-	
+
 	// Streamer has been created, so module is now "ready" for external use.
 	ReadyEvent.Broadcast(*this);
 }
@@ -151,7 +151,7 @@ void FPixelStreamingModule::InitStreamer()
 void FPixelStreamingModule::StartupModule()
 {
 	// Pixel Streaming does not make sense without an RHI so we don't run in commandlets without one.
-	if(IsRunningCommandlet() && !IsAllowCommandletRendering())
+	if (IsRunningCommandlet() && !IsAllowCommandletRendering())
 	{
 		return;
 	}
@@ -160,17 +160,12 @@ void FPixelStreamingModule::StartupModule()
 	PixelStreamingSettings::InitialiseSettings();
 
 	// only D3D11/D3D12 is supported
-	if (GDynamicRHI == nullptr ||
-		!( GDynamicRHI->GetName() == FString(TEXT("D3D11")) || 
-		   GDynamicRHI->GetName() == FString(TEXT("D3D12"))	||
-		   GDynamicRHI->GetName() == FString(TEXT("Vulkan"))))
+	if (GDynamicRHI == nullptr || !(GDynamicRHI->GetName() == FString(TEXT("D3D11")) || GDynamicRHI->GetName() == FString(TEXT("D3D12")) || GDynamicRHI->GetName() == FString(TEXT("Vulkan"))))
 	{
 		UE_LOG(PixelStreaming, Warning, TEXT("Only D3D11/D3D12/Vulkan Dynamic RHI is supported. Detected %s"), GDynamicRHI != nullptr ? GDynamicRHI->GetName() : TEXT("[null]"));
 		return;
 	}
-	else if( GDynamicRHI->GetName() == FString(TEXT("D3D11")) || 
-		   	 GDynamicRHI->GetName() == FString(TEXT("D3D12")) ||
-			 GDynamicRHI->GetName() == FString(TEXT("Vulkan")))
+	else if (GDynamicRHI->GetName() == FString(TEXT("D3D11")) || GDynamicRHI->GetName() == FString(TEXT("D3D12")) || GDynamicRHI->GetName() == FString(TEXT("Vulkan")))
 	{
 		// By calling InitStreamer post engine init we can use pixel streaming in standalone editor mode
 		FCoreDelegates::OnFEngineLoopInitComplete.AddRaw(this, &FPixelStreamingModule::InitStreamer);
@@ -190,23 +185,23 @@ void FPixelStreamingModule::ShutdownModule()
 
 IPixelStreamingModule* FPixelStreamingModule::GetModule()
 {
-	if(PixelStreamingModule) 
-    {
-        return PixelStreamingModule;
-    }
-    IPixelStreamingModule* Module = FModuleManager::Get().LoadModulePtr<IPixelStreamingModule>("PixelStreaming");
-    if(Module)
-    {
-        PixelStreamingModule = Module;
-    }
-    return PixelStreamingModule;
+	if (PixelStreamingModule)
+	{
+		return PixelStreamingModule;
+	}
+	IPixelStreamingModule* Module = FModuleManager::Get().LoadModulePtr<IPixelStreamingModule>("PixelStreaming");
+	if (Module)
+	{
+		PixelStreamingModule = Module;
+	}
+	return PixelStreamingModule;
 }
 
 bool FPixelStreamingModule::CheckPlatformCompatibility() const
 {
 	bool bCompatible = true;
 
-	#if PLATFORM_WINDOWS || PLATFORM_XBOXONE
+#if PLATFORM_WINDOWS || PLATFORM_XBOXONE
 	bool bWin8OrHigher = FPlatformMisc::VerifyWindowsVersion(6, 2);
 	if (!bWin8OrHigher)
 	{
@@ -217,8 +212,8 @@ bool FPixelStreamingModule::CheckPlatformCompatibility() const
 		UE_LOG(PixelStreamer, Error, TEXT("%s"), *ErrorString);
 		bCompatible = false;
 	}
-	#endif
-	
+#endif
+
 	if (!FStreamer::CheckPlatformCompatibility())
 	{
 		FText TitleText = FText::FromString(TEXT("Pixel Streaming Plugin"));
@@ -312,9 +307,9 @@ const TArray<UPixelStreamerInputComponent*> FPixelStreamingModule::GetInputCompo
 void FPixelStreamingModule::FreezeFrame(UTexture2D* Texture)
 {
 	if (Texture)
-	{		
-		ENQUEUE_RENDER_COMMAND(ReadSurfaceCommand)([this, Texture](FRHICommandListImmediate& RHICmdList)
-		{
+	{
+		ENQUEUE_RENDER_COMMAND(ReadSurfaceCommand)
+		([this, Texture](FRHICommandListImmediate& RHICmdList) {
 			// A frame is supplied so immediately read its data and send as a JPEG.
 			FTexture2DRHIRef Texture2DRHI = (Texture->GetResource() && Texture->GetResource()->TextureRHI) ? Texture->GetResource()->TextureRHI->GetTexture2D() : nullptr;
 			if (!Texture2DRHI)
@@ -324,9 +319,9 @@ void FPixelStreamingModule::FreezeFrame(UTexture2D* Texture)
 			}
 			uint32 Width = Texture2DRHI->GetSizeX();
 			uint32 Height = Texture2DRHI->GetSizeY();
-			
+
 			FTexture2DRHIRef DestTexture = CreateTexture(Width, Height);
-			
+
 			FGPUFenceRHIRef CopyFence = GDynamicRHI->RHICreateGPUFence(*FString::Printf(TEXT("FreezeFrameFence")));
 
 			// Copy freeze frame texture to empty texture
@@ -402,12 +397,10 @@ void FPixelStreamingModule::SendCommand(const FString& Descriptor)
 
 void FPixelStreamingModule::OnGameModePostLogin(AGameModeBase* GameMode, APlayerController* NewPlayer)
 {
-	
 }
 
 void FPixelStreamingModule::OnGameModeLogout(AGameModeBase* GameMode, AController* Exiting)
 {
-	
 }
 
 void FPixelStreamingModule::SendJpeg(TArray<FColor> RawData, const FIntRect& Rect)
@@ -434,7 +427,7 @@ void FPixelStreamingModule::SendJpeg(TArray<FColor> RawData, const FIntRect& Rec
 }
 
 void FPixelStreamingModule::SendFileData(TArray<uint8>& ByteData, FString& MimeType, FString& FileExtension)
-{	
+{
 	Streamer->SendFileData(ByteData, MimeType, FileExtension);
 }
 
@@ -450,7 +443,6 @@ bool FPixelStreamingModule::IsTickableInEditor() const
 
 void FPixelStreamingModule::Tick(float DeltaTime)
 {
-	
 }
 
 TStatId FPixelStreamingModule::GetStatId() const
@@ -460,7 +452,7 @@ TStatId FPixelStreamingModule::GetStatId() const
 
 IPixelStreamingAudioSink* FPixelStreamingModule::GetPeerAudioSink(FPlayerId PlayerId)
 {
-	if(!this->Streamer.IsValid())
+	if (!this->Streamer.IsValid())
 	{
 		UE_LOG(PixelStreamer, Error, TEXT("Cannot get audio sink when streamer does not yet exist."));
 		return nullptr;
@@ -471,7 +463,7 @@ IPixelStreamingAudioSink* FPixelStreamingModule::GetPeerAudioSink(FPlayerId Play
 
 IPixelStreamingAudioSink* FPixelStreamingModule::GetUnlistenedAudioSink()
 {
-	if(!this->Streamer.IsValid())
+	if (!this->Streamer.IsValid())
 	{
 		UE_LOG(PixelStreamer, Error, TEXT("Cannot get audio sink when streamer does not yet exist."));
 		return nullptr;
@@ -482,7 +474,7 @@ IPixelStreamingAudioSink* FPixelStreamingModule::GetUnlistenedAudioSink()
 
 void FPixelStreamingModule::AddAnyStatChangedCallback(TWeakPtr<IPixelStreamingStatsConsumer> Callback)
 {
-	if(!this->Streamer.IsValid())
+	if (!this->Streamer.IsValid())
 	{
 		UE_LOG(PixelStreamer, Error, TEXT("Cannot add stat callback when streamer does not yet exist."));
 		return;
@@ -493,7 +485,7 @@ void FPixelStreamingModule::AddAnyStatChangedCallback(TWeakPtr<IPixelStreamingSt
 
 void FPixelStreamingModule::RemoveAnyStatChangedCallback(TWeakPtr<IPixelStreamingStatsConsumer> Callback)
 {
-	if(!this->Streamer.IsValid())
+	if (!this->Streamer.IsValid())
 	{
 		UE_LOG(PixelStreamer, Error, TEXT("Cannot remove stat callback when streamer does not yet exist."));
 		return;
diff --git a/Engine/Plugins/Media/PixelStreaming/Source/PixelStreaming/Private/PixelStreamingRealEncoder.cpp b/Engine/Plugins/Media/PixelStreaming/Source/PixelStreaming/Private/PixelStreamingRealEncoder.cpp
--- a/Engine/Plugins/Media/PixelStreaming/Source/PixelStreaming/Private/PixelStreamingRealEncoder.cpp
+++ b/Engine/Plugins/Media/PixelStreaming/Source/PixelStreaming/Private/PixelStreamingRealEncoder.cpp
@@ -20,13 +20,10 @@ namespace
 			{
 				// either a 0,0,1 or 0,0,0,1 sequence indicates a new 'nal'
 				size_t nal_maker_length = 3;
-				if (offset < (CodedDataSize - 3) && CodedData[offset] == 0 &&
-					CodedData[offset + 1] == 0 && CodedData[offset + 2] == 1)
+				if (offset < (CodedDataSize - 3) && CodedData[offset] == 0 && CodedData[offset + 1] == 0 && CodedData[offset + 2] == 1)
 				{
 				}
-				else if (offset < (CodedDataSize - 4) && CodedData[offset] == 0 &&
-					CodedData[offset + 1] == 0 && CodedData[offset + 2] == 0 &&
-					CodedData[offset + 3] == 1)
+				else if (offset < (CodedDataSize - 4) && CodedData[offset] == 0 && CodedData[offset + 1] == 0 && CodedData[offset + 2] == 0 && CodedData[offset + 3] == 1)
 				{
 					nal_maker_length = 4;
 				}
@@ -76,7 +73,6 @@ namespace
 			FragHeader.fragmentationLength[i] = NALUIndex.payload_size;
 		}
 
-
 		Image.timing_.packetization_finish_ms = FTimespan::FromSeconds(FPlatformTime::Seconds()).GetTotalMilliseconds();
 		Image.timing_.encode_start_ms = InPacket.Timings.StartTs.GetTotalMilliseconds();
 		Image.timing_.encode_finish_ms = InPacket.Timings.FinishTs.GetTotalMilliseconds();
@@ -103,11 +99,11 @@ namespace
 
 		Factory->OnEncodedImage(SourceEncoderId, Image, &CodecInfo, &FragHeader);
 	}
-}
+} // namespace
 
 FPixelStreamingRealEncoder::FPixelStreamingRealEncoder(uint64 EncoderId, int Width, int Height, int MaxBitrate, int TargetBitrate, int MaxFramerate, FPixelStreamingVideoEncoderFactory* InFactory)
-:Id(EncoderId)
-,Factory(InFactory)
+	: Id(EncoderId)
+	, Factory(InFactory)
 {
 	checkf(Factory, TEXT("Encoder factory is null."));
 
@@ -125,8 +121,7 @@ FPixelStreamingRealEncoder::FPixelStreamingRealEncoder(uint64 EncoderId, int Wid
 	Encoder = AVEncoder::FVideoEncoderFactory::Get().Create(Available[0].ID, FrameFactory.GetOrCreateVideoEncoderInput(EncoderConfig.Width, EncoderConfig.Height), EncoderConfig);
 	checkf(Encoder, TEXT("Pixel Streaming video encoder creation failed, check encoder config."));
 
-	Encoder->SetOnEncodedPacket([EncoderId, InFactory](uint32 InLayerIndex, const AVEncoder::FVideoEncoderInputFrame* InFrame, const AVEncoder::FCodecPacket& InPacket)
-	{
+	Encoder->SetOnEncodedPacket([EncoderId, InFactory](uint32 InLayerIndex, const AVEncoder::FVideoEncoderInputFrame* InFrame, const AVEncoder::FCodecPacket& InPacket) {
 		OnEncodedPacket(EncoderId, InFactory, InLayerIndex, InFrame, InPacket);
 	});
 }
diff --git a/Engine/Plugins/Media/PixelStreaming/Source/PixelStreaming/Private/PixelStreamingSettings.cpp b/Engine/Plugins/Media/PixelStreaming/Source/PixelStreaming/Private/PixelStreamingSettings.cpp
--- a/Engine/Plugins/Media/PixelStreaming/Source/PixelStreaming/Private/PixelStreamingSettings.cpp
+++ b/Engine/Plugins/Media/PixelStreaming/Source/PixelStreaming/Private/PixelStreamingSettings.cpp
@@ -307,7 +307,7 @@ namespace PixelStreamingSettings
 	AVEncoder::FVideoEncoder::RateControlMode GetRateControlCVar()
 	{
 		auto const cvarStr = CVarPixelStreamingEncoderRateControl.GetValueOnAnyThread();
-		auto const it	   = RateControlCVarMap.find(cvarStr);
+		auto const it = RateControlCVarMap.find(cvarStr);
 		if (it == std::end(RateControlCVarMap))
 			return AVEncoder::FVideoEncoder::RateControlMode::CBR;
 		return it->second;
@@ -316,7 +316,7 @@ namespace PixelStreamingSettings
 	AVEncoder::FVideoEncoder::MultipassMode GetMultipassCVar()
 	{
 		auto const cvarStr = CVarPixelStreamingEncoderMultipass.GetValueOnAnyThread();
-		auto const it	   = MultipassCVarMap.find(cvarStr);
+		auto const it = MultipassCVarMap.find(cvarStr);
 		if (it == std::end(MultipassCVarMap))
 			return AVEncoder::FVideoEncoder::MultipassMode::FULL;
 		return it->second;
@@ -340,7 +340,7 @@ namespace PixelStreamingSettings
 	AVEncoder::FVideoEncoder::H264Profile GetH264Profile()
 	{
 		auto const cvarStr = CVarPixelStreamingH264Profile.GetValueOnAnyThread();
-		auto const it	   = H264ProfileMap.find(cvarStr);
+		auto const it = H264ProfileMap.find(cvarStr);
 		if (it == std::end(H264ProfileMap))
 			return AVEncoder::FVideoEncoder::H264Profile::BASELINE;
 		return it->second;
@@ -355,9 +355,9 @@ namespace PixelStreamingSettings
 
 		FString StringOptions;
 		bool bPassedSimulcastParams = FParse::Value(FCommandLine::Get(), TEXT("SimulcastParameters="), StringOptions, false);
-		
+
 		// If no simulcast parameters are passed use some default values
-		if(!bPassedSimulcastParams)
+		if (!bPassedSimulcastParams)
 		{
 			StringOptions = FString(TEXT("1.0,5000000,20000000,2.0,1000000,5000000,4.0,50000,1000000"));
 		}
@@ -402,14 +402,8 @@ namespace PixelStreamingSettings
 	{
 		bool bLogStats = Var->GetBool();
 		IPixelStreamingModule& Module = IPixelStreamingModule::Get();
-		
-
-		if(Module.IsReady())
-		{
 
-		}
-
-		if(bLogStats)
+		if (bLogStats)
 		{
 			Module.AddAnyStatChangedCallback(StatsLoggerWeak);
 		}
@@ -481,5 +475,4 @@ namespace PixelStreamingSettings
 		Module.OnReady().AddStatic(&PixelStreamingSettings::OnStreamerReady);
 	}
 
-
 } // namespace PixelStreamingSettings
diff --git a/Engine/Plugins/Media/PixelStreaming/Source/PixelStreaming/Private/PixelStreamingSettings.h b/Engine/Plugins/Media/PixelStreaming/Source/PixelStreaming/Private/PixelStreamingSettings.h
--- a/Engine/Plugins/Media/PixelStreaming/Source/PixelStreaming/Private/PixelStreamingSettings.h
+++ b/Engine/Plugins/Media/PixelStreaming/Source/PixelStreaming/Private/PixelStreamingSettings.h
@@ -15,7 +15,7 @@ namespace PixelStreamingSettings
 {
 	extern void InitialiseSettings();
 
-// Begin Encoder CVars
+	// Begin Encoder CVars
 	extern TAutoConsoleVariable<int32> CVarPixelStreamingEncoderTargetBitrate;
 	extern TAutoConsoleVariable<int32> CVarPixelStreamingEncoderMaxBitrate;
 	extern TAutoConsoleVariable<bool> CVarPixelStreamingDebugDumpFrame;
@@ -26,15 +26,15 @@ namespace PixelStreamingSettings
 	extern TAutoConsoleVariable<FString> CVarPixelStreamingEncoderMultipass;
 	extern TAutoConsoleVariable<FString> CVarPixelStreamingH264Profile;
 	extern TAutoConsoleVariable<int32> CVarPixelStreamingEncoderKeyframeInterval;
-// End Encoder CVars
+	// End Encoder CVars
 
-// Begin Capturer CVars
+	// Begin Capturer CVars
 	extern TAutoConsoleVariable<bool> CVarPixelStreamingUseBackBufferCaptureSize;
 	extern TAutoConsoleVariable<FString> CVarPixelStreamingCaptureSize;
 	extern TAutoConsoleVariable<int32> CVarPixelStreamingMaxNumBackBuffers;
-// End Capturer CVars
+	// End Capturer CVars
 
-// Begin WebRTC CVars
+	// Begin WebRTC CVars
 	extern TAutoConsoleVariable<FString> CVarPixelStreamingDegradationPreference;
 	extern TAutoConsoleVariable<int32> CVarPixelStreamingWebRTCFps;
 	extern TAutoConsoleVariable<int32> CVarPixelStreamingWebRTCStartBitrate;
@@ -47,9 +47,9 @@ namespace PixelStreamingSettings
 	extern TAutoConsoleVariable<bool> CVarPixelStreamingWebRTCDisableAudioSync;
 	extern TAutoConsoleVariable<bool> CVarPixelStreamingWebRTCUseLegacyAudioDevice;
 	extern TAutoConsoleVariable<bool> CVarPixelStreamingWebRTCDisableResolutionChange;
-// End WebRTC CVars
+	// End WebRTC CVars
 
-// Begin Pixel Streaming Plugin CVars
+	// Begin Pixel Streaming Plugin CVars
 	extern TAutoConsoleVariable<bool> CVarPixelStreamingAllowConsoleCommands;
 	extern TAutoConsoleVariable<bool> CVarPixelStreamingOnScreenStats;
 	extern TAutoConsoleVariable<bool> CVarPixelStreamingLogStats;
@@ -58,14 +58,14 @@ namespace PixelStreamingSettings
 	extern TAutoConsoleVariable<bool> CVarSendPlayerIdAsInteger;
 	extern TAutoConsoleVariable<bool> CVarPixelStreamingDisableLatencyTester;
 	extern TArray<FKey> FilteredKeys;
-// Ends Pixel Streaming Plugin CVars
+	// Ends Pixel Streaming Plugin CVars
 
-// Begin utility functions etc.
+	// Begin utility functions etc.
 	AVEncoder::FVideoEncoder::RateControlMode GetRateControlCVar();
 	AVEncoder::FVideoEncoder::MultipassMode GetMultipassCVar();
 	webrtc::DegradationPreference GetDegradationPreference();
 	AVEncoder::FVideoEncoder::H264Profile GetH264Profile();
-// End utility functions etc.
+	// End utility functions etc.
 
 	struct FSimulcastParameters
 	{
@@ -81,13 +81,13 @@ namespace PixelStreamingSettings
 
 	extern FSimulcastParameters SimulcastParameters;
 
-// Begin Command line args
+	// Begin Command line args
 
 	inline bool IsPixelStreamingHideCursor()
 	{
 		return FParse::Param(FCommandLine::Get(), TEXT("PixelStreamingHideCursor"));
 	}
-	
+
 	inline bool IsForceVP8()
 	{
 		return FParse::Param(FCommandLine::Get(), TEXT("PSForceVP8"));
@@ -113,6 +113,6 @@ namespace PixelStreamingSettings
 		return FParse::Value(FCommandLine::Get(), TEXT("PixelStreamingFastPan="), OutFastPan);
 	}
 
-// End Command line args
+	// End Command line args
 
-}
+} // namespace PixelStreamingSettings
diff --git a/Engine/Plugins/Media/PixelStreaming/Source/PixelStreaming/Private/PixelStreamingSimulcastEncoderAdapter.h b/Engine/Plugins/Media/PixelStreaming/Source/PixelStreaming/Private/PixelStreamingSimulcastEncoderAdapter.h
--- a/Engine/Plugins/Media/PixelStreaming/Source/PixelStreaming/Private/PixelStreamingSimulcastEncoderAdapter.h
+++ b/Engine/Plugins/Media/PixelStreaming/Source/PixelStreaming/Private/PixelStreamingSimulcastEncoderAdapter.h
@@ -8,77 +8,77 @@ class FPixelStreamingLayerFrameSource;
 
 namespace webrtc
 {
-    class SimulcastRateAllocator;
-    class VideoEncoderFactory;
-}
+	class SimulcastRateAllocator;
+	class VideoEncoderFactory;
+} // namespace webrtc
 
 // This is highly modified version of webrtc::simulcast_encoder_adapter
 
 class PixelStreamingSimulcastEncoderAdapter : public webrtc::VideoEncoder
 {
 public:
-    PixelStreamingSimulcastEncoderAdapter(webrtc::VideoEncoderFactory* primarty_factory, const webrtc::SdpVideoFormat& format);
-    ~PixelStreamingSimulcastEncoderAdapter() override;
+	PixelStreamingSimulcastEncoderAdapter(webrtc::VideoEncoderFactory* primarty_factory, const webrtc::SdpVideoFormat& format);
+	~PixelStreamingSimulcastEncoderAdapter() override;
 
-    // Implements VideoEncoder.
-    int Release() override;
-    int InitEncode(const webrtc::VideoCodec* codec_settings, const webrtc::VideoEncoder::Settings& settings) override;
-    int Encode(const webrtc::VideoFrame& input_image, const std::vector<webrtc::VideoFrameType>* frame_types) override;
-    int RegisterEncodeCompleteCallback(webrtc::EncodedImageCallback* callback) override;
+	// Implements VideoEncoder.
+	int Release() override;
+	int InitEncode(const webrtc::VideoCodec* codec_settings, const webrtc::VideoEncoder::Settings& settings) override;
+	int Encode(const webrtc::VideoFrame& input_image, const std::vector<webrtc::VideoFrameType>* frame_types) override;
+	int RegisterEncodeCompleteCallback(webrtc::EncodedImageCallback* callback) override;
 
-    void SetRates(const RateControlParameters& parameters) override;
-    void OnPacketLossRateUpdate(float packet_loss_rate) override;
-    void OnRttUpdate(int64_t rtt_ms) override;
-    void OnLossNotification(const LossNotification& loss_notification) override;
+	void SetRates(const RateControlParameters& parameters) override;
+	void OnPacketLossRateUpdate(float packet_loss_rate) override;
+	void OnRttUpdate(int64_t rtt_ms) override;
+	void OnLossNotification(const LossNotification& loss_notification) override;
 
-    // Eventual handler for the contained encoders' EncodedImageCallbacks, but
-    // called from an internal helper that also knows the correct stream
-    // index.
-    webrtc::EncodedImageCallback::Result OnEncodedImage(
-        size_t stream_idx,
-        const webrtc::EncodedImage& encoded_image,
-        const webrtc::CodecSpecificInfo* codec_specific_info,
-        const webrtc::RTPFragmentationHeader* fragmentation);
+	// Eventual handler for the contained encoders' EncodedImageCallbacks, but
+	// called from an internal helper that also knows the correct stream
+	// index.
+	webrtc::EncodedImageCallback::Result OnEncodedImage(
+		size_t stream_idx,
+		const webrtc::EncodedImage& encoded_image,
+		const webrtc::CodecSpecificInfo* codec_specific_info,
+		const webrtc::RTPFragmentationHeader* fragmentation);
 
-    EncoderInfo GetEncoderInfo() const override;
+	EncoderInfo GetEncoderInfo() const override;
 
 private:
-    struct StreamInfo
-    {
-        StreamInfo(std::unique_ptr<webrtc::VideoEncoder> encoder,
-            std::unique_ptr<webrtc::EncodedImageCallback> callback,
-            std::unique_ptr<webrtc::FramerateController> framerate_controller,
-            uint16_t width,
-            uint16_t height,
-            bool send_stream)
-            : encoder(std::move(encoder))
-            , callback(std::move(callback))
-            , framerate_controller(std::move(framerate_controller))
-            , width(width)
-            , height(height)
-            , key_frame_request(false)
-            , send_stream(send_stream)
-        {
-        }
-        std::unique_ptr<webrtc::VideoEncoder> encoder;
-        std::unique_ptr<webrtc::EncodedImageCallback> callback;
-        std::unique_ptr<webrtc::FramerateController> framerate_controller;
-        uint16_t width;
-        uint16_t height;
-        bool key_frame_request;
-        bool send_stream;
-    };
+	struct StreamInfo
+	{
+		StreamInfo(std::unique_ptr<webrtc::VideoEncoder> encoder,
+			std::unique_ptr<webrtc::EncodedImageCallback> callback,
+			std::unique_ptr<webrtc::FramerateController> framerate_controller,
+			uint16_t width,
+			uint16_t height,
+			bool send_stream)
+			: encoder(std::move(encoder))
+			, callback(std::move(callback))
+			, framerate_controller(std::move(framerate_controller))
+			, width(width)
+			, height(height)
+			, key_frame_request(false)
+			, send_stream(send_stream)
+		{
+		}
+		std::unique_ptr<webrtc::VideoEncoder> encoder;
+		std::unique_ptr<webrtc::EncodedImageCallback> callback;
+		std::unique_ptr<webrtc::FramerateController> framerate_controller;
+		uint16_t width;
+		uint16_t height;
+		bool key_frame_request;
+		bool send_stream;
+	};
 
-    bool IsInitialized() const;
+	bool IsInitialized() const;
 
 	TAtomic<bool> Initialized;
-    webrtc::VideoEncoderFactory* const PrimaryEncoderFactory;
-    const webrtc::SdpVideoFormat VideoFormat;
-    webrtc::VideoCodec CurrentCodec;
-    std::vector<StreamInfo> StreamInfos;
-    webrtc::EncodedImageCallback* EncodedCompleteCallback;
+	webrtc::VideoEncoderFactory* const PrimaryEncoderFactory;
+	const webrtc::SdpVideoFormat VideoFormat;
+	webrtc::VideoCodec CurrentCodec;
+	std::vector<StreamInfo> StreamInfos;
+	webrtc::EncodedImageCallback* EncodedCompleteCallback;
 
 	int LowestResolutionStreamIndex;
 	int HighestResolutionStreamIndex;
-    int EncodeStream(const webrtc::VideoFrame& input_image, FPixelStreamingLayerFrameSource* layer_frame_source, size_t stream_idx, bool send_key_frame);
+	int EncodeStream(const webrtc::VideoFrame& input_image, FPixelStreamingLayerFrameSource* layer_frame_source, size_t stream_idx, bool send_key_frame);
 };
diff --git a/Engine/Plugins/Media/PixelStreaming/Source/PixelStreaming/Private/ProtocolDefs.h b/Engine/Plugins/Media/PixelStreaming/Source/PixelStreaming/Private/ProtocolDefs.h
--- a/Engine/Plugins/Media/PixelStreaming/Source/PixelStreaming/Private/ProtocolDefs.h
+++ b/Engine/Plugins/Media/PixelStreaming/Source/PixelStreaming/Private/ProtocolDefs.h
@@ -14,7 +14,6 @@ namespace PixelStreamingProtocol
 		PSPFlag_IsSFU = 0x02,
 	};
 
-
 	enum class EToStreamerMsg : uint8
 	{
 		/**********************************************************************/
@@ -64,7 +63,7 @@ namespace PixelStreamingProtocol
 		GamepadButtonPressed = 90,
 		GamepadButtonReleased = 91,
 		GamepadAnalog = 92,
-		
+
 		/**********************************************************************/
 
 		/*
@@ -76,7 +75,8 @@ namespace PixelStreamingProtocol
 	};
 
 	//! Messages that can be sent to the webrtc players
-	enum class EToPlayerMsg : uint8 {
+	enum class EToPlayerMsg : uint8
+	{
 		QualityControlOwnership = 0,
 		Response = 1,
 		Command = 2,
@@ -90,7 +90,7 @@ namespace PixelStreamingProtocol
 		FileContents = 10
 	};
 
-	template<typename T>
+	template <typename T>
 	static const T& ParseBuffer(const uint8*& Data, uint32& Size)
 	{
 		checkf(sizeof(T) <= Size, TEXT("%d - %d"), sizeof(T), Size);
@@ -103,7 +103,8 @@ namespace PixelStreamingProtocol
 	static FString ParseString(const webrtc::DataBuffer& Buffer, const size_t Offset = 0)
 	{
 		FString Res;
-		if (Buffer.data.size() > Offset) {
+		if (Buffer.data.size() > Offset)
+		{
 			size_t StringLength = (Buffer.data.size() - Offset) / sizeof(TCHAR);
 			Res.GetCharArray().SetNumUninitialized(StringLength + 1);
 			FMemory::Memcpy(Res.GetCharArray().GetData(), Buffer.data.data() + Offset, StringLength * sizeof(TCHAR));
@@ -112,4 +113,4 @@ namespace PixelStreamingProtocol
 		return Res;
 	}
 
-};
+}; // namespace PixelStreamingProtocol
diff --git a/Engine/Plugins/Media/PixelStreaming/Source/PixelStreaming/Private/ThreadSafePlayerSessions.cpp b/Engine/Plugins/Media/PixelStreaming/Source/PixelStreaming/Private/ThreadSafePlayerSessions.cpp
--- a/Engine/Plugins/Media/PixelStreaming/Source/PixelStreaming/Private/ThreadSafePlayerSessions.cpp
+++ b/Engine/Plugins/Media/PixelStreaming/Source/PixelStreaming/Private/ThreadSafePlayerSessions.cpp
@@ -6,186 +6,182 @@
 #include "PixelStreamerDelegates.h"
 #include "Streamer.h"
 
-#define SUBMIT_TASK_WITH_RETURN(ReturnType, FuncWithReturnType) \
-    if(this->IsInSignallingThread()) \
-    { \
-        return this->FuncWithReturnType(); \
-    } \
-    else \
-    { \
-        return this->WebRtcSignallingThread->Invoke<ReturnType>(RTC_FROM_HERE, [this]() -> ReturnType { return this->FuncWithReturnType(); } ); \
-    } \
-
-#define SUBMIT_TASK_WITH_PARAMS_AND_RETURN(ReturnType, FuncWithReturnType, ...) \
-    if(this->IsInSignallingThread()) \
-    { \
-        return this->FuncWithReturnType(__VA_ARGS__); \
-    } \
-    else \
-    { \
-        return this->WebRtcSignallingThread->Invoke<ReturnType>(RTC_FROM_HERE, [this, __VA_ARGS__]() -> ReturnType { return this->FuncWithReturnType(__VA_ARGS__); } ); \
-    } \
-
-#define SUBMIT_TASK_NO_PARAMS(Func) \
-    if(this->IsInSignallingThread()) \
-    { \
-        this->Func(); \
-    } \
-    else \
-    { \
-        this->WebRtcSignallingThread->PostTask(RTC_FROM_HERE, [this]() { this->Func(); } ); \
-    } \
-
-#define SUBMIT_TASK_WITH_PARAMS(Func, ...) \
-    if(this->IsInSignallingThread()) \
-    { \
-        this->Func(__VA_ARGS__); \
-    } \
-    else \
-    { \
-        this->WebRtcSignallingThread->PostTask(RTC_FROM_HERE, [this, __VA_ARGS__]() { this->Func(__VA_ARGS__); } ); \
-    } \
+#define SUBMIT_TASK_WITH_RETURN(ReturnType, FuncWithReturnType)                                                                                \
+	if (this->IsInSignallingThread())                                                                                                          \
+	{                                                                                                                                          \
+		return this->FuncWithReturnType();                                                                                                     \
+	}                                                                                                                                          \
+	else                                                                                                                                       \
+	{                                                                                                                                          \
+		return this->WebRtcSignallingThread->Invoke<ReturnType>(RTC_FROM_HERE, [this]() -> ReturnType { return this->FuncWithReturnType(); }); \
+	}
+
+#define SUBMIT_TASK_WITH_PARAMS_AND_RETURN(ReturnType, FuncWithReturnType, ...)                                                                                        \
+	if (this->IsInSignallingThread())                                                                                                                                  \
+	{                                                                                                                                                                  \
+		return this->FuncWithReturnType(__VA_ARGS__);                                                                                                                  \
+	}                                                                                                                                                                  \
+	else                                                                                                                                                               \
+	{                                                                                                                                                                  \
+		return this->WebRtcSignallingThread->Invoke<ReturnType>(RTC_FROM_HERE, [this, __VA_ARGS__]() -> ReturnType { return this->FuncWithReturnType(__VA_ARGS__); }); \
+	}
+
+#define SUBMIT_TASK_NO_PARAMS(Func)                                                        \
+	if (this->IsInSignallingThread())                                                      \
+	{                                                                                      \
+		this->Func();                                                                      \
+	}                                                                                      \
+	else                                                                                   \
+	{                                                                                      \
+		this->WebRtcSignallingThread->PostTask(RTC_FROM_HERE, [this]() { this->Func(); }); \
+	}
+
+#define SUBMIT_TASK_WITH_PARAMS(Func, ...)                                                                         \
+	if (this->IsInSignallingThread())                                                                              \
+	{                                                                                                              \
+		this->Func(__VA_ARGS__);                                                                                   \
+	}                                                                                                              \
+	else                                                                                                           \
+	{                                                                                                              \
+		this->WebRtcSignallingThread->PostTask(RTC_FROM_HERE, [this, __VA_ARGS__]() { this->Func(__VA_ARGS__); }); \
+	}
 
 // NOTE: All public methods should be one-liner functions that call one of the above macros to ensure threadsafety of this class.
 
 FThreadSafePlayerSessions::FThreadSafePlayerSessions(rtc::Thread* InWebRtcSignallingThread)
-    : WebRtcSignallingThread(InWebRtcSignallingThread)
+	: WebRtcSignallingThread(InWebRtcSignallingThread)
 {
-
 }
 
 bool FThreadSafePlayerSessions::IsInSignallingThread() const
 {
-    if(this->WebRtcSignallingThread == nullptr)
-    {
-        return false;
-    }
+	if (this->WebRtcSignallingThread == nullptr)
+	{
+		return false;
+	}
 
 	return this->WebRtcSignallingThread->IsCurrent();
 }
 
-
 int FThreadSafePlayerSessions::GetNumPlayers() const
 {
-    SUBMIT_TASK_WITH_RETURN(int, GetNumPlayers_SignallingThread);
+	SUBMIT_TASK_WITH_RETURN(int, GetNumPlayers_SignallingThread);
 }
 
-IPixelStreamingAudioSink* FThreadSafePlayerSessions::GetAudioSink(FPlayerId PlayerId) const
-{
-    SUBMIT_TASK_WITH_PARAMS_AND_RETURN(IPixelStreamingAudioSink*, GetAudioSink_SignallingThread, PlayerId)
+IPixelStreamingAudioSink* FThreadSafePlayerSessions::GetAudioSink(FPlayerId PlayerId) const {
+	SUBMIT_TASK_WITH_PARAMS_AND_RETURN(IPixelStreamingAudioSink*, GetAudioSink_SignallingThread, PlayerId)
 }
 
 IPixelStreamingAudioSink* FThreadSafePlayerSessions::GetUnlistenedAudioSink() const
 {
-    SUBMIT_TASK_WITH_RETURN(IPixelStreamingAudioSink*, GetUnlistenedAudioSink_SignallingThread)
+	SUBMIT_TASK_WITH_RETURN(IPixelStreamingAudioSink*, GetUnlistenedAudioSink_SignallingThread)
 }
 
 bool FThreadSafePlayerSessions::IsQualityController(FPlayerId PlayerId) const
 {
-    // Due to how many theads this particular method is called on we choose not to schedule reading it as a new task and instead just made the
-    // the variable atomic. Additionally some of the calling threads were deadlocking waiting for each other to finish other methods while calling this method.
-    // Not scheduling this as a task is the exception to the general rule of this class though.
-    FScopeLock Lock(&QualityControllerCS);
-    return this->QualityControllingPlayer == PlayerId;
+	// Due to how many theads this particular method is called on we choose not to schedule reading it as a new task and instead just made the
+	// the variable atomic. Additionally some of the calling threads were deadlocking waiting for each other to finish other methods while calling this method.
+	// Not scheduling this as a task is the exception to the general rule of this class though.
+	FScopeLock Lock(&QualityControllerCS);
+	return this->QualityControllingPlayer == PlayerId;
 }
 
 void FThreadSafePlayerSessions::SetQualityController(FPlayerId PlayerId)
 {
-    {
-        FScopeLock Lock(&QualityControllerCS);
-        this->QualityControllingPlayer = PlayerId;
-    }
-    
-    SUBMIT_TASK_WITH_PARAMS(SetQualityController_SignallingThread, PlayerId)
+	{
+		FScopeLock Lock(&QualityControllerCS);
+		this->QualityControllingPlayer = PlayerId;
+	}
+
+	SUBMIT_TASK_WITH_PARAMS(SetQualityController_SignallingThread, PlayerId)
 }
 
 bool FThreadSafePlayerSessions::SendMessage(FPlayerId PlayerId, PixelStreamingProtocol::EToPlayerMsg Type, const FString& Descriptor) const
 {
-    SUBMIT_TASK_WITH_PARAMS_AND_RETURN(bool, SendMessage_SignallingThread, PlayerId, Type, Descriptor)
+	SUBMIT_TASK_WITH_PARAMS_AND_RETURN(bool, SendMessage_SignallingThread, PlayerId, Type, Descriptor)
 }
 
 void FThreadSafePlayerSessions::SendLatestQP(FPlayerId PlayerId, int LatestQP) const
 {
-    SUBMIT_TASK_WITH_PARAMS(SendLatestQP_SignallingThread, PlayerId, LatestQP)
+	SUBMIT_TASK_WITH_PARAMS(SendLatestQP_SignallingThread, PlayerId, LatestQP)
 }
 
 void FThreadSafePlayerSessions::SendFreezeFrameTo(FPlayerId PlayerId, const TArray64<uint8>& JpegBytes) const
 {
-    SUBMIT_TASK_WITH_PARAMS(SendFreezeFrameTo_SignallingThread, PlayerId, JpegBytes)
+	SUBMIT_TASK_WITH_PARAMS(SendFreezeFrameTo_SignallingThread, PlayerId, JpegBytes)
 }
 
 void FThreadSafePlayerSessions::SendFreezeFrame(const TArray64<uint8>& JpegBytes)
 {
-    SUBMIT_TASK_WITH_PARAMS(SendFreezeFrame_SignallingThread, JpegBytes)
+	SUBMIT_TASK_WITH_PARAMS(SendFreezeFrame_SignallingThread, JpegBytes)
 }
 
 void FThreadSafePlayerSessions::SendUnfreezeFrame()
 {
-    SUBMIT_TASK_NO_PARAMS(SendUnfreezeFrame_SignallingThread)
+	SUBMIT_TASK_NO_PARAMS(SendUnfreezeFrame_SignallingThread)
 }
 
-void FThreadSafePlayerSessions::SendFileData(const TArray<uint8>& ByteData, const FString& MimeType, const FString& FileExtension)
-{
-    SUBMIT_TASK_WITH_PARAMS(SendFileData_SignallingThread, ByteData, MimeType, FileExtension)
+void FThreadSafePlayerSessions::SendFileData(const TArray<uint8>& ByteData, const FString& MimeType, const FString& FileExtension){
+	SUBMIT_TASK_WITH_PARAMS(SendFileData_SignallingThread, ByteData, MimeType, FileExtension)
 }
 
-webrtc::PeerConnectionInterface* FThreadSafePlayerSessions::CreatePlayerSession(FPlayerId PlayerId, 
-    rtc::scoped_refptr<webrtc::PeerConnectionFactoryInterface> PeerConnectionFactory,
-    webrtc::PeerConnectionInterface::RTCConfiguration PeerConnectionConfig,
-    FSignallingServerConnection* SignallingServerConnection,
+webrtc::PeerConnectionInterface* FThreadSafePlayerSessions::CreatePlayerSession(
+    FPlayerId PlayerId,
+	rtc::scoped_refptr<webrtc::PeerConnectionFactoryInterface> PeerConnectionFactory,
+	webrtc::PeerConnectionInterface::RTCConfiguration PeerConnectionConfig,
+	FSignallingServerConnection* SignallingServerConnection,
 	int Flags)
 {
-    SUBMIT_TASK_WITH_PARAMS_AND_RETURN(webrtc::PeerConnectionInterface*, CreatePlayerSession_SignallingThread, PlayerId, PeerConnectionFactory, PeerConnectionConfig, SignallingServerConnection, Flags)
+	SUBMIT_TASK_WITH_PARAMS_AND_RETURN(webrtc::PeerConnectionInterface*, CreatePlayerSession_SignallingThread, PlayerId, PeerConnectionFactory, PeerConnectionConfig, SignallingServerConnection, Flags)
 }
 
 void FThreadSafePlayerSessions::SetPlayerSessionDataChannel(FPlayerId PlayerId, rtc::scoped_refptr<webrtc::DataChannelInterface> DataChannel)
 {
-    SUBMIT_TASK_WITH_PARAMS(SetPlayerSessionDataChannel_SignallingThread, PlayerId, DataChannel)
+	SUBMIT_TASK_WITH_PARAMS(SetPlayerSessionDataChannel_SignallingThread, PlayerId, DataChannel)
 }
 
 void FThreadSafePlayerSessions::DeleteAllPlayerSessions()
 {
-    SUBMIT_TASK_NO_PARAMS(DeleteAllPlayerSessions_SignallingThread)
+	SUBMIT_TASK_NO_PARAMS(DeleteAllPlayerSessions_SignallingThread)
 }
 
 int FThreadSafePlayerSessions::DeletePlayerSession(FPlayerId PlayerId)
 {
-    SUBMIT_TASK_WITH_PARAMS_AND_RETURN(int, DeletePlayerSession_SignallingThread, PlayerId)
+	SUBMIT_TASK_WITH_PARAMS_AND_RETURN(int, DeletePlayerSession_SignallingThread, PlayerId)
 }
 
-void FThreadSafePlayerSessions::DisconnectPlayer(FPlayerId PlayerId, const FString& Reason)
-{
-    SUBMIT_TASK_WITH_PARAMS(DisconnectPlayer_SignallingThread, PlayerId, Reason)
+void FThreadSafePlayerSessions::DisconnectPlayer(FPlayerId PlayerId, const FString& Reason){
+	SUBMIT_TASK_WITH_PARAMS(DisconnectPlayer_SignallingThread, PlayerId, Reason)
 }
 
 FPixelStreamingDataChannelObserver* FThreadSafePlayerSessions::GetDataChannelObserver(FPlayerId PlayerId)
 {
-    SUBMIT_TASK_WITH_PARAMS_AND_RETURN(FPixelStreamingDataChannelObserver*, GetDataChannelObserver_SignallingThread, PlayerId)
+	SUBMIT_TASK_WITH_PARAMS_AND_RETURN(FPixelStreamingDataChannelObserver*, GetDataChannelObserver_SignallingThread, PlayerId)
 }
 
 void FThreadSafePlayerSessions::SendMessageAll(PixelStreamingProtocol::EToPlayerMsg Type, const FString& Descriptor) const
 {
-    SUBMIT_TASK_WITH_PARAMS(SendMessageAll_SignallingThread, Type, Descriptor)
+	SUBMIT_TASK_WITH_PARAMS(SendMessageAll_SignallingThread, Type, Descriptor)
 }
 
 void FThreadSafePlayerSessions::SendLatestQPAllPlayers(int LatestQP) const
 {
-    SUBMIT_TASK_WITH_PARAMS(SendLatestQPAllPlayers_SignallingThread, LatestQP)
+	SUBMIT_TASK_WITH_PARAMS(SendLatestQPAllPlayers_SignallingThread, LatestQP)
 }
 
 void FThreadSafePlayerSessions::OnRemoteIceCandidate(FPlayerId PlayerId, const std::string& SdpMid, int SdpMLineIndex, const std::string& Sdp)
 {
-    SUBMIT_TASK_WITH_PARAMS(OnRemoteIceCandidate_SignallingThread, PlayerId, SdpMid, SdpMLineIndex, Sdp)
+	SUBMIT_TASK_WITH_PARAMS(OnRemoteIceCandidate_SignallingThread, PlayerId, SdpMid, SdpMLineIndex, Sdp)
 }
 
 void FThreadSafePlayerSessions::OnAnswer(FPlayerId PlayerId, FString Sdp)
 {
-    SUBMIT_TASK_WITH_PARAMS(OnAnswer_SignallingThread, PlayerId, Sdp)
+	SUBMIT_TASK_WITH_PARAMS(OnAnswer_SignallingThread, PlayerId, Sdp)
 }
 
 void FThreadSafePlayerSessions::PollWebRTCStats() const
 {
-    SUBMIT_TASK_NO_PARAMS(PollWebRTCStats_SignallingThread)
+	SUBMIT_TASK_NO_PARAMS(PollWebRTCStats_SignallingThread)
 }
 
 /////////////////////////
@@ -194,28 +190,28 @@ void FThreadSafePlayerSessions::PollWebRTCStats() const
 
 void FThreadSafePlayerSessions::PollWebRTCStats_SignallingThread() const
 {
-    checkf(this->IsInSignallingThread(), TEXT("This method must be called on the signalling thread."));
+	checkf(this->IsInSignallingThread(), TEXT("This method must be called on the signalling thread."));
 
-    for (auto& Entry : Players)
+	for (auto& Entry : Players)
 	{
 		FPlayerSession* Player = Entry.Value;
-        if(Player)
-        {
-            Player->PollWebRTCStats();
-        }
-        else
-        {
-            UE_LOG(PixelStreamer, Log, TEXT("Could not poll WebRTC stats for peer because peer was nullptr."));
-        }
-    }
+		if (Player)
+		{
+			Player->PollWebRTCStats();
+		}
+		else
+		{
+			UE_LOG(PixelStreamer, Log, TEXT("Could not poll WebRTC stats for peer because peer was nullptr."));
+		}
+	}
 }
 
 void FThreadSafePlayerSessions::OnRemoteIceCandidate_SignallingThread(FPlayerId PlayerId, const std::string& SdpMid, int SdpMLineIndex, const std::string& Sdp)
 {
-    checkf(this->IsInSignallingThread(), TEXT("This method must be called on the signalling thread."));
+	checkf(this->IsInSignallingThread(), TEXT("This method must be called on the signalling thread."));
 
-    FPlayerSession* Player = this->GetPlayerSession_SignallingThread(PlayerId);
-	if(Player)
+	FPlayerSession* Player = this->GetPlayerSession_SignallingThread(PlayerId);
+	if (Player)
 	{
 		Player->OnRemoteIceCandidate(SdpMid, SdpMLineIndex, Sdp);
 	}
@@ -227,30 +223,30 @@ void FThreadSafePlayerSessions::OnRemoteIceCandidate_SignallingThread(FPlayerId
 
 void FThreadSafePlayerSessions::OnAnswer_SignallingThread(FPlayerId PlayerId, FString Sdp)
 {
-    checkf(this->IsInSignallingThread(), TEXT("This method must be called on the signalling thread."));
+	checkf(this->IsInSignallingThread(), TEXT("This method must be called on the signalling thread."));
 
-    FPlayerSession* Player = this->GetPlayerSession_SignallingThread(PlayerId);
-    if(Player)
-    {
-        Player->OnAnswer(Sdp);
-    }
-    else
-    {
-        UE_LOG(PixelStreamer, Log, TEXT("Could not pass answer to player because Player %s no available."), *PlayerId);
-    }
+	FPlayerSession* Player = this->GetPlayerSession_SignallingThread(PlayerId);
+	if (Player)
+	{
+		Player->OnAnswer(Sdp);
+	}
+	else
+	{
+		UE_LOG(PixelStreamer, Log, TEXT("Could not pass answer to player because Player %s no available."), *PlayerId);
+	}
 }
 
 IPixelStreamingAudioSink* FThreadSafePlayerSessions::GetUnlistenedAudioSink_SignallingThread() const
 {
-    checkf(this->IsInSignallingThread(), TEXT("This method must be called on the signalling thread."));
+	checkf(this->IsInSignallingThread(), TEXT("This method must be called on the signalling thread."));
 
-    for (auto& Entry : Players)
+	for (auto& Entry : Players)
 	{
 		FPlayerSession* Session = Entry.Value;
 		FPixelStreamingAudioSink& AudioSink = Session->GetAudioSink();
-		if(!AudioSink.HasAudioConsumers())
+		if (!AudioSink.HasAudioConsumers())
 		{
-            FPixelStreamingAudioSink* AudioSinkPtr = &AudioSink;
+			FPixelStreamingAudioSink* AudioSinkPtr = &AudioSink;
 			return static_cast<IPixelStreamingAudioSink*>(AudioSinkPtr);
 		}
 	}
@@ -259,34 +255,34 @@ IPixelStreamingAudioSink* FThreadSafePlayerSessions::GetUnlistenedAudioSink_Sign
 
 IPixelStreamingAudioSink* FThreadSafePlayerSessions::GetAudioSink_SignallingThread(FPlayerId PlayerId) const
 {
-    checkf(this->IsInSignallingThread(), TEXT("This method must be called on the signalling thread."));
+	checkf(this->IsInSignallingThread(), TEXT("This method must be called on the signalling thread."));
 
-    FPlayerSession* PlayerSession = this->GetPlayerSession_SignallingThread(PlayerId);
+	FPlayerSession* PlayerSession = this->GetPlayerSession_SignallingThread(PlayerId);
 
-    if(PlayerSession)
-    {
-        FPixelStreamingAudioSink* AudioSink = &PlayerSession->GetAudioSink();
-        return static_cast<IPixelStreamingAudioSink*>(AudioSink);
-    }
-    return nullptr;
+	if (PlayerSession)
+	{
+		FPixelStreamingAudioSink* AudioSink = &PlayerSession->GetAudioSink();
+		return static_cast<IPixelStreamingAudioSink*>(AudioSink);
+	}
+	return nullptr;
 }
 
 void FThreadSafePlayerSessions::SendLatestQPAllPlayers_SignallingThread(int LatestQP) const
 {
-    checkf(this->IsInSignallingThread(), TEXT("This method must be called on the signalling thread."));
+	checkf(this->IsInSignallingThread(), TEXT("This method must be called on the signalling thread."));
 
-    for(auto& PlayerEntry : this->Players)
-    {
-        PlayerEntry.Value->SendVideoEncoderQP(LatestQP);
-    }
+	for (auto& PlayerEntry : this->Players)
+	{
+		PlayerEntry.Value->SendVideoEncoderQP(LatestQP);
+	}
 }
 
 void FThreadSafePlayerSessions::SendLatestQP_SignallingThread(FPlayerId PlayerId, int LatestQP) const
 {
-    checkf(this->IsInSignallingThread(), TEXT("This method must be called on the signalling thread."));
+	checkf(this->IsInSignallingThread(), TEXT("This method must be called on the signalling thread."));
 
 	FPlayerSession* Session = this->GetPlayerSession_SignallingThread(PlayerId);
-	if(Session)
+	if (Session)
 	{
 		return Session->SendVideoEncoderQP(LatestQP);
 	}
@@ -298,76 +294,76 @@ void FThreadSafePlayerSessions::SendLatestQP_SignallingThread(FPlayerId PlayerId
 
 bool FThreadSafePlayerSessions::SendMessage_SignallingThread(FPlayerId PlayerId, PixelStreamingProtocol::EToPlayerMsg Type, const FString& Descriptor) const
 {
-    checkf(this->IsInSignallingThread(), TEXT("This method must be called on the signalling thread."));
+	checkf(this->IsInSignallingThread(), TEXT("This method must be called on the signalling thread."));
 
-    UE_LOG(PixelStreamer, Log, TEXT("SendMessage to: %s | Type: %d | Message: %s"), *PlayerId, static_cast<int32>(Type), *Descriptor);
+	UE_LOG(PixelStreamer, Log, TEXT("SendMessage to: %s | Type: %d | Message: %s"), *PlayerId, static_cast<int32>(Type), *Descriptor);
 
-    FPlayerSession* PlayerSession = this->GetPlayerSession_SignallingThread(PlayerId);
+	FPlayerSession* PlayerSession = this->GetPlayerSession_SignallingThread(PlayerId);
 
-    if(PlayerSession)
-    {
-        return PlayerSession->SendMessage(Type, Descriptor);
-    }
-    else
-    {
-        UE_LOG(PixelStreamer, Log, TEXT("Cannot send message to player: %s - player does not exist."), *PlayerId);
-        return false;
-    }
+	if (PlayerSession)
+	{
+		return PlayerSession->SendMessage(Type, Descriptor);
+	}
+	else
+	{
+		UE_LOG(PixelStreamer, Log, TEXT("Cannot send message to player: %s - player does not exist."), *PlayerId);
+		return false;
+	}
 }
 
 void FThreadSafePlayerSessions::SendMessageAll_SignallingThread(PixelStreamingProtocol::EToPlayerMsg Type, const FString& Descriptor) const
 {
-    checkf(this->IsInSignallingThread(), TEXT("This method must be called on the signalling thread."));
+	checkf(this->IsInSignallingThread(), TEXT("This method must be called on the signalling thread."));
 
-    UE_LOG(PixelStreamer, Log, TEXT("SendMessageAll: %d - %s"), static_cast<int32>(Type), *Descriptor);
+	UE_LOG(PixelStreamer, Log, TEXT("SendMessageAll: %d - %s"), static_cast<int32>(Type), *Descriptor);
 
-    for(auto& PlayerEntry : this->Players)
-    {
-        PlayerEntry.Value->SendMessage(Type, Descriptor);
-    }
+	for (auto& PlayerEntry : this->Players)
+	{
+		PlayerEntry.Value->SendMessage(Type, Descriptor);
+	}
 }
 
 FPixelStreamingDataChannelObserver* FThreadSafePlayerSessions::GetDataChannelObserver_SignallingThread(FPlayerId PlayerId)
 {
-    checkf(this->IsInSignallingThread(), TEXT("This method must be called on the signalling thread."));
+	checkf(this->IsInSignallingThread(), TEXT("This method must be called on the signalling thread."));
 
-    FPlayerSession* Player = this->GetPlayerSession_SignallingThread(PlayerId);
-    if(Player)
-    {
-        return &Player->GetDataChannelObserver();
-    }
-    else
-    {
-        UE_LOG(PixelStreamer, Log, TEXT("Cannot get data channel observer for player: %s - player does not exist."), *PlayerId);
-        return nullptr;
-    }
+	FPlayerSession* Player = this->GetPlayerSession_SignallingThread(PlayerId);
+	if (Player)
+	{
+		return &Player->GetDataChannelObserver();
+	}
+	else
+	{
+		UE_LOG(PixelStreamer, Log, TEXT("Cannot get data channel observer for player: %s - player does not exist."), *PlayerId);
+		return nullptr;
+	}
 }
 
 void FThreadSafePlayerSessions::DisconnectPlayer_SignallingThread(FPlayerId PlayerId, const FString& Reason)
 {
-    checkf(this->IsInSignallingThread(), TEXT("This method must be called on the signalling thread."));
+	checkf(this->IsInSignallingThread(), TEXT("This method must be called on the signalling thread."));
 
-    FPlayerSession* Player = this->GetPlayerSession_SignallingThread(PlayerId);
-    if(Player != nullptr)
-    {
-        Player->DisconnectPlayer(Reason);
-    }
-    else
-    {
-        UE_LOG(PixelStreamer, Log, TEXT("Cannot disconnect player: %s - player does not exist."), *PlayerId);
-    }
+	FPlayerSession* Player = this->GetPlayerSession_SignallingThread(PlayerId);
+	if (Player != nullptr)
+	{
+		Player->DisconnectPlayer(Reason);
+	}
+	else
+	{
+		UE_LOG(PixelStreamer, Log, TEXT("Cannot disconnect player: %s - player does not exist."), *PlayerId);
+	}
 }
 
 int FThreadSafePlayerSessions::GetNumPlayers_SignallingThread() const
 {
-    checkf(this->IsInSignallingThread(), TEXT("This method must be called on the signalling thread."));
+	checkf(this->IsInSignallingThread(), TEXT("This method must be called on the signalling thread."));
 
-    return this->Players.Num();
+	return this->Players.Num();
 }
 
 void FThreadSafePlayerSessions::SendFreezeFrame_SignallingThread(const TArray64<uint8>& JpegBytes)
 {
-    checkf(this->IsInSignallingThread(), TEXT("This method must be called on the signalling thread."));
+	checkf(this->IsInSignallingThread(), TEXT("This method must be called on the signalling thread."));
 
 	UE_LOG(PixelStreamer, Log, TEXT("Sending freeze frame to players: %d bytes"), JpegBytes.Num());
 	{
@@ -380,104 +376,102 @@ void FThreadSafePlayerSessions::SendFreezeFrame_SignallingThread(const TArray64<
 
 void FThreadSafePlayerSessions::SendFileData_SignallingThread(const TArray<uint8>& ByteData, const FString& MimeType, const FString& FileExtension)
 {
-    checkf(this->IsInSignallingThread(), TEXT("This method must be called on the signalling thread."));
+	checkf(this->IsInSignallingThread(), TEXT("This method must be called on the signalling thread."));
 
 	UE_LOG(PixelStreamer, Log, TEXT("Sending file with Mime type %s to all players: %d bytes"), *MimeType, ByteData.Num());
-    {
-        for(auto& PlayerEntry : this->Players)
-        {
-            PlayerEntry.Value->SendFileData(ByteData, MimeType, FileExtension);
-        }
-    }
+	{
+		for (auto& PlayerEntry : this->Players)
+		{
+			PlayerEntry.Value->SendFileData(ByteData, MimeType, FileExtension);
+		}
+	}
 }
 
 void FThreadSafePlayerSessions::SendUnfreezeFrame_SignallingThread()
 {
-    checkf(this->IsInSignallingThread(), TEXT("This method must be called on the signalling thread."));
+	checkf(this->IsInSignallingThread(), TEXT("This method must be called on the signalling thread."));
 
-    UE_LOG(PixelStreamer, Log, TEXT("Sending unfreeze message to players"));
+	UE_LOG(PixelStreamer, Log, TEXT("Sending unfreeze message to players"));
 
-    for (auto& PlayerEntry : Players)
-    {
-        PlayerEntry.Value->SendUnfreezeFrame();
-    }
+	for (auto& PlayerEntry : Players)
+	{
+		PlayerEntry.Value->SendUnfreezeFrame();
+	}
 }
 
 void FThreadSafePlayerSessions::SendFreezeFrameTo_SignallingThread(FPlayerId PlayerId, const TArray64<uint8>& JpegBytes) const
 {
-    checkf(this->IsInSignallingThread(), TEXT("This method must be called on the signalling thread."));
+	checkf(this->IsInSignallingThread(), TEXT("This method must be called on the signalling thread."));
 
-    FPlayerSession* Player = this->GetPlayerSession_SignallingThread(PlayerId);
+	FPlayerSession* Player = this->GetPlayerSession_SignallingThread(PlayerId);
 
-    if(Player != nullptr)
-    {
-        this->Players[PlayerId]->SendFreezeFrame(JpegBytes);
-    }
-    else
-    {
-        UE_LOG(PixelStreamer, Log, TEXT("Cannot send freeze frame to player: %s - player does not exist."), *PlayerId);
-    }
+	if (Player != nullptr)
+	{
+		this->Players[PlayerId]->SendFreezeFrame(JpegBytes);
+	}
+	else
+	{
+		UE_LOG(PixelStreamer, Log, TEXT("Cannot send freeze frame to player: %s - player does not exist."), *PlayerId);
+	}
 }
 
 FPlayerSession* FThreadSafePlayerSessions::GetPlayerSession_SignallingThread(FPlayerId PlayerId) const
 {
-    checkf(this->IsInSignallingThread(), TEXT("This method must be called on the signalling thread."));
+	checkf(this->IsInSignallingThread(), TEXT("This method must be called on the signalling thread."));
 
-    if(this->Players.Contains(PlayerId))
-    {
-        return this->Players[PlayerId];
-    }
-    return nullptr;
+	if (this->Players.Contains(PlayerId))
+	{
+		return this->Players[PlayerId];
+	}
+	return nullptr;
 }
 
 void FThreadSafePlayerSessions::SetPlayerSessionDataChannel_SignallingThread(FPlayerId PlayerId, rtc::scoped_refptr<webrtc::DataChannelInterface> DataChannel)
 {
-    FPlayerSession* Player = this->GetPlayerSession_SignallingThread(PlayerId);
-    if(Player)
-    {
-        Player->SetDataChannel(DataChannel);
-    }
+	FPlayerSession* Player = this->GetPlayerSession_SignallingThread(PlayerId);
+	if (Player)
+	{
+		Player->SetDataChannel(DataChannel);
+	}
 }
 
 void FThreadSafePlayerSessions::DeleteAllPlayerSessions_SignallingThread()
 {
-    checkf(this->IsInSignallingThread(), TEXT("This method must be called on the signalling thread."));
+	checkf(this->IsInSignallingThread(), TEXT("This method must be called on the signalling thread."));
 
-    UPixelStreamerDelegates* Delegates = UPixelStreamerDelegates::GetPixelStreamerDelegates();
+	UPixelStreamerDelegates* Delegates = UPixelStreamerDelegates::GetPixelStreamerDelegates();
 
-    // Delete each player
-    for(auto& Entry : this->Players)
-    {
-        FPlayerId PlayerId = Entry.Key;
-        FPlayerSession* Player = Entry.Value;
-        delete Player;
-        bool bWasQualityController = this->QualityControllingPlayer == PlayerId;
+	// Delete each player
+	for (auto& Entry : this->Players)
+	{
+		FPlayerId PlayerId = Entry.Key;
+		FPlayerSession* Player = Entry.Value;
+		delete Player;
+		bool bWasQualityController = this->QualityControllingPlayer == PlayerId;
 
-        if(Delegates && FModuleManager::Get().IsModuleLoaded("PixelStreaming"))
-        {
-            Delegates->OnClosedConnection.Broadcast(PlayerId, bWasQualityController);
-        }
+		if (Delegates && FModuleManager::Get().IsModuleLoaded("PixelStreaming"))
+		{
+			Delegates->OnClosedConnection.Broadcast(PlayerId, bWasQualityController);
+		}
 
-        // Player deleted, tell all our C++ listeners.
-        this->OnPlayerDeleted.Broadcast(PlayerId);
-    }
+		// Player deleted, tell all our C++ listeners.
+		this->OnPlayerDeleted.Broadcast(PlayerId);
+	}
 
-    this->Players.Empty();
-    this->QualityControllingPlayer = INVALID_PLAYER_ID;
+	this->Players.Empty();
+	this->QualityControllingPlayer = INVALID_PLAYER_ID;
 
-    if(Delegates && FModuleManager::Get().IsModuleLoaded("PixelStreaming"))
-    {
-        Delegates->OnAllConnectionsClosed.Broadcast();
-    }
+	if (Delegates && FModuleManager::Get().IsModuleLoaded("PixelStreaming"))
+	{
+		Delegates->OnAllConnectionsClosed.Broadcast();
+	}
 }
 
-
-
 int FThreadSafePlayerSessions::DeletePlayerSession_SignallingThread(FPlayerId PlayerId)
 {
-    checkf(this->IsInSignallingThread(), TEXT("This method must be called on the signalling thread."));
+	checkf(this->IsInSignallingThread(), TEXT("This method must be called on the signalling thread."));
 
-    FPlayerSession* Player = this->GetPlayerSession_SignallingThread(PlayerId);
+	FPlayerSession* Player = this->GetPlayerSession_SignallingThread(PlayerId);
 	if (!Player)
 	{
 		UE_LOG(PixelStreamer, VeryVerbose, TEXT("Failed to delete player %s - that player was not found."), *PlayerId);
@@ -486,8 +480,8 @@ int FThreadSafePlayerSessions::DeletePlayerSession_SignallingThread(FPlayerId Pl
 
 	bool bWasQualityController = this->IsQualityController(PlayerId);
 
-    // The actual modification to the players map
-    Players.Remove(PlayerId);
+	// The actual modification to the players map
+	Players.Remove(PlayerId);
 	delete Player;
 
 	UPixelStreamerDelegates* Delegates = UPixelStreamerDelegates::GetPixelStreamerDelegates();
@@ -496,8 +490,8 @@ int FThreadSafePlayerSessions::DeletePlayerSession_SignallingThread(FPlayerId Pl
 		Delegates->OnClosedConnection.Broadcast(PlayerId, bWasQualityController);
 	}
 
-    // Player deleted, tell all our C++ listeners.
-    this->OnPlayerDeleted.Broadcast(PlayerId);
+	// Player deleted, tell all our C++ listeners.
+	this->OnPlayerDeleted.Broadcast(PlayerId);
 
 	// this is called from WebRTC signalling thread, the only thread were `Players` map is modified, so no need to lock it
 	if (Players.Num() == 0)
@@ -512,36 +506,36 @@ int FThreadSafePlayerSessions::DeletePlayerSession_SignallingThread(FPlayerId Pl
 	else if (bWasQualityController)
 	{
 		// Quality Controller session has been just removed, set quality control to any of remaining sessions
-		for(auto& Entry : this->Players)
-        {
-            this->SetQualityController_SignallingThread(Entry.Key);
-            break;
-        }
-		
+		for (auto& Entry : this->Players)
+		{
+			this->SetQualityController_SignallingThread(Entry.Key);
+			break;
+		}
 	}
 
-    return this->GetNumPlayers_SignallingThread();
+	return this->GetNumPlayers_SignallingThread();
 }
 
-webrtc::PeerConnectionInterface* FThreadSafePlayerSessions::CreatePlayerSession_SignallingThread(FPlayerId PlayerId, 
-    rtc::scoped_refptr<webrtc::PeerConnectionFactoryInterface> PeerConnectionFactory, 
-    webrtc::PeerConnectionInterface::RTCConfiguration PeerConnectionConfig,
-    FSignallingServerConnection* SignallingServerConnection,
-    int Flags)
+webrtc::PeerConnectionInterface* FThreadSafePlayerSessions::CreatePlayerSession_SignallingThread(
+    FPlayerId PlayerId,
+	rtc::scoped_refptr<webrtc::PeerConnectionFactoryInterface> PeerConnectionFactory,
+	webrtc::PeerConnectionInterface::RTCConfiguration PeerConnectionConfig,
+	FSignallingServerConnection* SignallingServerConnection,
+	int Flags)
 {
-    checkf(this->IsInSignallingThread(), TEXT("This method must be called on the signalling thread."));
+	checkf(this->IsInSignallingThread(), TEXT("This method must be called on the signalling thread."));
 	check(PeerConnectionFactory);
 
 	// With unified plan, we get several calls to OnOffer, which in turn calls
 	// this several times.
 	// Therefore, we only try to create the player if not created already
 	if (Players.Contains(PlayerId))
-    {
-        return nullptr;
-    }
+	{
+		return nullptr;
+	}
 
 	UE_LOG(PixelStreamer, Log, TEXT("Creating player session for PlayerId=%s"), *PlayerId);
-	
+
 	// this is called from WebRTC signalling thread, the only thread were `Players` map is modified, so no need to lock it
 	bool bMakeQualityController = Players.Num() == 0; // first player controls quality by default
 	FPlayerSession* Session = new FPlayerSession(this, SignallingServerConnection, PlayerId);
@@ -558,12 +552,12 @@ webrtc::PeerConnectionInterface* FThreadSafePlayerSessions::CreatePlayerSession_
 
 	Session->SetPeerConnection(PeerConnection);
 
-    // The actual modification of the players map
+	// The actual modification of the players map
 	Players.Add(PlayerId, Session);
 
-	if(bMakeQualityController)
+	if (bMakeQualityController)
 	{
-        this->SetQualityController_SignallingThread(PlayerId);
+		this->SetQualityController_SignallingThread(PlayerId);
 	}
 
 	if (UPixelStreamerDelegates* Delegates = UPixelStreamerDelegates::GetPixelStreamerDelegates())
@@ -571,24 +565,24 @@ webrtc::PeerConnectionInterface* FThreadSafePlayerSessions::CreatePlayerSession_
 		Delegates->OnNewConnection.Broadcast(PlayerId, bMakeQualityController);
 	}
 
-    return PeerConnection.get();
+	return PeerConnection.get();
 }
 
 void FThreadSafePlayerSessions::SetQualityController_SignallingThread(FPlayerId PlayerId)
 {
-    checkf(this->IsInSignallingThread(), TEXT("This method must be called on the signalling thread."));
+	checkf(this->IsInSignallingThread(), TEXT("This method must be called on the signalling thread."));
 
 	if (this->Players.Contains(PlayerId))
 	{
 
-        // The actual assignment of the quality controlling peeer
-        {
-            FScopeLock Lock(&QualityControllerCS);
-            this->QualityControllingPlayer = PlayerId;
-        }
-		
-        // Let any listeners know the quality controller has changed
-        this->OnQualityControllerChanged.Broadcast(PlayerId);
+		// The actual assignment of the quality controlling peeer
+		{
+			FScopeLock Lock(&QualityControllerCS);
+			this->QualityControllingPlayer = PlayerId;
+		}
+
+		// Let any listeners know the quality controller has changed
+		this->OnQualityControllerChanged.Broadcast(PlayerId);
 
 		UE_LOG(PixelStreamer, Log, TEXT("Quality controller is now PlayerId=%s."), *PlayerId);
 
@@ -596,18 +590,17 @@ void FThreadSafePlayerSessions::SetQualityController_SignallingThread(FPlayerId
 		for (auto& Entry : Players)
 		{
 			FPlayerSession* Session = Entry.Value;
-            bool bIsQualityController = Entry.Key == this->QualityControllingPlayer;
-			if(Session)
+			bool bIsQualityController = Entry.Key == this->QualityControllingPlayer;
+			if (Session)
 			{
 				Session->SendQualityControlStatus(bIsQualityController);
 			}
 		}
-
 	}
-    else
-    {
-        UE_LOG(PixelStreamer, Log, TEXT("Could not set quality controller for PlayerId=%s - that player does not exist."), *PlayerId);
-    }
+	else
+	{
+		UE_LOG(PixelStreamer, Log, TEXT("Could not set quality controller for PlayerId=%s - that player does not exist."), *PlayerId);
+	}
 }
 
 #undef SUBMIT_TASK_WITH_RETURN
diff --git a/Engine/Plugins/Media/PixelStreaming/Source/PixelStreaming/Private/WebRTCLogging.cpp b/Engine/Plugins/Media/PixelStreaming/Source/PixelStreaming/Private/WebRTCLogging.cpp
--- a/Engine/Plugins/Media/PixelStreaming/Source/PixelStreaming/Private/WebRTCLogging.cpp
+++ b/Engine/Plugins/Media/PixelStreaming/Source/PixelStreaming/Private/WebRTCLogging.cpp
@@ -41,8 +41,7 @@ private:
 		const char* tag) override
 	{
 #if !NO_LOGGING
-		static const ELogVerbosity::Type RtcToUnrealLogCategoryMap[] =
-		{
+		static const ELogVerbosity::Type RtcToUnrealLogCategoryMap[] = {
 			ELogVerbosity::VeryVerbose,
 			ELogVerbosity::Verbose,
 			ELogVerbosity::Log,
@@ -70,18 +69,26 @@ private:
 
 		switch (severity)
 		{
-		case rtc::LS_VERBOSE:
-			UE_LOG(PixelStreamingWebRTC, Verbose, TEXT("%s"), *Msg);
-			break;
-		case rtc::LS_INFO:
-			UE_LOG(PixelStreamingWebRTC, Log, TEXT("%s"), *Msg);
-			break;
-		case rtc::LS_WARNING:
-			UE_LOG(PixelStreamingWebRTC, Warning, TEXT("%s"), *Msg);
-			break;
-		case rtc::LS_ERROR:
-			UE_LOG(PixelStreamingWebRTC, Error, TEXT("%s"), *Msg);
-			break;
+			case rtc::LS_VERBOSE:
+			{
+				UE_LOG(PixelStreamingWebRTC, Verbose, TEXT("%s"), *Msg);
+				break;
+			}
+			case rtc::LS_INFO:
+			{
+				UE_LOG(PixelStreamingWebRTC, Log, TEXT("%s"), *Msg);
+				break;
+			}
+			case rtc::LS_WARNING:
+			{
+				UE_LOG(PixelStreamingWebRTC, Warning, TEXT("%s"), *Msg);
+				break;
+			}
+			case rtc::LS_ERROR:
+			{
+				UE_LOG(PixelStreamingWebRTC, Error, TEXT("%s"), *Msg);
+				break;
+			}
 		}
 #endif
 	}
