Merging Pixel Streaming refactor using //UE5/Dev-TensorWorks-to-//UE5/Release-5.0

diff --git a/Engine/Plugins/Media/PixelStreaming/Source/PixelStreaming/Private/AudioDeviceModule.h b/Engine/Plugins/Media/PixelStreaming/Source/PixelStreaming/Private/AudioDeviceModule.h
new file mode 100644
--- /dev/null
+++ b/Engine/Plugins/Media/PixelStreaming/Source/PixelStreaming/Private/AudioDeviceModule.h
@@ -0,0 +1,314 @@
+// Copyright Epic Games, Inc. All Rights Reserved.
+
+#pragma once
+
+#include "PixelStreamingPrivate.h"
+#include "WebRTCIncludes.h"
+#include "AudioMixerDevice.h"
+#include "HAL/CriticalSection.h"
+#include "GenericPlatform/GenericPlatformMisc.h"
+#include "Templates/Function.h"
+#include "Templates/SharedPointer.h"
+
+namespace UE {
+	namespace PixelStreaming {
+		// Captures audio from UE and passes it along to WebRTC.
+		class FSubmixCapturer : public ISubmixBufferListener
+		{
+
+		public:
+			// This magic number is the max volume used in webrtc fake audio device module.
+			static const uint32_t MaxVolumeLevel = 14392;
+
+			FSubmixCapturer()
+				: bInitialised(false)
+				, bCapturing(false)
+				, TargetSampleRate(48000)
+				, TargetNumChannels(2)
+				, bReportedSampleRateMismatch(false)
+				, AudioCallback(nullptr)
+				, VolumeLevel(FSubmixCapturer::MaxVolumeLevel)
+				, RecordingBuffer()
+				, CriticalSection()
+			{
+			}
+
+			virtual ~FSubmixCapturer() = default;
+
+			bool Init();
+			bool IsInitialised() const;
+			bool IsCapturing() const;
+			void Uninitialise();
+			bool StartCapturing();
+			bool EndCapturing();
+			uint32_t GetVolume() const;
+			void SetVolume(uint32_t NewVolume);
+
+			// ISubmixBufferListener interface
+			void OnNewSubmixBuffer(const USoundSubmix* OwningSubmix, float* AudioData, int32 NumSamples, int32 NumChannels, const int32 SampleRate, double AudioClock) override;
+
+			void RegisterAudioCallback(webrtc::AudioTransport* AudioCb);
+
+		private:
+			int32 GetSamplesPerDurationSecs(float InSeconds) const;
+
+		private:
+			FThreadSafeBool bInitialised;
+			FThreadSafeBool bCapturing;
+			uint32 TargetSampleRate;
+			uint8 TargetNumChannels;
+			bool bReportedSampleRateMismatch;
+			webrtc::AudioTransport* AudioCallback;
+			uint32_t VolumeLevel;
+			TArray<int16_t> RecordingBuffer;
+			FCriticalSection CriticalSection; // One thread captures audio from UE, the other controls the state of the capturer.
+		};
+
+		// Requests audio from WebRTC at a regular interval (10ms)
+		// This is required so that WebRTC audio sinks actually have
+		// some audio data for their sinks. Without this WebRTC assumes
+		// there is no demand for audio and does not populate the sinks.
+		class FAudioPlayoutRequester
+		{
+		public:
+			class Runnable : public FRunnable
+			{
+			public:
+				Runnable(TFunction<void()> RequestPlayoutFunc)
+					: bIsRunning(false)
+					, LastAudioRequestTimeMs(0)
+					, RequestPlayoutFunc(RequestPlayoutFunc) {}
+				// Begin FRunnable interface.
+				virtual bool Init() override;
+				virtual uint32 Run() override;
+				virtual void Stop() override;
+				virtual void Exit() override;
+				// End FRunnable interface
+			private:
+				bool bIsRunning;
+				int64_t LastAudioRequestTimeMs;
+				TFunction<void()> RequestPlayoutFunc;
+			};
+
+			FAudioPlayoutRequester()
+				: bIsPlayoutInitialised(false)
+				, bIsPlaying(false)
+				, SampleRate(48000) //webrtc will mix all sources into one big buffer with this sample rate
+				, NumChannels(2)	//webrtc will mix all sources into one big buffer with this many channels
+				, RequesterRunnable()
+				, RequesterThread()
+				, AudioCallback(nullptr)
+				, PlayoutCriticalSection()
+				, PlayoutBuffer(){};
+
+			void InitPlayout();
+			void StartPlayout();
+			void StopPlayout();
+			bool Playing() const;
+			bool PlayoutIsInitialized() const;
+			void Uninitialise();
+			void RegisterAudioCallback(webrtc::AudioTransport* AudioCallback);
+
+		public:
+			static int16_t const RequestIntervalMs = 10;
+
+		private:
+			FThreadSafeBool bIsPlayoutInitialised;
+			FThreadSafeBool bIsPlaying;
+			uint32 SampleRate;
+			uint8 NumChannels;
+			TUniquePtr<FAudioPlayoutRequester::Runnable> RequesterRunnable;
+			TUniquePtr<FRunnableThread> RequesterThread;
+			webrtc::AudioTransport* AudioCallback;
+			FCriticalSection PlayoutCriticalSection;
+			TArray<int16_t> PlayoutBuffer;
+		};
+
+		// A custom audio device module for WebRTC.
+		// It lets us receive WebRTC audio in UE
+		// and also transmit UE audio to WebRTC.
+		class FAudioDeviceModule : public webrtc::AudioDeviceModule
+		{
+
+		public:
+			FAudioDeviceModule()
+				: bInitialized(false)
+				, Capturer()
+				, Requester(){};
+
+			virtual ~FAudioDeviceModule() = default;
+
+		private:
+			FThreadSafeBool bInitialized; // True when we setup capturer/playout requester.
+			FSubmixCapturer Capturer;
+			FAudioPlayoutRequester Requester;
+
+		private:
+			//
+			// webrtc::AudioDeviceModule interface
+			//
+			int32 ActiveAudioLayer(AudioLayer* audioLayer) const override;
+			int32 RegisterAudioCallback(webrtc::AudioTransport* audioCallback) override;
+
+			// Main initialization and termination
+			int32 Init() override;
+			int32 Terminate() override;
+			bool Initialized() const override;
+
+			// Device enumeration
+			int16 PlayoutDevices() override;
+			int16 RecordingDevices() override;
+			int32 PlayoutDeviceName(uint16 index, char name[webrtc::kAdmMaxDeviceNameSize], char guid[webrtc::kAdmMaxGuidSize]) override;
+			int32 RecordingDeviceName(uint16 index, char name[webrtc::kAdmMaxDeviceNameSize], char guid[webrtc::kAdmMaxGuidSize]) override;
+
+			// Device selection
+			int32 SetPlayoutDevice(uint16 index) override;
+			int32 SetPlayoutDevice(WindowsDeviceType device) override;
+			int32 SetRecordingDevice(uint16 index) override;
+			int32 SetRecordingDevice(WindowsDeviceType device) override;
+
+			// Audio transport initialization
+			int32 PlayoutIsAvailable(bool* available) override;
+			int32 InitPlayout() override;
+			bool PlayoutIsInitialized() const override;
+			int32 RecordingIsAvailable(bool* available) override;
+			int32 InitRecording() override;
+			bool RecordingIsInitialized() const override;
+
+			// Audio transport control
+			virtual int32 StartPlayout() override;
+			virtual int32 StopPlayout() override;
+
+			// True when audio is being pulled by the instance.
+			virtual bool Playing() const override;
+
+			virtual int32 StartRecording() override;
+			virtual int32 StopRecording() override;
+			virtual bool Recording() const override;
+
+			// Audio mixer initialization
+			virtual int32 InitSpeaker() override;
+			virtual bool SpeakerIsInitialized() const override;
+			virtual int32 InitMicrophone() override;
+			virtual bool MicrophoneIsInitialized() const override;
+
+			// Speaker volume controls
+			virtual int32 SpeakerVolumeIsAvailable(bool* available) override
+			{
+				return -1;
+			}
+			virtual int32 SetSpeakerVolume(uint32 volume) override
+			{
+				return -1;
+			}
+			virtual int32 SpeakerVolume(uint32* volume) const override
+			{
+				return -1;
+			}
+			virtual int32 MaxSpeakerVolume(uint32* maxVolume) const override
+			{
+				return -1;
+			}
+			virtual int32 MinSpeakerVolume(uint32* minVolume) const override
+			{
+				return -1;
+			}
+
+			// Microphone volume controls
+			virtual int32 MicrophoneVolumeIsAvailable(bool* available) override
+			{
+				return 0;
+			}
+			virtual int32 SetMicrophoneVolume(uint32 volume) override
+			{
+				//Capturer.SetVolume(volume);
+				return 0;
+			}
+			virtual int32 MicrophoneVolume(uint32* volume) const override
+			{
+				//*volume = Capturer.GetVolume();
+				return 0;
+			}
+			virtual int32 MaxMicrophoneVolume(uint32* maxVolume) const override
+			{
+				*maxVolume = FSubmixCapturer::MaxVolumeLevel;
+				return 0;
+			}
+			virtual int32 MinMicrophoneVolume(uint32* minVolume) const override
+			{
+				return 0;
+			}
+
+			// Speaker mute control
+			virtual int32 SpeakerMuteIsAvailable(bool* available) override
+			{
+				return -1;
+			}
+			virtual int32 SetSpeakerMute(bool enable) override
+			{
+				return -1;
+			}
+			virtual int32 SpeakerMute(bool* enabled) const override
+			{
+				return -1;
+			}
+
+			// Microphone mute control
+			virtual int32 MicrophoneMuteIsAvailable(bool* available) override
+			{
+				*available = false;
+				return -1;
+			}
+			virtual int32 SetMicrophoneMute(bool enable) override
+			{
+				return -1;
+			}
+			virtual int32 MicrophoneMute(bool* enabled) const override
+			{
+				return -1;
+			}
+
+			// Stereo support
+			virtual int32 StereoPlayoutIsAvailable(bool* available) const override;
+			virtual int32 SetStereoPlayout(bool enable) override;
+			virtual int32 StereoPlayout(bool* enabled) const override;
+			virtual int32 StereoRecordingIsAvailable(bool* available) const override;
+			virtual int32 SetStereoRecording(bool enable) override;
+			virtual int32 StereoRecording(bool* enabled) const override;
+
+			// Playout delay
+			virtual int32 PlayoutDelay(uint16* delayMS) const override
+			{
+				*delayMS = 0;
+				return 0;
+			}
+
+			virtual bool BuiltInAECIsAvailable() const override
+			{
+				return false;
+			}
+			virtual bool BuiltInAGCIsAvailable() const override
+			{
+				return false;
+			}
+			virtual bool BuiltInNSIsAvailable() const override
+			{
+				return false;
+			}
+
+			// Enables the built-in audio effects. Only supported on Android.
+			virtual int32 EnableBuiltInAEC(bool enable) override
+			{
+				return -1;
+			}
+			virtual int32 EnableBuiltInAGC(bool enable) override
+			{
+				return -1;
+			}
+			virtual int32 EnableBuiltInNS(bool enable) override
+			{
+				return -1;
+			}
+		};
+	}
+}
diff --git a/Engine/Plugins/Media/PixelStreaming/Source/PixelStreaming/Private/AudioSink.cpp b/Engine/Plugins/Media/PixelStreaming/Source/PixelStreaming/Private/AudioSink.cpp
new file mode 100644
--- /dev/null
+++ b/Engine/Plugins/Media/PixelStreaming/Source/PixelStreaming/Private/AudioSink.cpp
@@ -0,0 +1,62 @@
+// Copyright Epic Games, Inc. All Rights Reserved.
+
+#include "AudioSink.h"
+
+UE::PixelStreaming::FAudioSink::~FAudioSink()
+{
+	for (auto Iter = AudioConsumers.CreateIterator(); Iter; ++Iter)
+	{
+		IPixelStreamingAudioConsumer* AudioConsumer = Iter.ElementIt->Value;
+		Iter.RemoveCurrent();
+		if (AudioConsumer != nullptr)
+		{
+			AudioConsumer->OnConsumerRemoved();
+		}
+	}
+}
+
+void UE::PixelStreaming::FAudioSink::OnData(const void* audio_data, int bits_per_sample, int sample_rate, size_t number_of_channels, size_t number_of_frames, absl::optional<int64_t> absolute_capture_timestamp_ms)
+{
+
+	// This data is populated from the internals of WebRTC, basically each audio track sent from the browser has its RTP audio source received and decoded.
+	// The sample rate and number of channels here has absolutely no relationship with PixelStreamingAudioDeviceModule.
+	// The sample rate and number of channels here is determined adaptively by WebRTC's NetEQ class that selects sample rate/number of channels
+	// based on network conditions and other factors.
+
+	if (!HasAudioConsumers())
+	{
+		return;
+	}
+
+	// Iterate audio consumers and pass this data to their buffers
+	for (IPixelStreamingAudioConsumer* AudioConsumer : AudioConsumers)
+	{
+		AudioConsumer->ConsumeRawPCM(static_cast<const int16_t*>(audio_data), sample_rate, number_of_channels, number_of_frames);
+	}
+
+	return;
+}
+
+void UE::PixelStreaming::FAudioSink::AddAudioConsumer(IPixelStreamingAudioConsumer* AudioConsumer)
+{
+	bool bAlreadyInSet = false;
+	AudioConsumers.Add(AudioConsumer, &bAlreadyInSet);
+	if (!bAlreadyInSet)
+	{
+		AudioConsumer->OnConsumerAdded();
+	}
+}
+
+void UE::PixelStreaming::FAudioSink::RemoveAudioConsumer(IPixelStreamingAudioConsumer* AudioConsumer)
+{
+	if (AudioConsumers.Contains(AudioConsumer))
+	{
+		AudioConsumers.Remove(AudioConsumer);
+		AudioConsumer->OnConsumerRemoved();
+	}
+}
+
+bool UE::PixelStreaming::FAudioSink::HasAudioConsumers()
+{
+	return AudioConsumers.Num() > 0;
+}
\ No newline at end of file
diff --git a/Engine/Plugins/Media/PixelStreaming/Source/PixelStreaming/Private/AudioSink.h b/Engine/Plugins/Media/PixelStreaming/Source/PixelStreaming/Private/AudioSink.h
--- a/Engine/Plugins/Media/PixelStreaming/Source/PixelStreaming/Private/AudioSink.h
+++ b/Engine/Plugins/Media/PixelStreaming/Source/PixelStreaming/Private/AudioSink.h
@@ -1,100 +1,39 @@
 // Copyright Epic Games, Inc. All Rights Reserved.
-
 #pragma once
 
 #include "WebRTCIncludes.h"
-
-#include "Templates/Function.h"
-#include "Containers/Array.h"
-#include "IMediaAudioSample.h"
-#include "MediaObjectPool.h"
-#include "Misc/AssertionMacros.h"
-#include "HAL/PlatformTime.h"
-
-using FAudioSampleRef = TSharedRef<IMediaAudioSample, ESPMode::ThreadSafe>;
-
-class FAudioSample :
-	public IMediaAudioSample,
-	public IMediaPoolable
-{
-public:
-	void Init(const void* AudioData, uint32 BitsPerSample, uint32 InSampleRate, uint32 InNumChannels, uint32 InNumFrames)
-	{
-		check(BitsPerSample == 16); // `EMediaAudioSampleFormat::Int16` is hardcoded
-
-		NumChannels = InNumChannels;
-		SampleRate = InSampleRate;
-		NumFrames = InNumFrames;
-		Timestamp = FTimespan::FromSeconds(FPlatformTime::Seconds());
-		Buffer = TArray<uint8>(reinterpret_cast<const uint8*>(AudioData), NumChannels * sizeof(int16) * InNumFrames);
-
-		UE_LOG(PixelPlayer, VeryVerbose, TEXT("Audio: %dc, %dHz, %d samples"), NumChannels, SampleRate, NumFrames);
-	}
-
-	const void* GetBuffer() override
-	{
-		return Buffer.GetData();
-	}
-
-	uint32 GetChannels() const override
-	{
-		return NumChannels;
-	}
-
-	FTimespan GetDuration() const override
-	{
-		return FTimespan::FromSeconds(static_cast<double>(NumFrames) / SampleRate);
-	}
-
-	EMediaAudioSampleFormat GetFormat() const override
-	{
-		return EMediaAudioSampleFormat::Int16;
-	}
-
-	uint32 GetFrames() const override
-	{
-		return NumFrames;
-	}
-
-	uint32 GetSampleRate() const override
-	{
-		return SampleRate;
-	}
-
-	FMediaTimeStamp GetTime() const override
-	{
-		return Timestamp;
-	}
-
-private:
-	uint32 NumChannels;
-	uint32 SampleRate;
-	uint32 NumFrames;
-	FMediaTimeStamp Timestamp;
-	TArray<uint8> Buffer;
-};
-
-using FAudioSamplePool = TMediaObjectPool<FAudioSample>;
-
-class FAudioSink : public webrtc::AudioTrackSinkInterface
-{
-public:
-	using FDelegate = TUniqueFunction<void(const FAudioSampleRef&)>;
-
-	explicit FAudioSink(FDelegate&& InDelegate)
-		: Delegate(MoveTemp(InDelegate))
-	{
-	}
-
-private:
-	void OnData(const void* AudioData, int32 BitsPerSample, int32 SampleRate, SIZE_T NumberOfChannels, SIZE_T NumberOfFrames) override
-	{
-		TSharedRef<FAudioSample, ESPMode::ThreadSafe> Sample = SamplePool.AcquireShared();
-		Sample->Init(AudioData, BitsPerSample, SampleRate, NumberOfChannels, NumberOfFrames);
-		Delegate(Sample);
-	}
-
-private:
-	FDelegate Delegate;
-	FAudioSamplePool SamplePool;
-};
+#include "IPixelStreamingAudioSink.h"
+#include "IPixelStreamingAudioConsumer.h"
+#include "Containers/Set.h"
+
+namespace UE {
+	namespace PixelStreaming {
+		// This is a PixelStreaming sink AND a WebRTC sink.
+		// It collects audio coming in from a WebRTC audio source and passes into into UE's audio system.
+		class FAudioSink : public webrtc::AudioTrackSinkInterface, public IPixelStreamingAudioSink
+		{
+		public:
+			FAudioSink()
+				: AudioConsumers(){};
+
+			// Note: destructor will call destroy on any attached audio consumers
+			virtual ~FAudioSink();
+
+			// Begin AudioTrackSinkInterface
+			void OnData(const void* audio_data, int bits_per_sample, int sample_rate, size_t number_of_channels, size_t number_of_frames, absl::optional<int64_t> absolute_capture_timestamp_ms) override;
+			// End AudioTrackSinkInterface
+
+			// Begin IPixelStreamingAudioSink
+			void AddAudioConsumer(IPixelStreamingAudioConsumer* AudioConsumer) override;
+			void RemoveAudioConsumer(IPixelStreamingAudioConsumer* AudioConsumer) override;
+			bool HasAudioConsumers() override;
+			// End IPixelStreamingAudioSink
+
+		private:
+			void UpdateAudioSettings(int InNumChannels, int InSampleRate);
+
+		private:
+			TSet<IPixelStreamingAudioConsumer*> AudioConsumers;
+		};
+	}
+}
diff --git a/Engine/Plugins/Media/PixelStreaming/Source/PixelStreaming/Private/DataChannelObserver.cpp b/Engine/Plugins/Media/PixelStreaming/Source/PixelStreaming/Private/DataChannelObserver.cpp
new file mode 100644
--- /dev/null
+++ b/Engine/Plugins/Media/PixelStreaming/Source/PixelStreaming/Private/DataChannelObserver.cpp
@@ -0,0 +1,172 @@
+// Copyright Epic Games, Inc. All Rights Reserved.
+
+#include "DataChannelObserver.h"
+#include "Async/Async.h"
+#include "PixelStreamingPrivate.h"
+#include "Settings.h"
+#include "InputDevice.h"
+#include "Modules/ModuleManager.h"
+#include "IPixelStreamingModule.h"
+#include "PixelStreamingModule.h"
+#include "HAL/PlatformTime.h"
+#include "IPixelStreamingSessions.h"
+#include "Stats.h"
+#include "PixelStreamingStatNames.h"
+#include "PixelStreamingDelegates.h"
+
+UE::PixelStreaming::FDataChannelObserver::FDataChannelObserver(UE::PixelStreaming::IPixelStreamingSessions* InPlayerSessions, FPixelStreamingPlayerId InPlayerId)
+	: PlayerSessions(InPlayerSessions)
+	, PlayerId(InPlayerId)
+	, InputDevice(static_cast<FInputDevice&>(IPixelStreamingModule::Get().GetInputDevice()))
+{
+}
+
+void UE::PixelStreaming::FDataChannelObserver::Register(rtc::scoped_refptr<webrtc::DataChannelInterface> InDataChannel)
+{
+	DataChannel = InDataChannel;
+	DataChannel->RegisterObserver(this);
+}
+
+void UE::PixelStreaming::FDataChannelObserver::Unregister()
+{
+	if (DataChannel)
+	{
+		DataChannel->UnregisterObserver();
+		DataChannel = nullptr;
+	}
+}
+
+void UE::PixelStreaming::FDataChannelObserver::OnStateChange()
+{
+	if (!DataChannel)
+	{
+		return;
+	}
+	
+	webrtc::DataChannelInterface::DataState State = DataChannel->state();
+	UPixelStreamingDelegates* Delegates = UPixelStreamingDelegates::GetPixelStreamingDelegates();
+	if (State == webrtc::DataChannelInterface::DataState::kOpen && Delegates)
+	{
+		Delegates->OnDataChannelOpenNative.Broadcast(this->PlayerId, this->DataChannel.get());
+	}
+}
+
+void UE::PixelStreaming::FDataChannelObserver::OnMessage(const webrtc::DataBuffer& Buffer)
+{
+	if (!DataChannel)
+	{
+		return;
+	}
+
+	UE::PixelStreaming::Protocol::EToStreamerMsg MsgType = static_cast<UE::PixelStreaming::Protocol::EToStreamerMsg>(Buffer.data.data()[0]);
+
+	if (MsgType == UE::PixelStreaming::Protocol::EToStreamerMsg::RequestQualityControl)
+	{
+		check(Buffer.data.size() == 1);
+		UE_LOG(LogPixelStreaming, Log, TEXT("Player %s has requested quality control through the data channel."), *PlayerId);
+		PlayerSessions->SetQualityController(PlayerId);
+	}
+	else if (MsgType == UE::PixelStreaming::Protocol::EToStreamerMsg::LatencyTest)
+	{
+		SendLatencyReport();
+	}
+	else if (MsgType == UE::PixelStreaming::Protocol::EToStreamerMsg::RequestInitialSettings)
+	{
+		SendInitialSettings();
+	}
+	else if (!IsEngineExitRequested())
+	{
+		InputDevice.OnMessage(Buffer);
+	}
+}
+
+void UE::PixelStreaming::FDataChannelObserver::SendLatencyReport() const
+{
+	if (!DataChannel)
+	{
+		return;
+	}
+
+	if (UE::PixelStreaming::Settings::CVarPixelStreamingDisableLatencyTester.GetValueOnAnyThread())
+	{
+		return;
+	}
+
+	double ReceiptTimeMs = FPlatformTime::ToMilliseconds64(FPlatformTime::Cycles64());
+
+	AsyncTask(ENamedThreads::GameThread, [this, ReceiptTimeMs]() {
+		double EncodeMs = -1.0;
+		double CaptureToSendMs = 0.0;
+
+		UE::PixelStreaming::FStats* Stats = UE::PixelStreaming::FStats::Get();
+		if (Stats)
+		{
+			// bool QueryPeerStat_GameThread(FPixelStreamingPlayerId PlayerId, FName StatToQuery, double& OutStatValue)
+			Stats->QueryPeerStat_GameThread(PlayerId, PixelStreamingStatNames::MeanEncodeTime, EncodeMs);
+			Stats->QueryPeerStat_GameThread(PlayerId, PixelStreamingStatNames::MeanSendDelay, CaptureToSendMs);
+		}
+
+		double TransmissionTimeMs = FPlatformTime::ToMilliseconds64(FPlatformTime::Cycles64());
+
+		FString ReportToTransmitJSON = FString::Printf(
+			TEXT("{ \"ReceiptTimeMs\": %.2f, \"EncodeMs\": %.2f, \"CaptureToSendMs\": %.2f, \"TransmissionTimeMs\": %.2f }"),
+			ReceiptTimeMs,
+			EncodeMs,
+			CaptureToSendMs,
+			TransmissionTimeMs);
+
+		PlayerSessions->SendMessage(PlayerId, UE::PixelStreaming::Protocol::EToPlayerMsg::LatencyTest, ReportToTransmitJSON);
+	});
+}
+
+void UE::PixelStreaming::FDataChannelObserver::OnBufferedAmountChange(uint64_t PreviousAmount)
+{
+	if (!DataChannel)
+	{
+		return;
+	}
+
+	UE_LOG(
+		LogPixelStreaming,
+		VeryVerbose,
+		TEXT("player %s: OnBufferedAmountChanged: prev %d, cur %d"),
+		*PlayerId,
+		PreviousAmount,
+		DataChannel->buffered_amount());
+}
+
+void UE::PixelStreaming::FDataChannelObserver::SendInitialSettings() const
+{
+	if (!DataChannel)
+	{
+		return;
+	}
+
+	const FString PixelStreamingPayload = FString::Printf(TEXT("{ \"AllowPixelStreamingCommands\": %s, \"DisableLatencyTest\": %s }"),
+		UE::PixelStreaming::Settings::CVarPixelStreamingAllowConsoleCommands.GetValueOnAnyThread() ? TEXT("true") : TEXT("false"),
+		UE::PixelStreaming::Settings::CVarPixelStreamingDisableLatencyTester.GetValueOnAnyThread() ? TEXT("true") : TEXT("false"));
+
+	const FString WebRTCPayload = FString::Printf(TEXT("{ \"DegradationPref\": \"%s\", \"FPS\": %d, \"MinBitrate\": %d, \"MaxBitrate\": %d, \"LowQP\": %d, \"HighQP\": %d }"),
+		*UE::PixelStreaming::Settings::CVarPixelStreamingDegradationPreference.GetValueOnAnyThread(),
+		UE::PixelStreaming::Settings::CVarPixelStreamingWebRTCFps.GetValueOnAnyThread(),
+		UE::PixelStreaming::Settings::CVarPixelStreamingWebRTCMinBitrate.GetValueOnAnyThread(),
+		UE::PixelStreaming::Settings::CVarPixelStreamingWebRTCMaxBitrate.GetValueOnAnyThread(),
+		UE::PixelStreaming::Settings::CVarPixelStreamingWebRTCLowQpThreshold.GetValueOnAnyThread(),
+		UE::PixelStreaming::Settings::CVarPixelStreamingWebRTCHighQpThreshold.GetValueOnAnyThread());
+
+	const FString EncoderPayload = FString::Printf(TEXT("{ \"TargetBitrate\": %d, \"MaxBitrate\": %d, \"MinQP\": %d, \"MaxQP\": %d, \"RateControl\": \"%s\", \"FillerData\": %d, \"MultiPass\": \"%s\" }"),
+		UE::PixelStreaming::Settings::CVarPixelStreamingEncoderTargetBitrate.GetValueOnAnyThread(),
+		UE::PixelStreaming::Settings::CVarPixelStreamingEncoderMaxBitrate.GetValueOnAnyThread(),
+		UE::PixelStreaming::Settings::CVarPixelStreamingEncoderMinQP.GetValueOnAnyThread(),
+		UE::PixelStreaming::Settings::CVarPixelStreamingEncoderMaxQP.GetValueOnAnyThread(),
+		*UE::PixelStreaming::Settings::CVarPixelStreamingEncoderRateControl.GetValueOnAnyThread(),
+		UE::PixelStreaming::Settings::CVarPixelStreamingEnableFillerData.GetValueOnAnyThread() ? 1 : 0,
+		*UE::PixelStreaming::Settings::CVarPixelStreamingEncoderMultipass.GetValueOnAnyThread());
+
+	const FString FullPayload = FString::Printf(TEXT("{ \"PixelStreaming\": %s, \"Encoder\": %s, \"WebRTC\": %s }"), *PixelStreamingPayload, *EncoderPayload, *WebRTCPayload);
+
+	if (!PlayerSessions->SendMessage(PlayerId, UE::PixelStreaming::Protocol::EToPlayerMsg::InitialSettings, FullPayload))
+	{
+		UE_LOG(LogPixelStreaming, Log, TEXT("Failed to send initial Pixel Streaming settings."));
+	}
+}
\ No newline at end of file
diff --git a/Engine/Plugins/Media/PixelStreaming/Source/PixelStreaming/Private/DataChannelObserver.h b/Engine/Plugins/Media/PixelStreaming/Source/PixelStreaming/Private/DataChannelObserver.h
new file mode 100644
--- /dev/null
+++ b/Engine/Plugins/Media/PixelStreaming/Source/PixelStreaming/Private/DataChannelObserver.h
@@ -0,0 +1,51 @@
+// Copyright Epic Games, Inc. All Rights Reserved.
+
+#pragma once
+
+#include "WebRTCIncludes.h"
+#include "PixelStreamingPlayerId.h"
+#include "Delegates/Delegate.h"
+#include "Delegates/DelegateCombinations.h"
+
+namespace UE {
+	namespace PixelStreaming {
+		class IPixelStreamingSessions;
+		class FInputDevice;
+
+		// Observer on a particular player/peer's data channel.
+		class FDataChannelObserver : public webrtc::DataChannelObserver
+		{
+
+		public:
+			FDataChannelObserver(IPixelStreamingSessions* InSessions, FPixelStreamingPlayerId InPlayerId);
+
+			// Begin webrtc::DataChannelObserver
+
+			// The data channel state have changed.
+			virtual void OnStateChange() override;
+
+			//  A data buffer was successfully received.
+			virtual void OnMessage(const webrtc::DataBuffer& buffer) override;
+
+			// The data channel's buffered_amount has changed.
+			virtual void OnBufferedAmountChange(uint64_t sent_data_size) override;
+
+			// End webrtc::DataChannelObserver
+
+			void SendInitialSettings() const;
+
+			void Register(rtc::scoped_refptr<webrtc::DataChannelInterface> InDataChannel);
+			void Unregister();
+
+		private:
+			void SendLatencyReport() const;
+			
+
+		private:
+			rtc::scoped_refptr<webrtc::DataChannelInterface> DataChannel;
+			IPixelStreamingSessions* PlayerSessions;
+			FPixelStreamingPlayerId PlayerId;
+			FInputDevice& InputDevice;
+		};
+	}
+}
diff --git a/Engine/Plugins/Media/PixelStreaming/Source/PixelStreaming/Private/EncoderFactory.cpp b/Engine/Plugins/Media/PixelStreaming/Source/PixelStreaming/Private/EncoderFactory.cpp
new file mode 100644
--- /dev/null
+++ b/Engine/Plugins/Media/PixelStreaming/Source/PixelStreaming/Private/EncoderFactory.cpp
@@ -0,0 +1,210 @@
+// Copyright Epic Games, Inc. All Rights Reserved.
+
+#include "EncoderFactory.h"
+#include "VideoEncoderRTC.h"
+#include "Settings.h"
+#include "absl/strings/match.h"
+#include "modules/video_coding/codecs/vp8/include/vp8.h"
+#include "Misc/ScopeLock.h"
+#include "VideoEncoderH264Wrapper.h"
+#include "SimulcastEncoderAdapter.h"
+#include "VideoEncoderFactory.h"
+#include "PixelStreamingPrivate.h"
+
+/*
+* ------------- UE::PixelStreaming::FSimulcastEncoderFactory ---------------
+*/
+
+UE::PixelStreaming::FSimulcastEncoderFactory::FSimulcastEncoderFactory()
+	: RealFactory(new UE::PixelStreaming::FVideoEncoderFactory())
+{
+}
+
+UE::PixelStreaming::FSimulcastEncoderFactory::~FSimulcastEncoderFactory()
+{
+}
+
+std::vector<webrtc::SdpVideoFormat> UE::PixelStreaming::FSimulcastEncoderFactory::GetSupportedFormats() const
+{
+	return RealFactory->GetSupportedFormats();
+}
+
+UE::PixelStreaming::FSimulcastEncoderFactory::CodecInfo UE::PixelStreaming::FSimulcastEncoderFactory::QueryVideoEncoder(const webrtc::SdpVideoFormat& format) const
+{
+	return RealFactory->QueryVideoEncoder(format);
+}
+
+std::unique_ptr<webrtc::VideoEncoder> UE::PixelStreaming::FSimulcastEncoderFactory::CreateVideoEncoder(const webrtc::SdpVideoFormat& format)
+{
+	return std::make_unique<UE::PixelStreaming::FSimulcastEncoderAdapter>(RealFactory.Get(), format);
+}
+
+UE::PixelStreaming::FVideoEncoderFactory* UE::PixelStreaming::FSimulcastEncoderFactory::GetRealFactory() const
+{
+	return RealFactory.Get();
+}
+
+/*
+* ------------- UE::PixelStreaming::FVideoEncoderFactory ---------------
+*/
+
+UE::PixelStreaming::FVideoEncoderFactory::FVideoEncoderFactory()
+{
+}
+
+UE::PixelStreaming::FVideoEncoderFactory::~FVideoEncoderFactory()
+{
+}
+
+webrtc::SdpVideoFormat UE::PixelStreaming::FVideoEncoderFactory::CreateH264Format(webrtc::H264::Profile profile, webrtc::H264::Level level)
+{
+	const absl::optional<std::string> ProfileString =
+		webrtc::H264::ProfileLevelIdToString(webrtc::H264::ProfileLevelId(profile, level));
+	check(ProfileString);
+	return webrtc::SdpVideoFormat(
+		cricket::kH264CodecName,
+		{ { cricket::kH264FmtpProfileLevelId, *ProfileString },
+			{ cricket::kH264FmtpLevelAsymmetryAllowed, "1" },
+			{ cricket::kH264FmtpPacketizationMode, "1" } });
+}
+
+std::vector<webrtc::SdpVideoFormat> UE::PixelStreaming::FVideoEncoderFactory::GetSupportedFormats() const
+{
+	const bool bForceVP8 = UE::PixelStreaming::Settings::IsForceVP8();
+
+	std::vector<webrtc::SdpVideoFormat> video_formats;
+	if (bForceVP8)
+	{
+		video_formats.push_back(webrtc::SdpVideoFormat(cricket::kVp8CodecName));
+	}
+	else
+	{
+		video_formats.push_back(UE::PixelStreaming::FVideoEncoderFactory::CreateH264Format(webrtc::H264::kProfileConstrainedBaseline, webrtc::H264::kLevel3_1));
+	}
+
+	return video_formats;
+}
+
+UE::PixelStreaming::FVideoEncoderFactory::CodecInfo UE::PixelStreaming::FVideoEncoderFactory::QueryVideoEncoder(const webrtc::SdpVideoFormat& format) const
+{
+	webrtc::VideoEncoderFactory::CodecInfo CodecInfo = { false, false };
+	CodecInfo.is_hardware_accelerated = true;
+	CodecInfo.has_internal_source = false;
+	return CodecInfo;
+}
+
+std::unique_ptr<webrtc::VideoEncoder> UE::PixelStreaming::FVideoEncoderFactory::CreateVideoEncoder(const webrtc::SdpVideoFormat& format)
+{
+	if (absl::EqualsIgnoreCase(format.name, cricket::kVp8CodecName))
+		return webrtc::VP8Encoder::Create();
+	else
+	{
+		// Lock during encoder creation
+		FScopeLock Lock(&ActiveEncodersGuard);
+		auto VideoEncoder = std::make_unique<UE::PixelStreaming::FVideoEncoderRTC>(*this);
+		ActiveEncoders.Add(VideoEncoder.get());
+		return VideoEncoder;
+	}
+}
+
+void UE::PixelStreaming::FVideoEncoderFactory::OnEncodedImage(FHardwareEncoderId SourceEncoderId, const webrtc::EncodedImage& encoded_image, const webrtc::CodecSpecificInfo* codec_specific_info, const webrtc::RTPFragmentationHeader* fragmentation)
+{
+	// Lock as we send encoded image to each encoder.
+	FScopeLock Lock(&ActiveEncodersGuard);
+
+	// Go through each encoder and send our encoded image to its callback
+	for (UE::PixelStreaming::FVideoEncoderRTC* Encoder : ActiveEncoders)
+	{
+		Encoder->SendEncodedImage(SourceEncoderId, encoded_image, codec_specific_info, fragmentation);
+	}
+}
+
+void UE::PixelStreaming::FVideoEncoderFactory::ReleaseVideoEncoder(UE::PixelStreaming::FVideoEncoderRTC* Encoder)
+{
+	// Lock during deleting an encoder
+	FScopeLock Lock(&ActiveEncodersGuard);
+	ActiveEncoders.Remove(Encoder);
+}
+
+void UE::PixelStreaming::FVideoEncoderFactory::ForceKeyFrame()
+{
+	FScopeLock Lock(&ActiveEncodersGuard);
+	// Go through each encoder and send our encoded image to its callback
+	for (auto& KeyAndEncoder : HardwareEncoders)
+	{
+		KeyAndEncoder.Value->SetForceNextKeyframe();
+	}
+}
+
+namespace
+{
+	uint64 HashResolution(uint32 Width, uint32 Height)
+	{
+		return static_cast<uint64>(Width) << 32 | static_cast<uint64>(Height);
+	}
+} // namespace
+
+TOptional<UE::PixelStreaming::FVideoEncoderFactory::FHardwareEncoderId> UE::PixelStreaming::FVideoEncoderFactory::GetOrCreateHardwareEncoder(int Width, int Height, int MaxBitrate, int TargetBitrate, int MaxFramerate)
+{
+	FHardwareEncoderId EncoderId = HashResolution(Width, Height);
+	UE::PixelStreaming::FVideoEncoderH264Wrapper* Existing = GetHardwareEncoder(EncoderId);
+	
+	TOptional<FHardwareEncoderId> Result = EncoderId;
+
+	if (Existing == nullptr)
+	{
+		// Make AVEncoder frame factory.
+		TUniquePtr<FEncoderFrameFactory> FrameFactory = MakeUnique<FEncoderFrameFactory>();
+		FrameFactory->SetResolution(Width, Height);
+
+		// Make the encoder config
+		AVEncoder::FVideoEncoder::FLayerConfig EncoderConfig;
+		EncoderConfig.Width = Width;
+		EncoderConfig.Height = Height;
+		EncoderConfig.MaxFramerate = MaxFramerate;
+		EncoderConfig.TargetBitrate = TargetBitrate;
+		EncoderConfig.MaxBitrate = MaxBitrate;
+
+		// Make the actual AVEncoder encoder.
+		const TArray<AVEncoder::FVideoEncoderInfo>& Available = AVEncoder::FVideoEncoderFactory::Get().GetAvailable();
+		TUniquePtr<AVEncoder::FVideoEncoder> Encoder = AVEncoder::FVideoEncoderFactory::Get().Create(Available[0].ID, FrameFactory->GetOrCreateVideoEncoderInput(EncoderConfig.Width, EncoderConfig.Height), EncoderConfig);
+		if(Encoder.IsValid())
+		{
+			Encoder->SetOnEncodedPacket([EncoderId, this](uint32 InLayerIndex, const AVEncoder::FVideoEncoderInputFrame* InFrame, const AVEncoder::FCodecPacket& InPacket) {
+				// Note: this is a static method call.
+				FVideoEncoderH264Wrapper::OnEncodedPacket(EncoderId, this, InLayerIndex, InFrame, InPacket);
+			});
+
+			// Make the hardware encoder wrapper
+			auto EncoderWrapper = MakeUnique<FVideoEncoderH264Wrapper>(EncoderId, MoveTemp(FrameFactory), MoveTemp(Encoder));
+			FVideoEncoderH264Wrapper* ReturnValue = EncoderWrapper.Get();
+			{
+				FScopeLock Lock(&HardwareEncodersGuard);
+				HardwareEncoders.Add(EncoderId, MoveTemp(EncoderWrapper));
+			}
+
+			return Result;
+		}
+		else
+		{
+			UE_LOG(LogPixelStreaming, Error, TEXT("Could not create encoder. Check encoder config or perhaps you used up all your HW encoders."));
+			// We could not make the encoder, so indicate the id was not set successfully.
+			Result.Reset();
+			return Result;
+		}
+	}
+	else
+	{
+		return Result;
+	}
+}
+
+UE::PixelStreaming::FVideoEncoderH264Wrapper* UE::PixelStreaming::FVideoEncoderFactory::GetHardwareEncoder(FHardwareEncoderId Id)
+{
+	FScopeLock Lock(&HardwareEncodersGuard);
+	if (auto&& Existing = HardwareEncoders.Find(Id))
+	{
+		return Existing->Get();
+	}
+	return nullptr;
+}
diff --git a/Engine/Plugins/Media/PixelStreaming/Source/PixelStreaming/Private/EncoderFactory.h b/Engine/Plugins/Media/PixelStreaming/Source/PixelStreaming/Private/EncoderFactory.h
new file mode 100644
--- /dev/null
+++ b/Engine/Plugins/Media/PixelStreaming/Source/PixelStreaming/Private/EncoderFactory.h
@@ -0,0 +1,70 @@
+// Copyright Epic Games, Inc. All Rights Reserved.
+#pragma once
+
+#include "WebRTCIncludes.h"
+#include "VideoEncoder.h"
+#include "Containers/Queue.h"
+#include "PixelStreamingPlayerId.h"
+#include "HAL/CriticalSection.h"
+#include "Misc/Optional.h"
+
+namespace UE {
+	namespace PixelStreaming {
+		class FVideoEncoderFactory;
+		class FVideoEncoderRTC;
+		class FVideoEncoderH264Wrapper;
+
+		class FSimulcastEncoderFactory : public webrtc::VideoEncoderFactory
+		{
+		public:
+			FSimulcastEncoderFactory();
+			virtual ~FSimulcastEncoderFactory();
+			virtual std::vector<webrtc::SdpVideoFormat> GetSupportedFormats() const;
+			virtual CodecInfo QueryVideoEncoder(const webrtc::SdpVideoFormat& format) const override;
+			virtual std::unique_ptr<webrtc::VideoEncoder> CreateVideoEncoder(const webrtc::SdpVideoFormat& format);
+
+			FVideoEncoderFactory* GetRealFactory() const;
+
+		private:
+			// making this a unique ptr is a little cheeky since we pass around raw pointers to it
+			// but the sole ownership lies in this class
+			TUniquePtr<FVideoEncoderFactory> RealFactory;
+		};
+
+		class FVideoEncoderFactory : public webrtc::VideoEncoderFactory
+		{
+		public:
+			FVideoEncoderFactory();
+			virtual ~FVideoEncoderFactory() override;
+
+			virtual std::vector<webrtc::SdpVideoFormat> GetSupportedFormats() const override;
+
+			// Always returns our H264 hardware encoders codec_info for now/
+			virtual CodecInfo QueryVideoEncoder(const webrtc::SdpVideoFormat& Format) const override;
+
+			// Always returns our H264 hardware encoder for now.
+			virtual std::unique_ptr<webrtc::VideoEncoder> CreateVideoEncoder(const webrtc::SdpVideoFormat& Format) override;
+
+			void ReleaseVideoEncoder(FVideoEncoderRTC* Encoder);
+			void ForceKeyFrame();
+
+			using FHardwareEncoderId = uint64;
+			TOptional<FHardwareEncoderId> GetOrCreateHardwareEncoder(int Width, int Height, int MaxBitrate, int TargetBitrate, int MaxFramerate);
+			FVideoEncoderH264Wrapper* GetHardwareEncoder(FHardwareEncoderId Id);
+
+			void OnEncodedImage(FHardwareEncoderId SourceEncoderId, const webrtc::EncodedImage& Encoded_image, const webrtc::CodecSpecificInfo* CodecSpecificInfo, const webrtc::RTPFragmentationHeader* Fragmentation);
+
+		private:
+
+			static webrtc::SdpVideoFormat CreateH264Format(webrtc::H264::Profile Profile, webrtc::H264::Level Level);
+
+			// These are the actual hardware encoders serving multiple pixelstreaming encoders
+			TMap<FHardwareEncoderId, TUniquePtr<FVideoEncoderH264Wrapper>> HardwareEncoders;
+			FCriticalSection HardwareEncodersGuard;
+
+			// Encoders assigned to each peer. Each one of these will be assigned to one of the hardware encoders
+			TArray<FVideoEncoderRTC*> ActiveEncoders;
+			FCriticalSection ActiveEncodersGuard;
+		};
+	}
+}
diff --git a/Engine/Plugins/Media/PixelStreaming/Source/PixelStreaming/Private/EncoderFrameFactory.cpp b/Engine/Plugins/Media/PixelStreaming/Source/PixelStreaming/Private/EncoderFrameFactory.cpp
new file mode 100644
--- /dev/null
+++ b/Engine/Plugins/Media/PixelStreaming/Source/PixelStreaming/Private/EncoderFrameFactory.cpp
@@ -0,0 +1,609 @@
+// Copyright Epic Games, Inc. All Rights Reserved.
+
+#include "EncoderFrameFactory.h"
+#include "PixelStreamingPrivate.h"
+#include "CudaModule.h"
+#include "VulkanRHIPrivate.h"
+#include "Settings.h"
+#include "VideoEncoderInput.h"
+
+#if PLATFORM_WINDOWS
+	#include "VideoCommon.h"
+	#include "D3D11State.h"
+	#include "D3D11Resources.h"
+	#include "D3D12RHICommon.h"
+	#include "D3D12RHIPrivate.h"
+	#include "D3D12Resources.h"
+	#include "D3D12Texture.h"
+	#include "Windows/AllowWindowsPlatformTypes.h"
+THIRD_PARTY_INCLUDES_START
+	#include <VersionHelpers.h>
+THIRD_PARTY_INCLUDES_END
+	#include "Windows/HideWindowsPlatformTypes.h"
+#endif
+
+UE::PixelStreaming::FEncoderFrameFactory::FEncoderFrameFactory()
+{
+}
+
+UE::PixelStreaming::FEncoderFrameFactory::~FEncoderFrameFactory()
+{
+	FlushFrames();
+}
+
+void UE::PixelStreaming::FEncoderFrameFactory::FlushFrames()
+{
+	for (auto& Entry : TextureToFrameMapping)
+	{
+		EncoderInput->DestroyBuffer(Entry.Value);
+	}
+	TextureToFrameMapping.Empty();
+}
+
+void UE::PixelStreaming::FEncoderFrameFactory::RemoveStaleTextures()
+{
+	// Remove any textures whose only reference is the one held by this class
+
+	TMap<FTexture2DRHIRef, AVEncoder::FVideoEncoderInputFrame*>::TIterator Iter = TextureToFrameMapping.CreateIterator();
+	for(; Iter; ++Iter)
+	{
+		FTexture2DRHIRef& Tex = Iter.Key();
+		AVEncoder::FVideoEncoderInputFrame* Frame = Iter.Value();
+
+		if(Tex.GetRefCount() == 1)
+		{
+			EncoderInput->DestroyBuffer(Frame);
+			Iter.RemoveCurrent();
+		}
+	}
+}
+
+AVEncoder::FVideoEncoderInputFrame* UE::PixelStreaming::FEncoderFrameFactory::GetOrCreateFrame(int InWidth, int InHeight, const FTexture2DRHIRef InTexture)
+{
+	check(EncoderInput.IsValid());
+
+	RemoveStaleTextures();
+
+	AVEncoder::FVideoEncoderInputFrame* OutFrame;
+
+	if (TextureToFrameMapping.Contains(InTexture))
+	{
+		OutFrame = *(TextureToFrameMapping.Find(InTexture));
+		checkf(
+            InWidth == OutFrame->GetWidth() && InHeight == OutFrame->GetHeight(),
+            TEXT("The requested resolution does not match the existing frame, we do not support a resolution change at this step."));
+	}
+	// A frame needs to be created
+	else
+	{
+		OutFrame = EncoderInput->CreateBuffer([this](const AVEncoder::FVideoEncoderInputFrame* ReleasedFrame) { /* OnReleased */ });
+		SetTexture(OutFrame, InTexture);
+		TextureToFrameMapping.Add(InTexture, OutFrame);
+	}
+
+	OutFrame->SetFrameID(++FrameId);
+	return OutFrame;
+}
+
+AVEncoder::FVideoEncoderInputFrame* UE::PixelStreaming::FEncoderFrameFactory::GetFrameAndSetTexture(int InWidth, int InHeight, FTexture2DRHIRef InTexture)
+{
+	check(EncoderInput.IsValid());
+
+	AVEncoder::FVideoEncoderInputFrame* Frame = GetOrCreateFrame(InWidth, InHeight, InTexture);
+
+	return Frame;
+}
+
+TSharedPtr<AVEncoder::FVideoEncoderInput> UE::PixelStreaming::FEncoderFrameFactory::GetOrCreateVideoEncoderInput(int InWidth, int InHeight)
+{
+	if (!EncoderInput.IsValid())
+	{
+		EncoderInput = CreateVideoEncoderInput(InWidth, InHeight);
+		EncoderInput->SetMaxNumBuffers((uint32)UE::PixelStreaming::Settings::CVarPixelStreamingMaxNumBackBuffers.GetValueOnAnyThread());
+	}
+
+	return EncoderInput;
+}
+
+void UE::PixelStreaming::FEncoderFrameFactory::SetResolution(int InWidth, int InHeight)
+{
+	TSharedPtr<AVEncoder::FVideoEncoderInput> NewEncoderInput = GetOrCreateVideoEncoderInput(InWidth, InHeight);
+	NewEncoderInput->SetResolution(InWidth, InHeight);
+	NewEncoderInput->Flush();
+}
+
+TSharedPtr<AVEncoder::FVideoEncoderInput> UE::PixelStreaming::FEncoderFrameFactory::CreateVideoEncoderInput(int InWidth, int InHeight) const
+{
+	if (!GDynamicRHI)
+	{
+		UE_LOG(LogPixelStreaming, Error, TEXT("GDynamicRHI not valid for some reason."));
+		return nullptr;
+	}
+
+	FString RHIName = GDynamicRHI->GetName();
+
+	// Consider if we want to support runtime resolution changing?
+	bool bIsResizable = false;
+
+	if (RHIName == TEXT("Vulkan"))
+	{
+		if (IsRHIDeviceAMD())
+		{
+			FVulkanDynamicRHI* DynamicRHI = static_cast<FVulkanDynamicRHI*>(GDynamicRHI);
+			AVEncoder::FVulkanDataStruct VulkanData = { DynamicRHI->GetInstance(), DynamicRHI->GetDevice()->GetPhysicalHandle(), DynamicRHI->GetDevice()->GetInstanceHandle() };
+
+			return AVEncoder::FVideoEncoderInput::CreateForVulkan(&VulkanData, InWidth, InHeight, bIsResizable);
+		}
+		else if (IsRHIDeviceNVIDIA())
+		{
+			if (FModuleManager::Get().IsModuleLoaded("CUDA"))
+			{
+				return AVEncoder::FVideoEncoderInput::CreateForCUDA(FModuleManager::GetModuleChecked<FCUDAModule>("CUDA").GetCudaContext(), InWidth, InHeight, bIsResizable);
+			}
+			else
+			{
+				UE_LOG(LogPixelStreaming, Error, TEXT("CUDA module is not loaded!"));
+				return nullptr;
+			}
+		}
+	}
+#if PLATFORM_WINDOWS
+	else if (RHIName == TEXT("D3D11"))
+	{
+		if (IsRHIDeviceAMD())
+		{
+			return AVEncoder::FVideoEncoderInput::CreateForD3D11(GDynamicRHI->RHIGetNativeDevice(), InWidth, InHeight, bIsResizable, true);
+		}
+		else if (IsRHIDeviceNVIDIA())
+		{
+			return AVEncoder::FVideoEncoderInput::CreateForD3D11(GDynamicRHI->RHIGetNativeDevice(), InWidth, InHeight, bIsResizable, false);
+		}
+	}
+	else if (RHIName == TEXT("D3D12"))
+	{
+		if (IsRHIDeviceAMD())
+		{
+			return AVEncoder::FVideoEncoderInput::CreateForD3D12(GDynamicRHI->RHIGetNativeDevice(), InWidth, InHeight, bIsResizable, false);
+		}
+		else if (IsRHIDeviceNVIDIA())
+		{
+			return AVEncoder::FVideoEncoderInput::CreateForCUDA(FModuleManager::GetModuleChecked<FCUDAModule>("CUDA").GetCudaContext(), InWidth, InHeight, bIsResizable);
+		}
+	}
+#endif
+
+	UE_LOG(LogPixelStreaming, Error, TEXT("Current RHI %s is not supported in Pixel Streaming"), *RHIName);
+	return nullptr;
+}
+
+void UE::PixelStreaming::FEncoderFrameFactory::SetTexture(AVEncoder::FVideoEncoderInputFrame* InputFrame, const FTexture2DRHIRef& Texture)
+{
+	FString RHIName = GDynamicRHI->GetName();
+
+	// VULKAN
+	if (RHIName == TEXT("Vulkan"))
+	{
+		if (IsRHIDeviceAMD())
+		{
+			FVulkanTexture2D* VulkanTexture = static_cast<FVulkanTexture2D*>(Texture.GetReference());
+			InputFrame->SetTexture(VulkanTexture->Surface.Image, [](VkImage NativeTexture) { /* Do something with released texture if needed */ });
+		}
+		else if (IsRHIDeviceNVIDIA())
+		{
+			SetTextureCUDAVulkan(InputFrame, Texture);
+		}
+		else
+		{
+			UE_LOG(LogPixelStreaming, Error, TEXT("Pixel Streaming only supports AMD and NVIDIA devices, this device is neither of those."));
+		}
+	}
+#if PLATFORM_WINDOWS
+	// TODO: Fix CUDA DX11 (Using CUDA as a bridge between DX11 and NVENC currently produces garbled results)
+	// DX11
+	else if (RHIName == TEXT("D3D11"))
+	{
+		InputFrame->SetTexture((ID3D11Texture2D*)Texture->GetNativeResource(), [](ID3D11Texture2D* NativeTexture) { /* Do something with released texture if needed */ });
+	}
+	// DX12
+	else if (RHIName == TEXT("D3D12"))
+	{
+		if (IsRHIDeviceAMD())
+		{
+			InputFrame->SetTexture((ID3D12Resource*)Texture->GetNativeResource(), [](ID3D12Resource* NativeTexture) { /* Do something with released texture if needed */ });
+		}
+		else if (IsRHIDeviceNVIDIA())
+		{
+			SetTextureCUDAD3D12(InputFrame, Texture);
+		}
+		else
+		{
+			UE_LOG(LogPixelStreaming, Error, TEXT("Pixel Streaming only supports AMD and NVIDIA devices, this device is neither of those."));
+		}
+	}
+#endif // PLATFORM_WINDOWS
+	else
+	{
+		UE_LOG(LogPixelStreaming, Error, TEXT("Pixel Streaming does not support this RHI - %s"), *RHIName);
+	}
+}
+
+void UE::PixelStreaming::FEncoderFrameFactory::SetTextureCUDAVulkan(AVEncoder::FVideoEncoderInputFrame* InputFrame, const FTexture2DRHIRef& Texture)
+{
+	FVulkanTexture2D* VulkanTexture = static_cast<FVulkanTexture2D*>(Texture.GetReference());
+	VkDevice Device = static_cast<FVulkanDynamicRHI*>(GDynamicRHI)->GetDevice()->GetInstanceHandle();
+
+#if PLATFORM_WINDOWS
+	HANDLE Handle;
+	// It is recommended to use NT handles where available, but these are only supported from Windows 8 onward, for earliers versions of Windows
+	// we need to use a Win7 style handle. NT handles require us to close them when we are done with them to prevent memory leaks.
+	// Refer to remarks section of https://docs.microsoft.com/en-us/windows/win32/api/dxgi1_2/nf-dxgi1_2-idxgiresource1-createsharedhandle
+	bool bUseNTHandle = IsWindows8OrGreater();
+
+	{
+		// Generate VkMemoryGetWin32HandleInfoKHR
+		VkMemoryGetWin32HandleInfoKHR MemoryGetHandleInfoKHR = {};
+		MemoryGetHandleInfoKHR.sType = VK_STRUCTURE_TYPE_MEMORY_GET_WIN32_HANDLE_INFO_KHR;
+		MemoryGetHandleInfoKHR.pNext = NULL;
+		MemoryGetHandleInfoKHR.memory = VulkanTexture->Surface.GetAllocationHandle();
+		MemoryGetHandleInfoKHR.handleType = bUseNTHandle ? VK_EXTERNAL_MEMORY_HANDLE_TYPE_OPAQUE_WIN32_BIT : VK_EXTERNAL_MEMORY_HANDLE_TYPE_OPAQUE_WIN32_KMT_BIT;
+
+		// While this operation is safe (and unavoidable) C4191 has been enabled and this will trigger an error with warnings as errors
+	#pragma warning(push)
+	#pragma warning(disable : 4191)
+		PFN_vkGetMemoryWin32HandleKHR GetMemoryWin32HandleKHR = (PFN_vkGetMemoryWin32HandleKHR)VulkanRHI::vkGetDeviceProcAddr(Device, "vkGetMemoryWin32HandleKHR");
+		VERIFYVULKANRESULT(GetMemoryWin32HandleKHR(Device, &MemoryGetHandleInfoKHR, &Handle));
+	#pragma warning(pop)
+	}
+
+	FCUDAModule::CUDA().cuCtxPushCurrent(FModuleManager::GetModuleChecked<FCUDAModule>("CUDA").GetCudaContext());
+
+	CUexternalMemory MappedExternalMemory = nullptr;
+
+	{
+		// generate a cudaExternalMemoryHandleDesc
+		CUDA_EXTERNAL_MEMORY_HANDLE_DESC CudaExtMemHandleDesc = {};
+		CudaExtMemHandleDesc.type = bUseNTHandle ? CU_EXTERNAL_MEMORY_HANDLE_TYPE_OPAQUE_WIN32 : CU_EXTERNAL_MEMORY_HANDLE_TYPE_OPAQUE_WIN32_KMT;
+		CudaExtMemHandleDesc.handle.win32.name = NULL;
+		CudaExtMemHandleDesc.handle.win32.handle = Handle;
+		CudaExtMemHandleDesc.size = VulkanTexture->Surface.GetAllocationOffset() + VulkanTexture->Surface.GetMemorySize();
+
+		// import external memory
+		CUresult Result = FCUDAModule::CUDA().cuImportExternalMemory(&MappedExternalMemory, &CudaExtMemHandleDesc);
+		if (Result != CUDA_SUCCESS)
+		{
+			UE_LOG(LogPixelStreaming, Error, TEXT("Failed to import external memory from vulkan error: %d"), Result);
+		}
+	}
+
+	// Only store handle to be closed on frame destruction if it is an NT handle
+	Handle = bUseNTHandle ? Handle : NULL;
+#else
+	void* Handle = nullptr;
+
+	// Get the CUarray to that textures memory making sure the clear it when done
+	int Fd;
+
+	{
+		// Generate VkMemoryGetFdInfoKHR
+		VkMemoryGetFdInfoKHR MemoryGetFdInfoKHR = {};
+		MemoryGetFdInfoKHR.sType = VK_STRUCTURE_TYPE_MEMORY_GET_FD_INFO_KHR;
+		MemoryGetFdInfoKHR.pNext = NULL;
+		MemoryGetFdInfoKHR.memory = VulkanTexture->Surface.GetAllocationHandle();
+		MemoryGetFdInfoKHR.handleType = VK_EXTERNAL_MEMORY_HANDLE_TYPE_OPAQUE_FD_BIT_KHR;
+
+		// While this operation is safe (and unavoidable) C4191 has been enabled and this will trigger an error with warnings as errors
+	#pragma warning(push)
+	#pragma warning(disable : 4191)
+		PFN_vkGetMemoryFdKHR FPGetMemoryFdKHR = (PFN_vkGetMemoryFdKHR)VulkanRHI::vkGetDeviceProcAddr(Device, "vkGetMemoryFdKHR");
+		VERIFYVULKANRESULT(FPGetMemoryFdKHR(Device, &MemoryGetFdInfoKHR, &Fd));
+	#pragma warning(pop)
+	}
+
+	FCUDAModule::CUDA().cuCtxPushCurrent(FModuleManager::GetModuleChecked<FCUDAModule>("CUDA").GetCudaContext());
+
+	CUexternalMemory MappedExternalMemory = nullptr;
+
+	{
+		// generate a cudaExternalMemoryHandleDesc
+		CUDA_EXTERNAL_MEMORY_HANDLE_DESC CudaExtMemHandleDesc = {};
+		CudaExtMemHandleDesc.type = CU_EXTERNAL_MEMORY_HANDLE_TYPE_OPAQUE_FD;
+		CudaExtMemHandleDesc.handle.fd = Fd;
+		CudaExtMemHandleDesc.size = VulkanTexture->Surface.GetAllocationOffset() + VulkanTexture->Surface.GetMemorySize();
+
+		// import external memory
+		CUresult Result = FCUDAModule::CUDA().cuImportExternalMemory(&MappedExternalMemory, &CudaExtMemHandleDesc);
+		if (Result != CUDA_SUCCESS)
+		{
+			UE_LOG(LogPixelStreaming, Error, TEXT("Failed to import external memory from vulkan error: %d"), Result);
+		}
+	}
+
+#endif
+
+	CUmipmappedArray MappedMipArray = nullptr;
+	CUarray MappedArray = nullptr;
+
+	{
+		CUDA_EXTERNAL_MEMORY_MIPMAPPED_ARRAY_DESC MipmapDesc = {};
+		MipmapDesc.numLevels = 1;
+		MipmapDesc.offset = VulkanTexture->Surface.GetAllocationOffset();
+		MipmapDesc.arrayDesc.Width = Texture->GetSizeX();
+		MipmapDesc.arrayDesc.Height = Texture->GetSizeY();
+		MipmapDesc.arrayDesc.Depth = 0;
+		MipmapDesc.arrayDesc.NumChannels = 4;
+		MipmapDesc.arrayDesc.Format = CU_AD_FORMAT_UNSIGNED_INT8;
+		MipmapDesc.arrayDesc.Flags = CUDA_ARRAY3D_SURFACE_LDST | CUDA_ARRAY3D_COLOR_ATTACHMENT;
+
+		// get the CUarray from the external memory
+		CUresult Result = FCUDAModule::CUDA().cuExternalMemoryGetMappedMipmappedArray(&MappedMipArray, MappedExternalMemory, &MipmapDesc);
+		if (Result != CUDA_SUCCESS)
+		{
+			UE_LOG(LogPixelStreaming, Error, TEXT("Failed to bind mipmappedArray error: %d"), Result);
+		}
+	}
+
+	// get the CUarray from the external memory
+	CUresult MipMapArrGetLevelErr = FCUDAModule::CUDA().cuMipmappedArrayGetLevel(&MappedArray, MappedMipArray, 0);
+	if (MipMapArrGetLevelErr != CUDA_SUCCESS)
+	{
+		UE_LOG(LogPixelStreaming, Error, TEXT("Failed to bind to mip 0."));
+	}
+
+	FCUDAModule::CUDA().cuCtxPopCurrent(NULL);
+
+	InputFrame->SetTexture(MappedArray, AVEncoder::FVideoEncoderInputFrame::EUnderlyingRHI::Vulkan, Handle, [MappedArray, MappedMipArray, MappedExternalMemory](CUarray NativeTexture) {
+		// free the cuda types
+		FCUDAModule::CUDA().cuCtxPushCurrent(FModuleManager::GetModuleChecked<FCUDAModule>("CUDA").GetCudaContext());
+
+		if (MappedArray)
+		{
+			CUresult Result = FCUDAModule::CUDA().cuArrayDestroy(MappedArray);
+			if (Result != CUDA_SUCCESS)
+			{
+				UE_LOG(LogPixelStreaming, Error, TEXT("Failed to destroy MappedArray: %d"), Result);
+			}
+		}
+
+		if (MappedMipArray)
+		{
+			CUresult Result = FCUDAModule::CUDA().cuMipmappedArrayDestroy(MappedMipArray);
+			if (Result != CUDA_SUCCESS)
+			{
+				UE_LOG(LogPixelStreaming, Error, TEXT("Failed to destroy MappedMipArray: %d"), Result);
+			}
+		}
+
+		if (MappedExternalMemory)
+		{
+			CUresult Result = FCUDAModule::CUDA().cuDestroyExternalMemory(MappedExternalMemory);
+			if (Result != CUDA_SUCCESS)
+			{
+				UE_LOG(LogPixelStreaming, Error, TEXT("Failed to destroy MappedExternalMemoryArray: %d"), Result);
+			}
+		}
+
+		FCUDAModule::CUDA().cuCtxPopCurrent(NULL);
+	});
+}
+
+#if PLATFORM_WINDOWS
+void UE::PixelStreaming::FEncoderFrameFactory::SetTextureCUDAD3D11(AVEncoder::FVideoEncoderInputFrame* InputFrame, const FTexture2DRHIRef& Texture)
+{
+	FD3D11TextureBase* D3D11Texture = GetD3D11TextureFromRHITexture(Texture);
+	unsigned long long TextureMemorySize = D3D11Texture->GetMemorySize();
+
+	ID3D11Texture2D* D3D11NativeTexture = static_cast<ID3D11Texture2D*>(D3D11Texture->GetResource());
+
+	TRefCountPtr<IDXGIResource> DXGIResource;
+	HRESULT QueryResult = D3D11NativeTexture->QueryInterface(IID_PPV_ARGS(DXGIResource.GetInitReference()));
+	if (QueryResult != S_OK)
+	{
+		UE_LOG(LogPixelStreaming, Error, TEXT("Failed to get DX texture handle for importing memory to CUDA: %d"), QueryResult);
+	}
+
+	HANDLE D3D11TextureHandle;
+	DXGIResource->GetSharedHandle(&D3D11TextureHandle);
+	DXGIResource->Release();
+
+	FCUDAModule::CUDA().cuCtxPushCurrent(FModuleManager::GetModuleChecked<FCUDAModule>("CUDA").GetCudaContext());
+
+	CUexternalMemory MappedExternalMemory = nullptr;
+
+	{
+		// generate a cudaExternalMemoryHandleDesc
+		CUDA_EXTERNAL_MEMORY_HANDLE_DESC CudaExtMemHandleDesc = {};
+		CudaExtMemHandleDesc.type = CU_EXTERNAL_MEMORY_HANDLE_TYPE_D3D11_RESOURCE_KMT;
+		CudaExtMemHandleDesc.handle.win32.name = NULL;
+		CudaExtMemHandleDesc.handle.win32.handle = D3D11TextureHandle;
+		CudaExtMemHandleDesc.size = TextureMemorySize;
+		// Necessary for committed resources (DX11 and committed DX12 resources)
+		CudaExtMemHandleDesc.flags |= CUDA_EXTERNAL_MEMORY_DEDICATED;
+
+		// import external memory
+		CUresult Result = FCUDAModule::CUDA().cuImportExternalMemory(&MappedExternalMemory, &CudaExtMemHandleDesc);
+		if (Result != CUDA_SUCCESS)
+		{
+			UE_LOG(LogPixelStreaming, Error, TEXT("Failed to import external memory from vulkan error: %d"), Result);
+		}
+	}
+
+	CUmipmappedArray MappedMipArray = nullptr;
+	CUarray MappedArray = nullptr;
+
+	{
+		CUDA_EXTERNAL_MEMORY_MIPMAPPED_ARRAY_DESC MipmapDesc = {};
+		MipmapDesc.numLevels = 1;
+		MipmapDesc.offset = 0;
+		MipmapDesc.arrayDesc.Width = Texture->GetSizeX();
+		MipmapDesc.arrayDesc.Height = Texture->GetSizeY();
+		MipmapDesc.arrayDesc.Depth = 1;
+		MipmapDesc.arrayDesc.NumChannels = 4;
+		MipmapDesc.arrayDesc.Format = CU_AD_FORMAT_UNSIGNED_INT8;
+		MipmapDesc.arrayDesc.Flags = CUDA_ARRAY3D_SURFACE_LDST | CUDA_ARRAY3D_COLOR_ATTACHMENT;
+
+		// get the CUarray from the external memory
+		CUresult Result = FCUDAModule::CUDA().cuExternalMemoryGetMappedMipmappedArray(&MappedMipArray, MappedExternalMemory, &MipmapDesc);
+		if (Result != CUDA_SUCCESS)
+		{
+			UE_LOG(LogPixelStreaming, Error, TEXT("Failed to bind mipmappedArray error: %d"), Result);
+		}
+	}
+
+	// get the CUarray from the external memory
+	CUresult MipMapArrGetLevelErr = FCUDAModule::CUDA().cuMipmappedArrayGetLevel(&MappedArray, MappedMipArray, 0);
+	if (MipMapArrGetLevelErr != CUDA_SUCCESS)
+	{
+		UE_LOG(LogPixelStreaming, Error, TEXT("Failed to bind to mip 0."));
+	}
+
+	FCUDAModule::CUDA().cuCtxPopCurrent(NULL);
+
+	InputFrame->SetTexture(MappedArray, AVEncoder::FVideoEncoderInputFrame::EUnderlyingRHI::D3D11, nullptr, [MappedArray, MappedMipArray, MappedExternalMemory](CUarray NativeTexture) {
+		// free the cuda types
+		FCUDAModule::CUDA().cuCtxPushCurrent(FModuleManager::GetModuleChecked<FCUDAModule>("CUDA").GetCudaContext());
+
+		if (MappedArray)
+		{
+			CUresult Result = FCUDAModule::CUDA().cuArrayDestroy(MappedArray);
+			if (Result != CUDA_SUCCESS)
+			{
+				UE_LOG(LogPixelStreaming, Error, TEXT("Failed to destroy MappedArray: %d"), Result);
+			}
+		}
+
+		if (MappedMipArray)
+		{
+			CUresult Result = FCUDAModule::CUDA().cuMipmappedArrayDestroy(MappedMipArray);
+			if (Result != CUDA_SUCCESS)
+			{
+				UE_LOG(LogPixelStreaming, Error, TEXT("Failed to destroy MappedMipArray: %d"), Result);
+			}
+		}
+
+		if (MappedExternalMemory)
+		{
+			CUresult Result = FCUDAModule::CUDA().cuDestroyExternalMemory(MappedExternalMemory);
+			if (Result != CUDA_SUCCESS)
+			{
+				UE_LOG(LogPixelStreaming, Error, TEXT("Failed to destroy MappedExternalMemoryArray: %d"), Result);
+			}
+		}
+
+		FCUDAModule::CUDA().cuCtxPopCurrent(NULL);
+	});
+}
+
+void UE::PixelStreaming::FEncoderFrameFactory::SetTextureCUDAD3D12(AVEncoder::FVideoEncoderInputFrame* InputFrame, const FTexture2DRHIRef& Texture)
+{
+	FD3D12TextureBase* D3D12Texture = GetD3D12TextureFromRHITexture(Texture);
+	ID3D12Resource* NativeD3D12Resource = (ID3D12Resource*)Texture->GetNativeResource();
+	unsigned long long TextureMemorySize = D3D12Texture->GetMemorySize();
+
+	// Because we create our texture as RenderTargetable, it is created as a committed resource, which is what our current implementation here supports.
+	// To prevent a mystery crash in future, check that our resource is a committed resource
+	check(!D3D12Texture->GetResource()->IsPlacedResource());
+
+	TRefCountPtr<ID3D12Device> OwnerDevice;
+	HRESULT QueryResult;
+	if ((QueryResult = NativeD3D12Resource->GetDevice(IID_PPV_ARGS(OwnerDevice.GetInitReference()))) != S_OK)
+	{
+		UE_LOG(LogPixelStreaming, Error, TEXT("Failed to get DX texture handle for importing memory to CUDA: %d (Get Device)"), QueryResult);
+	}
+
+	//
+	// ID3D12Device::CreateSharedHandle gives as an NT Handle, and so we need to call CloseHandle on it
+	//
+	HANDLE D3D12TextureHandle;
+	if ((QueryResult = OwnerDevice->CreateSharedHandle(NativeD3D12Resource, NULL, GENERIC_ALL, NULL, &D3D12TextureHandle)) != S_OK)
+	{
+		UE_LOG(LogPixelStreaming, Error, TEXT("Failed to get DX texture handle for importing memory to CUDA: %d"), QueryResult);
+	}
+
+	FCUDAModule::CUDA().cuCtxPushCurrent(FModuleManager::GetModuleChecked<FCUDAModule>("CUDA").GetCudaContext());
+
+	CUexternalMemory MappedExternalMemory = nullptr;
+
+	{
+		// generate a cudaExternalMemoryHandleDesc
+		CUDA_EXTERNAL_MEMORY_HANDLE_DESC CudaExtMemHandleDesc = {};
+		CudaExtMemHandleDesc.type = CU_EXTERNAL_MEMORY_HANDLE_TYPE_D3D12_RESOURCE;
+		CudaExtMemHandleDesc.handle.win32.name = NULL;
+		CudaExtMemHandleDesc.handle.win32.handle = D3D12TextureHandle;
+		CudaExtMemHandleDesc.size = TextureMemorySize;
+		// Necessary for committed resources (DX11 and committed DX12 resources)
+		CudaExtMemHandleDesc.flags |= CUDA_EXTERNAL_MEMORY_DEDICATED;
+
+		// import external memory
+		CUresult Result = FCUDAModule::CUDA().cuImportExternalMemory(&MappedExternalMemory, &CudaExtMemHandleDesc);
+		if (Result != CUDA_SUCCESS)
+		{
+			UE_LOG(LogPixelStreaming, Error, TEXT("Failed to import external memory from vulkan error: %d"), Result);
+		}
+	}
+
+	CUmipmappedArray MappedMipArray = nullptr;
+	CUarray MappedArray = nullptr;
+
+	{
+		CUDA_EXTERNAL_MEMORY_MIPMAPPED_ARRAY_DESC MipmapDesc = {};
+		MipmapDesc.numLevels = 1;
+		MipmapDesc.offset = 0;
+		MipmapDesc.arrayDesc.Width = Texture->GetSizeX();
+		MipmapDesc.arrayDesc.Height = Texture->GetSizeY();
+		MipmapDesc.arrayDesc.Depth = 1;
+		MipmapDesc.arrayDesc.NumChannels = 4;
+		MipmapDesc.arrayDesc.Format = CU_AD_FORMAT_UNSIGNED_INT8;
+		MipmapDesc.arrayDesc.Flags = CUDA_ARRAY3D_SURFACE_LDST | CUDA_ARRAY3D_COLOR_ATTACHMENT;
+
+		// get the CUarray from the external memory
+		CUresult Result = FCUDAModule::CUDA().cuExternalMemoryGetMappedMipmappedArray(&MappedMipArray, MappedExternalMemory, &MipmapDesc);
+		if (Result != CUDA_SUCCESS)
+		{
+			UE_LOG(LogPixelStreaming, Error, TEXT("Failed to bind mipmappedArray error: %d"), Result);
+		}
+	}
+
+	// get the CUarray from the external memory
+	CUresult MipMapArrGetLevelErr = FCUDAModule::CUDA().cuMipmappedArrayGetLevel(&MappedArray, MappedMipArray, 0);
+	if (MipMapArrGetLevelErr != CUDA_SUCCESS)
+	{
+		UE_LOG(LogPixelStreaming, Error, TEXT("Failed to bind to mip 0."));
+	}
+
+	FCUDAModule::CUDA().cuCtxPopCurrent(NULL);
+
+	InputFrame->SetTexture(MappedArray, AVEncoder::FVideoEncoderInputFrame::EUnderlyingRHI::D3D12, D3D12TextureHandle, [MappedArray, MappedMipArray, MappedExternalMemory](CUarray NativeTexture) {
+		// free the cuda types
+		FCUDAModule::CUDA().cuCtxPushCurrent(FModuleManager::GetModuleChecked<FCUDAModule>("CUDA").GetCudaContext());
+
+		if (MappedArray)
+		{
+			CUresult Result = FCUDAModule::CUDA().cuArrayDestroy(MappedArray);
+			if (Result != CUDA_SUCCESS)
+			{
+				UE_LOG(LogPixelStreaming, Error, TEXT("Failed to destroy MappedArray: %d"), Result);
+			}
+		}
+
+		if (MappedMipArray)
+		{
+			CUresult Result = FCUDAModule::CUDA().cuMipmappedArrayDestroy(MappedMipArray);
+			if (Result != CUDA_SUCCESS)
+			{
+				UE_LOG(LogPixelStreaming, Error, TEXT("Failed to destroy MappedMipArray: %d"), Result);
+			}
+		}
+
+		if (MappedExternalMemory)
+		{
+			CUresult Result = FCUDAModule::CUDA().cuDestroyExternalMemory(MappedExternalMemory);
+			if (Result != CUDA_SUCCESS)
+			{
+				UE_LOG(LogPixelStreaming, Error, TEXT("Failed to destroy MappedExternalMemoryArray: %d"), Result);
+			}
+		}
+
+		FCUDAModule::CUDA().cuCtxPopCurrent(NULL);
+	});
+}
+#endif //PLATFORM_WINDOWS
\ No newline at end of file
diff --git a/Engine/Plugins/Media/PixelStreaming/Source/PixelStreaming/Private/EncoderFrameFactory.h b/Engine/Plugins/Media/PixelStreaming/Source/PixelStreaming/Private/EncoderFrameFactory.h
new file mode 100644
--- /dev/null
+++ b/Engine/Plugins/Media/PixelStreaming/Source/PixelStreaming/Private/EncoderFrameFactory.h
@@ -0,0 +1,44 @@
+// Copyright Epic Games, Inc. All Rights Reserved.
+
+#pragma once
+
+#include "VideoEncoderInput.h"
+#include "FrameBuffer.h"
+
+namespace UE {
+	namespace PixelStreaming {
+		/*
+		* Factory to create `AVEncoder::FVideoEncoderInputFrame` for use in Pixel Streaming encoders.
+		* This class is responsible for creating/managing the underlying `AVEncoder::FVideoEncoderInput` required to create
+		* `AVEncoder::FVideoEncoderInputFrame`.
+		* This class is not threadsafe and should only be accessed from a single thread.
+		*/
+		class FEncoderFrameFactory
+		{
+		public:
+			FEncoderFrameFactory();
+			~FEncoderFrameFactory();
+			AVEncoder::FVideoEncoderInputFrame* GetFrameAndSetTexture(int InWidth, int InHeight, FTexture2DRHIRef InTexture);
+			TSharedPtr<AVEncoder::FVideoEncoderInput> GetOrCreateVideoEncoderInput(int InWidth, int InHeight);
+			void SetResolution(int InWidth, int InHeight);
+
+		private:
+			AVEncoder::FVideoEncoderInputFrame* GetOrCreateFrame(int InWidth, int InHeight, const FTexture2DRHIRef InTexture);
+			void RemoveStaleTextures();
+			void FlushFrames();
+			TSharedPtr<AVEncoder::FVideoEncoderInput> CreateVideoEncoderInput(int InWidth, int InHeight) const;
+			void SetTexture(AVEncoder::FVideoEncoderInputFrame* Frame, const FTexture2DRHIRef& Texture);
+			void SetTextureCUDAVulkan(AVEncoder::FVideoEncoderInputFrame* Frame, const FTexture2DRHIRef& Texture);
+		#if PLATFORM_WINDOWS
+			void SetTextureCUDAD3D11(AVEncoder::FVideoEncoderInputFrame* Frame, const FTexture2DRHIRef& Texture);
+			void SetTextureCUDAD3D12(AVEncoder::FVideoEncoderInputFrame* Frame, const FTexture2DRHIRef& Texture);
+		#endif // PLATFORM_WINDOWS
+
+		private:
+			uint64 FrameId = 0;
+			TSharedPtr<AVEncoder::FVideoEncoderInput> EncoderInput;
+			// Store a mapping between raw textures and the FVideoEncoderInputFrames that wrap them
+			TMap<FTexture2DRHIRef, AVEncoder::FVideoEncoderInputFrame*> TextureToFrameMapping;
+		};
+	}
+}
diff --git a/Engine/Plugins/Media/PixelStreaming/Source/PixelStreaming/Private/FixedFPSPump.cpp b/Engine/Plugins/Media/PixelStreaming/Source/PixelStreaming/Private/FixedFPSPump.cpp
new file mode 100644
--- /dev/null
+++ b/Engine/Plugins/Media/PixelStreaming/Source/PixelStreaming/Private/FixedFPSPump.cpp
@@ -0,0 +1,93 @@
+// Copyright Epic Games, Inc. All Rights Reserved.
+
+#include "FixedFPSPump.h"
+#include "VideoSource.h"
+#include "Settings.h"
+#include "TextureSource.h"
+#include "FrameBuffer.h"
+#include "PlayerSession.h"
+
+UE::PixelStreaming::FFixedFPSPump& UE::PixelStreaming::FFixedFPSPump::Get()
+{
+	static UE::PixelStreaming::FFixedFPSPump Pump = UE::PixelStreaming::FFixedFPSPump();
+	return Pump;
+}
+
+UE::PixelStreaming::FFixedFPSPump::FFixedFPSPump()
+	: NextPumpEvent(FPlatformProcess::GetSynchEventFromPool(false))
+{
+	PumpThread = MakeUnique<FThread>(TEXT("PumpThread"), [this]() { PumpLoop(); });
+}
+
+UE::PixelStreaming::FFixedFPSPump::~FFixedFPSPump()
+{
+	bThreadRunning = false;
+	NextPumpEvent->Trigger();
+	PumpThread->Join();
+}
+
+void UE::PixelStreaming::FFixedFPSPump::UnregisterVideoSource(FPixelStreamingPlayerId PlayerId)
+{
+	FScopeLock Guard(&SourcesGuard);
+	VideoSources.Remove(PlayerId);
+	NextPumpEvent->Trigger();
+}
+
+void UE::PixelStreaming::FFixedFPSPump::RegisterVideoSource(FPixelStreamingPlayerId PlayerId, rtc::scoped_refptr<UE::PixelStreaming::IPumpedVideoSource> Source)
+{
+	FScopeLock Guard(&SourcesGuard);
+	VideoSources.Add(PlayerId, Source);
+	NextPumpEvent->Trigger();
+}
+
+void UE::PixelStreaming::FFixedFPSPump::PumpLoop()
+{
+	uint64 LastCycles = FPlatformTime::Cycles64();
+
+	while (bThreadRunning)
+	{
+		// No sources, so just wait.
+		if (VideoSources.Num() == 0)
+		{
+			NextPumpEvent->Wait();
+		}
+
+		{
+			FScopeLock Guard(&SourcesGuard);
+
+			const int32 FrameId = NextFrameId++;
+
+			// Pump each video source
+			TMap<FPixelStreamingPlayerId, rtc::scoped_refptr<UE::PixelStreaming::IPumpedVideoSource>>::TIterator Iter = VideoSources.CreateIterator();
+			for(; Iter; ++Iter)
+			{
+				rtc::scoped_refptr<UE::PixelStreaming::IPumpedVideoSource>& VideoSource = Iter.Value();
+
+				// If the pump is the last thing holding this video source, remove the video source.
+				if(VideoSource->HasOneRef())
+				{
+					Iter.RemoveCurrent();
+					continue;
+				}
+				
+				if(VideoSource->IsReadyForPump())
+				{
+					VideoSource->OnPump(FrameId);
+				}
+			}
+		}
+
+		// Sleep as long as we need for a constant FPS
+		const uint64 EndCycles	  = FPlatformTime::Cycles64();
+		const double DeltaMs	  = FPlatformTime::ToMilliseconds64(EndCycles - LastCycles);
+		const int32 FPS			  = UE::PixelStreaming::Settings::CVarPixelStreamingWebRTCFps.GetValueOnAnyThread();
+		const double FrameDeltaMs = 1000.0 / FPS;
+		const double SleepMs	  = FrameDeltaMs - DeltaMs;
+		LastCycles				  = EndCycles;
+
+		if (SleepMs > 0)
+		{
+			NextPumpEvent->Wait(SleepMs, false);
+		}
+	}
+}
diff --git a/Engine/Plugins/Media/PixelStreaming/Source/PixelStreaming/Private/FixedFPSPump.h b/Engine/Plugins/Media/PixelStreaming/Source/PixelStreaming/Private/FixedFPSPump.h
new file mode 100644
--- /dev/null
+++ b/Engine/Plugins/Media/PixelStreaming/Source/PixelStreaming/Private/FixedFPSPump.h
@@ -0,0 +1,59 @@
+// Copyright Epic Games, Inc. All Rights Reserved.
+#pragma once
+
+#include "PixelStreamingPlayerId.h"
+#include "HAL/Thread.h"
+#include "WebRTCIncludes.h"
+
+namespace UE {
+	namespace PixelStreaming {
+
+		class IPumpedVideoSource : public rtc::RefCountInterface
+		{
+			public:
+				virtual void OnPump(int32 FrameId) = 0;
+				virtual bool IsReadyForPump() const = 0;
+				
+				/* rtc::RefCountInterface */
+				virtual void AddRef() const = 0;
+				virtual rtc::RefCountReleaseStatus Release() const = 0;
+				virtual bool HasOneRef() const = 0;
+		};
+
+		/*
+		* Runs a seperate thread that "pumps" at a fixed FPS interval. WebRTC video sources may add themselves to be "pumped" 
+		* at a fixed interval indepedent of render FPS or game thread FPS. This is useful so that poorly 
+		* performing applications can still submit a constant amount of frames to WebRTC and are not penalized with a 
+		* bitrate reduction.
+		*/
+		class FFixedFPSPump
+		{
+
+		private:
+			// Private constructor as there should only ever be one instance of this class and the intended design is to use the static Get() method.
+			FFixedFPSPump();
+			~FFixedFPSPump();
+			void PumpLoop();
+
+		public:
+			
+			// Single accessor.
+			static FFixedFPSPump& Get();
+
+			void RegisterVideoSource(FPixelStreamingPlayerId PlayerId, rtc::scoped_refptr<IPumpedVideoSource> Source);
+			void UnregisterVideoSource(FPixelStreamingPlayerId PlayerId);
+
+		private:
+			FCriticalSection SourcesGuard;
+			TMap<FPixelStreamingPlayerId, rtc::scoped_refptr<IPumpedVideoSource>> VideoSources;
+
+
+			TUniquePtr<FThread> PumpThread;
+			bool bThreadRunning = true;
+			//FEvent* PlayersChangedEvent;
+			FEvent* NextPumpEvent;
+			int32 NextFrameId = 0;
+
+		};
+	}
+}
diff --git a/Engine/Plugins/Media/PixelStreaming/Source/PixelStreaming/Private/FrameBuffer.cpp b/Engine/Plugins/Media/PixelStreaming/Source/PixelStreaming/Private/FrameBuffer.cpp
new file mode 100644
--- /dev/null
+++ b/Engine/Plugins/Media/PixelStreaming/Source/PixelStreaming/Private/FrameBuffer.cpp
@@ -0,0 +1,91 @@
+// Copyright Epic Games, Inc. All Rights Reserved.
+
+#include "FrameBuffer.h"
+#include "TextureSource.h"
+
+/*
+* ----------------- FInitializeFrameBuffer -----------------
+*/
+
+UE::PixelStreaming::FInitializeFrameBuffer::FInitializeFrameBuffer(TSharedPtr<ITextureSource> InTextureSource)
+	: TextureSource(InTextureSource)
+{
+}
+
+UE::PixelStreaming::FInitializeFrameBuffer::~FInitializeFrameBuffer()
+{
+}
+
+int UE::PixelStreaming::FInitializeFrameBuffer::width() const
+{
+	return TextureSource->GetSourceWidth();
+}
+
+int UE::PixelStreaming::FInitializeFrameBuffer::height() const
+{
+	return TextureSource->GetSourceHeight();
+}
+
+/*
+* ----------------- FSimulcastFrameBuffer -----------------
+*/
+
+UE::PixelStreaming::FSimulcastFrameBuffer::FSimulcastFrameBuffer(TArray<TSharedPtr<ITextureSource>>& InTextureSources)
+	: TextureSources(InTextureSources)
+{
+}
+
+UE::PixelStreaming::FSimulcastFrameBuffer::~FSimulcastFrameBuffer()
+{
+}
+
+int UE::PixelStreaming::FSimulcastFrameBuffer::GetNumLayers() const
+{
+	return TextureSources.Num();
+}
+
+TSharedPtr<UE::PixelStreaming::ITextureSource> UE::PixelStreaming::FSimulcastFrameBuffer::GetLayerFrameSource(int LayerIndex) const
+{
+	checkf(LayerIndex >= 0 && LayerIndex < TextureSources.Num(), TEXT("Requested layer index was out of bounds."));
+	return TextureSources[LayerIndex];
+}
+
+int UE::PixelStreaming::FSimulcastFrameBuffer::width() const
+{
+	checkf(TextureSources.Num() > 0, TEXT("Must be at least one texture source to get the width from."));
+	return FMath::Max(TextureSources[0]->GetSourceWidth(), TextureSources[TextureSources.Num()-1]->GetSourceWidth());
+}
+
+int UE::PixelStreaming::FSimulcastFrameBuffer::height() const
+{
+	checkf(TextureSources.Num() > 0, TEXT("Must be at least one texture source to get the height from."));
+	return FMath::Max(TextureSources[0]->GetSourceHeight(), TextureSources[TextureSources.Num()-1]->GetSourceHeight());
+}
+
+/*
+* ----------------- FLayerFrameBuffer -----------------
+*/
+
+UE::PixelStreaming::FLayerFrameBuffer::FLayerFrameBuffer(TSharedPtr<UE::PixelStreaming::ITextureSource> InTextureSource)
+	: TextureSource(InTextureSource)
+{
+}
+
+UE::PixelStreaming::FLayerFrameBuffer::~FLayerFrameBuffer()
+{
+}
+
+FTexture2DRHIRef UE::PixelStreaming::FLayerFrameBuffer::GetFrame() const
+{
+	return TextureSource->GetTexture();
+}
+
+int UE::PixelStreaming::FLayerFrameBuffer::width() const
+{
+	return TextureSource->GetSourceWidth();
+}
+
+int UE::PixelStreaming::FLayerFrameBuffer::height() const
+{
+	return TextureSource->GetSourceHeight();
+}
diff --git a/Engine/Plugins/Media/PixelStreaming/Source/PixelStreaming/Private/FrameBuffer.h b/Engine/Plugins/Media/PixelStreaming/Source/PixelStreaming/Private/FrameBuffer.h
new file mode 100644
--- /dev/null
+++ b/Engine/Plugins/Media/PixelStreaming/Source/PixelStreaming/Private/FrameBuffer.h
@@ -0,0 +1,114 @@
+// Copyright Epic Games, Inc. All Rights Reserved.
+
+#pragma once
+
+#include "RHI.h"
+#include "WebRTCIncludes.h"
+#include "TextureSource.h"
+
+namespace UE {
+	namespace PixelStreaming {
+
+		enum FFrameBufferType
+		{
+			Initialize,
+			Simulcast,
+			Layer
+		};
+
+		/*
+		* ----------------- FFrameBuffer -----------------
+		* The base framebuffer that extends the WebRTC type.
+		*/
+		class FFrameBuffer : public webrtc::VideoFrameBuffer
+		{
+		public:
+			virtual ~FFrameBuffer() {}
+
+			virtual FFrameBufferType GetFrameBufferType() const = 0;
+
+			// Begin webrtc::VideoFrameBuffer interface
+			virtual webrtc::VideoFrameBuffer::Type type() const override
+			{
+				return webrtc::VideoFrameBuffer::Type::kNative;
+			}
+
+			virtual int width() const = 0;
+			virtual int height() const = 0;
+
+			virtual rtc::scoped_refptr<webrtc::I420BufferInterface> ToI420() override
+			{
+				unimplemented();
+				return nullptr;
+			}
+
+			virtual const webrtc::I420BufferInterface* GetI420() const override
+			{
+				unimplemented();
+				return nullptr;
+			}
+			// End webrtc::VideoFrameBuffer interface
+		};
+
+		/*
+		* ----------------- FInitializeFrameBuffer -----------------
+		* We use this frame to force webrtc to create encoders that will siphon off other streams
+		* but never get real frames directly pumped to them.
+		*/
+		class FInitializeFrameBuffer : public FFrameBuffer
+		{
+		public:
+			FInitializeFrameBuffer(TSharedPtr<ITextureSource> InTextureSource);
+			virtual ~FInitializeFrameBuffer();
+
+			virtual FFrameBufferType GetFrameBufferType() const { return Initialize; }
+
+			virtual int width() const override;
+			virtual int height() const override;
+
+		private:
+			TSharedPtr<ITextureSource> TextureSource;
+		};
+
+		/*
+		* ----------------- FSimulcastFrameBuffer -----------------
+		* Holds a number of textures (called layers) and each is assigned an index.
+		*/
+		class FSimulcastFrameBuffer : public FFrameBuffer
+		{
+		public:
+			FSimulcastFrameBuffer(TArray<TSharedPtr<ITextureSource>>& InTextureSources);
+			virtual ~FSimulcastFrameBuffer();
+
+			virtual FFrameBufferType GetFrameBufferType() const { return Simulcast; }
+			TSharedPtr<ITextureSource> GetLayerFrameSource(int LayerIndex) const;
+			int GetNumLayers() const;
+
+			virtual int width() const override;
+			virtual int height() const override;
+
+		private:
+			TArray<TSharedPtr<ITextureSource>>& TextureSources;
+		};
+
+		/*
+		* ----------------- FLayerFrameBuffer -----------------
+		* Holds a genuine single texture for encoding.
+		*/
+		class FLayerFrameBuffer : public FFrameBuffer
+		{
+		public:
+			FLayerFrameBuffer(TSharedPtr<ITextureSource> InTextureSource);
+			virtual ~FLayerFrameBuffer();
+
+			virtual FFrameBufferType GetFrameBufferType() const { return Layer; }
+			FTexture2DRHIRef GetFrame() const;
+
+			virtual int width() const override;
+			virtual int height() const override;
+
+		private:
+			TSharedPtr<ITextureSource> TextureSource;
+		};
+	}
+}
diff --git a/Engine/Plugins/Media/PixelStreaming/Source/PixelStreaming/Private/GPUFencePoller.cpp b/Engine/Plugins/Media/PixelStreaming/Source/PixelStreaming/Private/GPUFencePoller.cpp
new file mode 100644
--- /dev/null
+++ b/Engine/Plugins/Media/PixelStreaming/Source/PixelStreaming/Private/GPUFencePoller.cpp
@@ -0,0 +1,134 @@
+// Copyright Epic Games, Inc. All Rights Reserved.
+#include "GPUFencePoller.h"
+
+/*
+* ------------- UE::PixelStreaming::FGPUFencePoller ---------------------
+*/
+
+UE::PixelStreaming::FGPUFencePoller& UE::PixelStreaming::FGPUFencePoller::Get()
+{
+	static UE::PixelStreaming::FGPUFencePoller Poller;
+	return Poller;
+}
+
+UE::PixelStreaming::FGPUFencePoller::FGPUFencePoller()
+	: Runnable()
+	, Thread(FRunnableThread::Create(&Runnable, TEXT("Pixel Streaming GPUFencePoller")))
+{
+}
+
+UE::PixelStreaming::FGPUFencePoller::~FGPUFencePoller()
+{
+	Runnable.Stop();
+	Thread->Kill(true);
+}
+
+void UE::PixelStreaming::FGPUFencePoller::AddJob(FGPUFenceRHIRef Fence, TSharedRef<bool, ESPMode::ThreadSafe> bKeepPolling, TFunction<void()> FenceDoneCallback)
+{
+	checkf(Runnable.IsRunning(), TEXT("Poller runnable must be running to submit jobs to it."));
+	Runnable.AddJob({ Fence, bKeepPolling, FenceDoneCallback });
+}
+
+/*
+* ------------------- FPollJob ----------------------
+*/
+
+UE::PixelStreaming::FGPUFencePoller::FPollJob::FPollJob(FGPUFenceRHIRef InFence, TSharedRef<bool, ESPMode::ThreadSafe> bInKeepPolling, TFunction<void()> InFenceDoneCallback)
+	: Fence(InFence)
+	, bKeepPolling(bInKeepPolling)
+	, FenceDoneCallback(InFenceDoneCallback)
+{
+}
+
+UE::PixelStreaming::FGPUFencePoller::FPollJob::FPollJob()
+	: FPollJob(FGPUFenceRHIRef(), MakeShared<bool, ESPMode::ThreadSafe>(false), []() {})
+{
+}
+
+/*
+* ------------- FPollerRunnable ---------------------
+*/
+
+UE::PixelStreaming::FGPUFencePoller::FPollerRunnable::FPollerRunnable()
+	: bIsRunning(false)
+	, JobAddedEvent(FPlatformProcess::GetSynchEventFromPool(false))
+{
+}
+
+bool UE::PixelStreaming::FGPUFencePoller::FPollerRunnable::IsRunning() const
+{
+	return bIsRunning;
+}
+
+bool UE::PixelStreaming::FGPUFencePoller::FPollerRunnable::Init()
+{
+	return true;
+}
+
+void UE::PixelStreaming::FGPUFencePoller::FPollerRunnable::AddJob(FPollJob Job)
+{
+	// Note: this is threadsafe because use of TQueue
+	JobsToAdd.Enqueue(Job);
+	JobAddedEvent->Trigger();
+}
+
+uint32 UE::PixelStreaming::FGPUFencePoller::FPollerRunnable::Run()
+{
+	bIsRunning = true;
+
+	while (bIsRunning)
+	{
+
+		// Add any new jobs
+		while (!JobsToAdd.IsEmpty())
+		{
+			FPollJob* Job = JobsToAdd.Peek();
+			if (Job)
+			{
+				PollJobs.Add(*Job);
+			}
+			JobsToAdd.Pop();
+		}
+
+		// Iterate backwards so we can remove elements if the job is completed
+		for (int i = PollJobs.Num() - 1; i >= 0; i--)
+		{
+			FPollJob& Job = PollJobs[i];
+			if (Job.bKeepPolling.Get() == false)
+			{
+				PollJobs.RemoveAt(i);
+				continue;
+			}
+
+			bool bJobCompleted = Job.Fence->Poll();
+			if (bJobCompleted)
+			{
+				Job.FenceDoneCallback();
+				PollJobs.RemoveAt(i);
+			}
+		}
+
+		//No jobs so sleep this thread waiting for a job to be added.
+		if (PollJobs.Num() == 0)
+		{
+			JobAddedEvent->Wait(); 
+		}
+		else
+		{
+			// Free CPU usage up for other threads.
+			FPlatformProcess::Sleep(0.0f);
+		}
+	}
+
+	return 0;
+}
+
+void UE::PixelStreaming::FGPUFencePoller::FPollerRunnable::Stop()
+{
+	bIsRunning = false;
+}
+
+void UE::PixelStreaming::FGPUFencePoller::FPollerRunnable::Exit()
+{
+	bIsRunning = false;
+}
diff --git a/Engine/Plugins/Media/PixelStreaming/Source/PixelStreaming/Private/GPUFencePoller.h b/Engine/Plugins/Media/PixelStreaming/Source/PixelStreaming/Private/GPUFencePoller.h
new file mode 100644
--- /dev/null
+++ b/Engine/Plugins/Media/PixelStreaming/Source/PixelStreaming/Private/GPUFencePoller.h
@@ -0,0 +1,68 @@
+// Copyright Epic Games, Inc. All Rights Reserved.
+#pragma once
+
+#include "Containers/Queue.h"
+#include "Containers/Array.h"
+#include "RHI.h"
+#include "RHIResources.h"
+#include "Templates/SharedPointer.h"
+#include "HAL/Event.h"
+#include "GenericPlatform/GenericPlatformMisc.h"
+#include "HAL/RunnableThread.h"
+#include "HAL/Runnable.h"
+
+namespace UE {
+    namespace PixelStreaming {
+        /*
+        * A thread with a single purpose, to poll GPU fences until they are done.
+        * Polling jobs can be submitted to this thread using its static methods.
+        * Once a polling job is complete a callback is called on the submitted job.
+        */
+        class FGPUFencePoller
+        {
+            public:
+                static FGPUFencePoller& Get();
+
+            private:
+                /* Private constructor because the intended usage is through the static Get method. */
+                FGPUFencePoller();
+
+            public:
+                virtual ~FGPUFencePoller();
+                void AddJob(FGPUFenceRHIRef Fence, TSharedRef<bool, ESPMode::ThreadSafe> bKeepPolling, TFunction<void()> FenceDoneCallback);
+
+            private:
+                struct FPollJob
+                {
+                    FPollJob();
+                    FPollJob(FGPUFenceRHIRef InFence, TSharedRef<bool, ESPMode::ThreadSafe> bInKeepPolling, TFunction<void()> InFenceDoneCallback);
+                    FGPUFenceRHIRef Fence;
+                    TSharedRef<bool, ESPMode::ThreadSafe> bKeepPolling;
+                    TFunction<void()> FenceDoneCallback;
+                };
+
+                class FPollerRunnable : public FRunnable
+                {
+                public:
+                    FPollerRunnable();
+                    void AddJob(FPollJob Job);
+                    bool IsRunning() const;
+                    // Begin FRunnable interface.
+                    virtual bool Init() override;
+                    virtual uint32 Run() override;
+                    virtual void Stop() override;
+                    virtual void Exit() override;
+                    // End FRunnable interface
+                private:
+                    bool bIsRunning;
+                    FEvent* JobAddedEvent;
+                    TArray<FPollJob> PollJobs;
+                    TQueue<FPollJob, EQueueMode::Mpsc> JobsToAdd;
+                };
+
+            private:
+                FPollerRunnable Runnable;
+                FRunnableThread* Thread;
+        };
+    }
+}
diff --git a/Engine/Plugins/Media/PixelStreaming/Source/PixelStreaming/Private/IPixelStreamingSessions.h b/Engine/Plugins/Media/PixelStreaming/Source/PixelStreaming/Private/IPixelStreamingSessions.h
--- a/Engine/Plugins/Media/PixelStreaming/Source/PixelStreaming/Private/IPixelStreamingSessions.h
+++ b/Engine/Plugins/Media/PixelStreaming/Source/PixelStreaming/Private/IPixelStreamingSessions.h
@@ -3,24 +3,28 @@
 #pragma once
 
 #include "IPixelStreamingAudioSink.h"
-#include "PlayerId.h"
+#include "PixelStreamingPlayerId.h"
 #include "ProtocolDefs.h"
 #include "Containers/UnrealString.h"
 
-// `FPlayerSession` is only directly accessible to `FThreadSafePlayerSessions` as it completely owns the lifecycle of those objects.
-// This is nice from a safety point of view; however, some other objects in Pixel Streaming need to interface with `FPlayerSession`.
-// By using `FPlayerId` users can use this `IPixelStreamingSessions` interface to perform a limited set of allowable
-// actions on `FPlayerSession` objects.
-class IPixelStreamingSessions
-{
-public:
-	virtual int GetNumPlayers() const = 0;
-	virtual IPixelStreamingAudioSink* GetAudioSink(FPlayerId PlayerId) const = 0;
-	virtual IPixelStreamingAudioSink* GetUnlistenedAudioSink() const = 0;
-	virtual bool IsQualityController(FPlayerId PlayerId) const = 0;
-	virtual void SetQualityController(FPlayerId PlayerId) = 0;
-	virtual bool SendMessage(FPlayerId PlayerId, PixelStreamingProtocol::EToPlayerMsg Type, const FString& Descriptor) const = 0;
-	virtual void SendLatestQP(FPlayerId PlayerId, int LatestQP) const = 0;
-	virtual void SendFreezeFrameTo(FPlayerId PlayerId, const TArray64<uint8>& JpegBytes) const = 0;
-	virtual void PollWebRTCStats() const = 0;
-};
+namespace UE {
+	namespace PixelStreaming {
+		// `FPlayerSession` is only directly accessible to `FThreadSafePlayerSessions` as it completely owns the lifecycle of those objects.
+		// This is nice from a safety point of view; however, some other objects in Pixel Streaming need to interface with `FPlayerSession`.
+		// By using `FPixelStreamingPlayerId` users can use this `IPixelStreamingSessions` interface to perform a limited set of allowable
+		// actions on `FPlayerSession` objects.
+		class IPixelStreamingSessions
+		{
+		public:
+			virtual int GetNumPlayers() const = 0;
+			virtual IPixelStreamingAudioSink* GetAudioSink(FPixelStreamingPlayerId PlayerId) const = 0;
+			virtual IPixelStreamingAudioSink* GetUnlistenedAudioSink() const = 0;
+			virtual bool IsQualityController(FPixelStreamingPlayerId PlayerId) const = 0;
+			virtual void SetQualityController(FPixelStreamingPlayerId PlayerId) = 0;
+			virtual bool SendMessage(FPixelStreamingPlayerId PlayerId, Protocol::EToPlayerMsg Type, const FString& Descriptor) const = 0;
+			virtual void SendLatestQP(FPixelStreamingPlayerId PlayerId, int LatestQP) const = 0;
+			virtual void SendFreezeFrameTo(FPixelStreamingPlayerId PlayerId, const TArray64<uint8>& JpegBytes) const = 0;
+			virtual void PollWebRTCStats() const = 0;
+		};
+	}
+}
