Merging Pixel Streaming refactor using //UE5/Dev-TensorWorks-to-//UE5/Release-5.0

diff --git a/Engine/Plugins/Media/PixelStreaming/Source/PixelStreaming/Private/InputDevice.cpp b/Engine/Plugins/Media/PixelStreaming/Source/PixelStreaming/Private/InputDevice.cpp
--- a/Engine/Plugins/Media/PixelStreaming/Source/PixelStreaming/Private/InputDevice.cpp
+++ b/Engine/Plugins/Media/PixelStreaming/Source/PixelStreaming/Private/InputDevice.cpp
@@ -1,8 +1,8 @@
 // Copyright Epic Games, Inc. All Rights Reserved.
 
 #include "InputDevice.h"
-#include "PixelStreamerInputComponent.h"
-#include "PixelStreamingSettings.h"
+#include "PixelStreamingInputComponent.h"
+#include "Settings.h"
 #include "PixelStreamingModule.h"
 #include "ProtocolDefs.h"
 #include "JavaScriptKeyCodes.inl"
@@ -18,109 +18,115 @@
 #include "Framework/Application/SlateUser.h"
 #include "EditorPixelStreamingSettings.h"
 
-DECLARE_LOG_CATEGORY_EXTERN(PixelStreamerInputDevice, Log, VeryVerbose);
-DEFINE_LOG_CATEGORY(PixelStreamerInputDevice);
+DECLARE_LOG_CATEGORY_EXTERN(LogPixelStreamingInputDevice, Log, VeryVerbose);
+DEFINE_LOG_CATEGORY(LogPixelStreamingInputDevice);
 
-DECLARE_LOG_CATEGORY_EXTERN(PixelStreamerInput, Log, VeryVerbose);
-DEFINE_LOG_CATEGORY(PixelStreamerInput);
+DECLARE_LOG_CATEGORY_EXTERN(LogPixelStreamingInput, Log, VeryVerbose);
+DEFINE_LOG_CATEGORY(LogPixelStreamingInput);
 
-/**
- * When reading input from a browser then the cursor position will be sent
- * across with mouse events. We want to use this position and avoid getting the
- * cursor position from the operating system. This is not relevant to touch
- * events.
- */
-class FCursor : public ICursor
+namespace UE
 {
-public:
-	FCursor() {}
-	virtual ~FCursor() = default;
-	virtual FVector2D GetPosition() const override { return Position; }
-	virtual void SetPosition(const int32 X, const int32 Y) override { Position = FVector2D(X, Y); };
-	virtual void SetType(const EMouseCursor::Type InNewCursor) override{};
-	virtual EMouseCursor::Type GetType() const override { return EMouseCursor::Type::Default; };
-	virtual void GetSize(int32& Width, int32& Height) const override{};
-	virtual void Show(bool bShow) override{};
-	virtual void Lock(const RECT* const Bounds) override{};
-	virtual void SetTypeShape(EMouseCursor::Type InCursorType, void* CursorHandle) override{};
-
-private:
-	/** The cursor position sent across with mouse events. */
-	FVector2D Position;
-};
-
-/**
-* Wrap the GenericApplication layer so we can replace the cursor and override
-* certain behavior.
-*/
-class FApplicationWrapper : public GenericApplication
-{
-public:
-	FApplicationWrapper(TSharedPtr<GenericApplication> InWrappedApplication)
-		: GenericApplication(MakeShareable(new FCursor()))
-		, WrappedApplication(InWrappedApplication)
+	namespace PixelStreaming
 	{
-		// Whether we want to always consider the mouse as attached. This allow
-		// us to run Pixel Streaming on a machine which has no physical mouse
-		// and just let the browser supply mouse positions.
-		const UPixelStreamingSettings* Settings = GetDefault<UPixelStreamingSettings>();
-		check(Settings);
-		bMouseAlwaysAttached = Settings->bPixelStreamerMouseAlwaysAttached;
-	}
-
-	/**
-	 * Functions passed directly to the wrapped application.
-	 */
-
-	virtual void SetMessageHandler(const TSharedRef<FGenericApplicationMessageHandler>& InMessageHandler) { WrappedApplication->SetMessageHandler(InMessageHandler); }
-	virtual void PollGameDeviceState(const float TimeDelta) { WrappedApplication->PollGameDeviceState(TimeDelta); }
-	virtual void PumpMessages(const float TimeDelta) { WrappedApplication->PumpMessages(TimeDelta); }
-	virtual void ProcessDeferredEvents(const float TimeDelta) { WrappedApplication->ProcessDeferredEvents(TimeDelta); }
-	virtual void Tick(const float TimeDelta) { WrappedApplication->Tick(TimeDelta); }
-	virtual TSharedRef<FGenericWindow> MakeWindow() { return WrappedApplication->MakeWindow(); }
-	virtual void InitializeWindow(const TSharedRef<FGenericWindow>& Window, const TSharedRef<FGenericWindowDefinition>& InDefinition, const TSharedPtr<FGenericWindow>& InParent, const bool bShowImmediately) { WrappedApplication->InitializeWindow(Window, InDefinition, InParent, bShowImmediately); }
-	virtual void SetCapture(const TSharedPtr<FGenericWindow>& InWindow) { WrappedApplication->SetCapture(InWindow); }
-	virtual void* GetCapture(void) const { return WrappedApplication->GetCapture(); }
-	virtual FModifierKeysState GetModifierKeys() const { return WrappedApplication->GetModifierKeys(); }
-	virtual TSharedPtr<FGenericWindow> GetWindowUnderCursor() { return WrappedApplication->GetWindowUnderCursor(); }
-	virtual void SetHighPrecisionMouseMode(const bool Enable, const TSharedPtr<FGenericWindow>& InWindow) { WrappedApplication->SetHighPrecisionMouseMode(Enable, InWindow); };
-	virtual bool IsUsingHighPrecisionMouseMode() const { return WrappedApplication->IsUsingHighPrecisionMouseMode(); }
-	virtual bool IsUsingTrackpad() const { return WrappedApplication->IsUsingTrackpad(); }
-	virtual bool IsGamepadAttached() const { return WrappedApplication->IsGamepadAttached(); }
-	virtual void RegisterConsoleCommandListener(const FOnConsoleCommandListener& InListener) { WrappedApplication->RegisterConsoleCommandListener(InListener); }
-	virtual void AddPendingConsoleCommand(const FString& InCommand) { WrappedApplication->AddPendingConsoleCommand(InCommand); }
-	virtual FPlatformRect GetWorkArea(const FPlatformRect& CurrentWindow) const { return WrappedApplication->GetWorkArea(CurrentWindow); }
-	virtual bool TryCalculatePopupWindowPosition(const FPlatformRect& InAnchor, const FVector2D& InSize, const FVector2D& ProposedPlacement, const EPopUpOrientation::Type Orientation, /*OUT*/ FVector2D* const CalculatedPopUpPosition) const { return WrappedApplication->TryCalculatePopupWindowPosition(InAnchor, InSize, ProposedPlacement, Orientation, CalculatedPopUpPosition); }
-	virtual void GetInitialDisplayMetrics(FDisplayMetrics& OutDisplayMetrics) const { WrappedApplication->GetInitialDisplayMetrics(OutDisplayMetrics); }
-	virtual EWindowTitleAlignment::Type GetWindowTitleAlignment() const { return WrappedApplication->GetWindowTitleAlignment(); }
-	virtual EWindowTransparency GetWindowTransparencySupport() const { return WrappedApplication->GetWindowTransparencySupport(); }
-	virtual void DestroyApplication() { WrappedApplication->DestroyApplication(); }
-	virtual IInputInterface* GetInputInterface() { return WrappedApplication->GetInputInterface(); }
-	virtual ITextInputMethodSystem* GetTextInputMethodSystem() { return WrappedApplication->GetTextInputMethodSystem(); }
-	virtual void SendAnalytics(IAnalyticsProvider* Provider) { WrappedApplication->SendAnalytics(Provider); }
-	virtual bool SupportsSystemHelp() const { return WrappedApplication->SupportsSystemHelp(); }
-	virtual void ShowSystemHelp() { WrappedApplication->ShowSystemHelp(); }
-	virtual bool ApplicationLicenseValid(FPlatformUserId PlatformUser = PLATFORMUSERID_NONE) { return WrappedApplication->ApplicationLicenseValid(PlatformUser); }
-
-	/**
-	 * Functions with overridden behavior.
-	 */
-	virtual bool IsCursorDirectlyOverSlateWindow() const { return true; }
-	virtual bool IsMouseAttached() const { return bMouseAlwaysAttached ? true : WrappedApplication->IsMouseAttached(); }
-
-	TSharedPtr<GenericApplication> WrappedApplication;
-	bool bMouseAlwaysAttached;
-};
-
-const FVector2D FInputDevice::UnfocusedPos(-1.0f, -1.0f);
-const size_t FInputDevice::MessageHeaderOffset = 3;
+		/**
+		 * When reading input from a browser then the cursor position will be sent
+		 * across with mouse events. We want to use this position and avoid getting the
+		 * cursor position from the operating system. This is not relevant to touch
+		 * events.
+		 */
+		class FCursor : public ICursor
+		{
+		public:
+			FCursor() {}
+			virtual ~FCursor() = default;
+			virtual FVector2D GetPosition() const override { return Position; }
+			virtual void SetPosition(const int32 X, const int32 Y) override { Position = FVector2D(X, Y); };
+			virtual void SetType(const EMouseCursor::Type InNewCursor) override{};
+			virtual EMouseCursor::Type GetType() const override { return EMouseCursor::Type::Default; };
+			virtual void GetSize(int32& Width, int32& Height) const override{};
+			virtual void Show(bool bShow) override{};
+			virtual void Lock(const RECT* const Bounds) override{};
+			virtual void SetTypeShape(EMouseCursor::Type InCursorType, void* CursorHandle) override{};
+
+		private:
+			/** The cursor position sent across with mouse events. */
+			FVector2D Position;
+		};
+
+		/**
+		* Wrap the GenericApplication layer so we can replace the cursor and override
+		* certain behavior.
+		*/
+		class FApplicationWrapper : public GenericApplication
+		{
+		public:
+			FApplicationWrapper(TSharedPtr<GenericApplication> InWrappedApplication)
+				: GenericApplication(MakeShareable(new FCursor()))
+				, WrappedApplication(InWrappedApplication)
+			{
+				// Whether we want to always consider the mouse as attached. This allow
+				// us to run Pixel Streaming on a machine which has no physical mouse
+				// and just let the browser supply mouse positions.
+				const UPixelStreamingSettings* Settings = GetDefault<UPixelStreamingSettings>();
+				check(Settings);
+				bMouseAlwaysAttached = Settings->bMouseAlwaysAttached;
+			}
 
-FInputDevice::FInputDevice(const TSharedRef<FGenericApplicationMessageHandler>& InMessageHandler)
+			/**
+			 * Functions passed directly to the wrapped application.
+			 */
+
+			virtual void SetMessageHandler(const TSharedRef<FGenericApplicationMessageHandler>& InMessageHandler) { WrappedApplication->SetMessageHandler(InMessageHandler); }
+			virtual void PollGameDeviceState(const float TimeDelta) { WrappedApplication->PollGameDeviceState(TimeDelta); }
+			virtual void PumpMessages(const float TimeDelta) { WrappedApplication->PumpMessages(TimeDelta); }
+			virtual void ProcessDeferredEvents(const float TimeDelta) { WrappedApplication->ProcessDeferredEvents(TimeDelta); }
+			virtual void Tick(const float TimeDelta) { WrappedApplication->Tick(TimeDelta); }
+			virtual TSharedRef<FGenericWindow> MakeWindow() { return WrappedApplication->MakeWindow(); }
+			virtual void InitializeWindow(const TSharedRef<FGenericWindow>& Window, const TSharedRef<FGenericWindowDefinition>& InDefinition, const TSharedPtr<FGenericWindow>& InParent, const bool bShowImmediately) { WrappedApplication->InitializeWindow(Window, InDefinition, InParent, bShowImmediately); }
+			virtual void SetCapture(const TSharedPtr<FGenericWindow>& InWindow) { WrappedApplication->SetCapture(InWindow); }
+			virtual void* GetCapture(void) const { return WrappedApplication->GetCapture(); }
+			virtual FModifierKeysState GetModifierKeys() const { return WrappedApplication->GetModifierKeys(); }
+			virtual TSharedPtr<FGenericWindow> GetWindowUnderCursor() { return WrappedApplication->GetWindowUnderCursor(); }
+			virtual void SetHighPrecisionMouseMode(const bool Enable, const TSharedPtr<FGenericWindow>& InWindow) { WrappedApplication->SetHighPrecisionMouseMode(Enable, InWindow); };
+			virtual bool IsUsingHighPrecisionMouseMode() const { return WrappedApplication->IsUsingHighPrecisionMouseMode(); }
+			virtual bool IsUsingTrackpad() const { return WrappedApplication->IsUsingTrackpad(); }
+			virtual bool IsGamepadAttached() const { return WrappedApplication->IsGamepadAttached(); }
+			virtual void RegisterConsoleCommandListener(const FOnConsoleCommandListener& InListener) { WrappedApplication->RegisterConsoleCommandListener(InListener); }
+			virtual void AddPendingConsoleCommand(const FString& InCommand) { WrappedApplication->AddPendingConsoleCommand(InCommand); }
+			virtual FPlatformRect GetWorkArea(const FPlatformRect& CurrentWindow) const { return WrappedApplication->GetWorkArea(CurrentWindow); }
+			virtual bool TryCalculatePopupWindowPosition(const FPlatformRect& InAnchor, const FVector2D& InSize, const FVector2D& ProposedPlacement, const EPopUpOrientation::Type Orientation, /*OUT*/ FVector2D* const CalculatedPopUpPosition) const { return WrappedApplication->TryCalculatePopupWindowPosition(InAnchor, InSize, ProposedPlacement, Orientation, CalculatedPopUpPosition); }
+			virtual void GetInitialDisplayMetrics(FDisplayMetrics& OutDisplayMetrics) const { WrappedApplication->GetInitialDisplayMetrics(OutDisplayMetrics); }
+			virtual EWindowTitleAlignment::Type GetWindowTitleAlignment() const { return WrappedApplication->GetWindowTitleAlignment(); }
+			virtual EWindowTransparency GetWindowTransparencySupport() const { return WrappedApplication->GetWindowTransparencySupport(); }
+			virtual void DestroyApplication() { WrappedApplication->DestroyApplication(); }
+			virtual IInputInterface* GetInputInterface() { return WrappedApplication->GetInputInterface(); }
+			virtual ITextInputMethodSystem* GetTextInputMethodSystem() { return WrappedApplication->GetTextInputMethodSystem(); }
+			virtual void SendAnalytics(IAnalyticsProvider* Provider) { WrappedApplication->SendAnalytics(Provider); }
+			virtual bool SupportsSystemHelp() const { return WrappedApplication->SupportsSystemHelp(); }
+			virtual void ShowSystemHelp() { WrappedApplication->ShowSystemHelp(); }
+			virtual bool ApplicationLicenseValid(FPlatformUserId PlatformUser = PLATFORMUSERID_NONE) { return WrappedApplication->ApplicationLicenseValid(PlatformUser); }
+
+			/**
+			 * Functions with overridden behavior.
+			 */
+			virtual bool IsCursorDirectlyOverSlateWindow() const { return true; }
+			virtual bool IsMouseAttached() const { return bMouseAlwaysAttached ? true : WrappedApplication->IsMouseAttached(); }
+
+			TSharedPtr<GenericApplication> WrappedApplication;
+			bool bMouseAlwaysAttached;
+		};
+	} // namespace PixelStreaming
+} // namespace UE
+
+const FVector2D UE::PixelStreaming::FInputDevice::UnfocusedPos(-1.0f, -1.0f);
+const size_t UE::PixelStreaming::FInputDevice::MessageHeaderOffset = 3;
+
+UE::PixelStreaming::FInputDevice::FInputDevice(const TSharedRef<FGenericApplicationMessageHandler>& InMessageHandler)
 	: PixelStreamerApplicationWrapper(MakeShareable(new FApplicationWrapper(FSlateApplication::Get().GetPlatformApplication())))
 	, MessageHandler(InMessageHandler)
 	, bFakingTouchEvents(FSlateApplication::Get().IsFakingTouchEvents())
 	, FocusedPos(UnfocusedPos)
-	, PixelStreamingModule(FPixelStreamingModule::GetModule())
+	, PixelStreamingModule(UE::PixelStreaming::FPixelStreamingModule::GetModule())
 {
 	if (GEngine->GameViewport && !GEngine->GameViewport->HasSoftwareCursor(EMouseCursor::Default))
 	{
@@ -132,16 +138,16 @@ FInputDevice::FInputDevice(const TSharedRef<FGenericApplicationMessageHandler>&
 		// Check to see if we want to hide the cursor (by making it invisible).
 		// This is required if we want to make the cursor client-side (displayed
 		// only by the the browser).
-		bool bHideCursor = PixelStreamingSettings::IsPixelStreamingHideCursor();
+		bool bHideCursor = UE::PixelStreaming::Settings::IsPixelStreamingHideCursor();
 		if (bHideCursor)
 		{
-			GEngine->GameViewport->AddSoftwareCursor(EMouseCursor::Default, Settings->PixelStreamerHiddenCursorClassName);
-			GEngine->GameViewport->AddSoftwareCursor(EMouseCursor::TextEditBeam, Settings->PixelStreamerHiddenCursorClassName);
+			GEngine->GameViewport->AddSoftwareCursor(EMouseCursor::Default, Settings->HiddenCursorClassName);
+			GEngine->GameViewport->AddSoftwareCursor(EMouseCursor::TextEditBeam, Settings->HiddenCursorClassName);
 		}
 		else
 		{
-			GEngine->GameViewport->AddSoftwareCursor(EMouseCursor::Default, Settings->PixelStreamerDefaultCursorClassName);
-			GEngine->GameViewport->AddSoftwareCursor(EMouseCursor::TextEditBeam, Settings->PixelStreamerTextEditBeamCursorClassName);
+			GEngine->GameViewport->AddSoftwareCursor(EMouseCursor::Default, Settings->DefaultCursorClassName);
+			GEngine->GameViewport->AddSoftwareCursor(EMouseCursor::TextEditBeam, Settings->TextEditBeamCursorClassName);
 		}
 	}
 }
@@ -158,7 +164,7 @@ void QuantizeAndNormalize(const FVector2D& InPos, uint16& OutX, uint16& OutY)
 
 bool FilterKey(const FKey& Key)
 {
-	for (auto&& FilteredKey : PixelStreamingSettings::FilteredKeys)
+	for (auto&& FilteredKey : UE::PixelStreaming::Settings::FilteredKeys)
 	{
 		if (FilteredKey == Key)
 			return false;
@@ -166,7 +172,7 @@ bool FilterKey(const FKey& Key)
 	return true;
 }
 
-void FInputDevice::Tick(float DeltaTime)
+void UE::PixelStreaming::FInputDevice::Tick(float DeltaTime)
 {
 	FEvent Event;
 	while (Events.Dequeue(Event))
@@ -183,7 +189,7 @@ void FInputDevice::Tick(float DeltaTime)
 				uint8 JavaScriptKeyCode;
 				bool IsRepeat;
 				Event.GetKeyDown(JavaScriptKeyCode, IsRepeat);
-				const FKey* AgnosticKey = JavaScriptKeyCodeToFKey[JavaScriptKeyCode];
+				const FKey* AgnosticKey = UE::PixelStreaming::JavaScriptKeyCodeToFKey[JavaScriptKeyCode];
 				if (FilterKey(*AgnosticKey))
 				{
 					const uint32* KeyCodePtr;
@@ -192,7 +198,7 @@ void FInputDevice::Tick(float DeltaTime)
 					uint32 KeyCode = KeyCodePtr ? *KeyCodePtr : 0;
 					uint32 CharacterCode = CharacterCodePtr ? *CharacterCodePtr : 0;
 					MessageHandler->OnKeyDown(KeyCode, CharacterCode, IsRepeat);
-					UE_LOG(PixelStreamerInputDevice, Verbose, TEXT("KEY_DOWN: KeyCode = %d; CharacterCode = %d; IsRepeat = %s"), KeyCode, CharacterCode, IsRepeat ? TEXT("True") : TEXT("False"));
+					UE_LOG(LogPixelStreamingInputDevice, Verbose, TEXT("KEY_DOWN: KeyCode = %d; CharacterCode = %d; IsRepeat = %s"), KeyCode, CharacterCode, IsRepeat ? TEXT("True") : TEXT("False"));
 				}
 			}
 			break;
@@ -200,7 +206,7 @@ void FInputDevice::Tick(float DeltaTime)
 			{
 				uint8 JavaScriptKeyCode;
 				Event.GetKeyUp(JavaScriptKeyCode);
-				const FKey* AgnosticKey = JavaScriptKeyCodeToFKey[JavaScriptKeyCode];
+				const FKey* AgnosticKey = UE::PixelStreaming::JavaScriptKeyCodeToFKey[JavaScriptKeyCode];
 				if (FilterKey(*AgnosticKey))
 				{
 					const uint32* KeyCodePtr;
@@ -209,7 +215,7 @@ void FInputDevice::Tick(float DeltaTime)
 					uint32 KeyCode = KeyCodePtr ? *KeyCodePtr : 0;
 					uint32 CharacterCode = CharacterCodePtr ? *CharacterCodePtr : 0;
 					MessageHandler->OnKeyUp(KeyCode, CharacterCode, false); // Key up events are never repeats.
-					UE_LOG(PixelStreamerInputDevice, Verbose, TEXT("KEY_UP: KeyCode = %d; CharacterCode = %d"), KeyCode, CharacterCode);
+					UE_LOG(LogPixelStreamingInputDevice, Verbose, TEXT("KEY_UP: KeyCode = %d; CharacterCode = %d"), KeyCode, CharacterCode);
 				}
 			}
 			break;
@@ -218,7 +224,7 @@ void FInputDevice::Tick(float DeltaTime)
 				TCHAR UnicodeCharacter;
 				Event.GetCharacterCode(UnicodeCharacter);
 				MessageHandler->OnKeyChar(UnicodeCharacter, false); // Key press repeat not yet available but are not intrinsically used.
-				UE_LOG(PixelStreamerInputDevice, Verbose, TEXT("KEY_PRESSED: Character = '%c'"), UnicodeCharacter);
+				UE_LOG(LogPixelStreamingInputDevice, Verbose, TEXT("KEY_PRESSED: Character = '%c'"), UnicodeCharacter);
 			}
 			break;
 			case EventType::MOUSE_ENTER:
@@ -230,7 +236,7 @@ void FInputDevice::Tick(float DeltaTime)
 				// Make sure the viewport is active.
 				FSlateApplication::Get().ProcessApplicationActivationEvent(true);
 
-				UE_LOG(PixelStreamerInputDevice, Verbose, TEXT("MOUSE_ENTER"));
+				UE_LOG(LogPixelStreamingInputDevice, Verbose, TEXT("MOUSE_ENTER"));
 			}
 			break;
 			case EventType::MOUSE_LEAVE:
@@ -238,7 +244,7 @@ void FInputDevice::Tick(float DeltaTime)
 				// Restore normal application layer.
 				FSlateApplication::Get().OverridePlatformApplication(PixelStreamerApplicationWrapper->WrappedApplication);
 
-				UE_LOG(PixelStreamerInputDevice, Verbose, TEXT("MOUSE_LEAVE"));
+				UE_LOG(LogPixelStreamingInputDevice, Verbose, TEXT("MOUSE_LEAVE"));
 			}
 			break;
 			case EventType::MOUSE_MOVE:
@@ -254,7 +260,7 @@ void FInputDevice::Tick(float DeltaTime)
 				FVector2D CursorPos = GEngine->GameViewport->GetWindow()->GetPositionInScreen() + FVector2D(PosX, PosY);
 				PixelStreamerApplicationWrapper->Cursor->SetPosition(CursorPos.X, CursorPos.Y);
 				MessageHandler->OnRawMouseMove(DeltaX, DeltaY);
-				UE_LOG(PixelStreamerInputDevice, VeryVerbose, TEXT("MOUSE_MOVE: Pos = (%d, %d); CursorPos = (%d, %d); Delta = (%d, %d)"), PosX, PosY, static_cast<int>(CursorPos.X), static_cast<int>(CursorPos.Y), DeltaX, DeltaY);
+				UE_LOG(LogPixelStreamingInputDevice, VeryVerbose, TEXT("MOUSE_MOVE: Pos = (%d, %d); CursorPos = (%d, %d); Delta = (%d, %d)"), PosX, PosY, static_cast<int>(CursorPos.X), static_cast<int>(CursorPos.Y), DeltaX, DeltaY);
 			}
 			break;
 			case EventType::MOUSE_DOWN:
@@ -275,7 +281,7 @@ void FInputDevice::Tick(float DeltaTime)
 				FVector2D CursorPos = GEngine->GameViewport->GetWindow()->GetPositionInScreen() + FVector2D(PosX, PosY);
 				PixelStreamerApplicationWrapper->Cursor->SetPosition(CursorPos.X, CursorPos.Y);
 				MessageHandler->OnMouseDown(GEngine->GameViewport->GetWindow()->GetNativeWindow(), Button, CursorPos);
-				UE_LOG(PixelStreamerInputDevice, Verbose, TEXT("MOUSE_DOWN: Button = %d; Pos = (%d, %d); CursorPos = (%d, %d)"), Button, PosX, PosY, static_cast<int>(CursorPos.X), static_cast<int>(CursorPos.Y));
+				UE_LOG(LogPixelStreamingInputDevice, Verbose, TEXT("MOUSE_DOWN: Button = %d; Pos = (%d, %d); CursorPos = (%d, %d)"), Button, PosX, PosY, static_cast<int>(CursorPos.X), static_cast<int>(CursorPos.Y));
 
 				// The browser may be faking a mouse when touching so it will send
 				// over a mouse down event.
@@ -291,7 +297,7 @@ void FInputDevice::Tick(float DeltaTime)
 				FVector2D CursorPos = GEngine->GameViewport->GetWindow()->GetPositionInScreen() + FVector2D(PosX, PosY);
 				PixelStreamerApplicationWrapper->Cursor->SetPosition(CursorPos.X, CursorPos.Y);
 				MessageHandler->OnMouseUp(Button);
-				UE_LOG(PixelStreamerInputDevice, Verbose, TEXT("MOUSE_UP: Button = %d; Pos = (%d, %d); CursorPos = (%d, %d)"), Button, PosX, PosY, static_cast<int>(CursorPos.X), static_cast<int>(CursorPos.Y));
+				UE_LOG(LogPixelStreamingInputDevice, Verbose, TEXT("MOUSE_UP: Button = %d; Pos = (%d, %d); CursorPos = (%d, %d)"), Button, PosX, PosY, static_cast<int>(CursorPos.X), static_cast<int>(CursorPos.Y));
 			}
 			break;
 			case EventType::MOUSE_WHEEL:
@@ -303,7 +309,7 @@ void FInputDevice::Tick(float DeltaTime)
 				const float SpinFactor = 1 / 120.0f;
 				FVector2D CursorPos = GEngine->GameViewport->GetWindow()->GetPositionInScreen() + FVector2D(PosX, PosY);
 				MessageHandler->OnMouseWheel(Delta * SpinFactor, CursorPos);
-				UE_LOG(PixelStreamerInputDevice, Verbose, TEXT("MOUSE_WHEEL: Delta = %d; Pos = (%d, %d); CursorPos = (%d, %d)"), Delta, PosX, PosY, static_cast<int>(CursorPos.X), static_cast<int>(CursorPos.Y));
+				UE_LOG(LogPixelStreamingInputDevice, Verbose, TEXT("MOUSE_WHEEL: Delta = %d; Pos = (%d, %d); CursorPos = (%d, %d)"), Delta, PosX, PosY, static_cast<int>(CursorPos.X), static_cast<int>(CursorPos.Y));
 			}
 			break;
 			case EventType::TOUCH_START:
@@ -315,7 +321,7 @@ void FInputDevice::Tick(float DeltaTime)
 				Event.GetTouch(TouchIndex, PosX, PosY, Force);
 				FVector2D CursorPos = GEngine->GameViewport->GetWindow()->GetPositionInScreen() + FVector2D(PosX, PosY);
 				MessageHandler->OnTouchStarted(GEngine->GameViewport->GetWindow()->GetNativeWindow(), CursorPos, Force / 255.0f, TouchIndex, 0); // TODO: ControllerId?
-				UE_LOG(PixelStreamerInputDevice, Verbose, TEXT("TOUCH_START: TouchIndex = %d; Pos = (%d, %d); CursorPos = (%d, %d); Force = %.3f"), TouchIndex, PosX, PosY, static_cast<int>(CursorPos.X), static_cast<int>(CursorPos.Y), Force / 255.0f);
+				UE_LOG(LogPixelStreamingInputDevice, Verbose, TEXT("TOUCH_START: TouchIndex = %d; Pos = (%d, %d); CursorPos = (%d, %d); Force = %.3f"), TouchIndex, PosX, PosY, static_cast<int>(CursorPos.X), static_cast<int>(CursorPos.Y), Force / 255.0f);
 
 				// If the user starts a touch then they may be focusing an editable
 				// widget.
@@ -331,7 +337,7 @@ void FInputDevice::Tick(float DeltaTime)
 				Event.GetTouch(TouchIndex, PosX, PosY, Force);
 				FVector2D CursorPos = GEngine->GameViewport->GetWindow()->GetPositionInScreen() + FVector2D(PosX, PosY);
 				MessageHandler->OnTouchEnded(CursorPos, TouchIndex, 0); // TODO: ControllerId?
-				UE_LOG(PixelStreamerInputDevice, Verbose, TEXT("TOUCH_END: TouchIndex = %d; Pos = (%d, %d); CursorPos = (%d, %d)"), TouchIndex, PosX, PosY, static_cast<int>(CursorPos.X), static_cast<int>(CursorPos.Y));
+				UE_LOG(LogPixelStreamingInputDevice, Verbose, TEXT("TOUCH_END: TouchIndex = %d; Pos = (%d, %d); CursorPos = (%d, %d)"), TouchIndex, PosX, PosY, static_cast<int>(CursorPos.X), static_cast<int>(CursorPos.Y));
 			}
 			break;
 			case EventType::TOUCH_MOVE:
@@ -343,7 +349,7 @@ void FInputDevice::Tick(float DeltaTime)
 				Event.GetTouch(TouchIndex, PosX, PosY, Force);
 				FVector2D CursorPos = GEngine->GameViewport->GetWindow()->GetPositionInScreen() + FVector2D(PosX, PosY);
 				MessageHandler->OnTouchMoved(CursorPos, Force / 255.0f, TouchIndex, 0); // TODO: ControllerId?
-				UE_LOG(PixelStreamerInputDevice, VeryVerbose, TEXT("TOUCH_MOVE: TouchIndex = %d; Pos = (%d, %d); CursorPos = (%d, %d); Force = %.3f"), TouchIndex, PosX, PosY, static_cast<int>(CursorPos.X), static_cast<int>(CursorPos.Y), Force / 255.0f);
+				UE_LOG(LogPixelStreamingInputDevice, VeryVerbose, TEXT("TOUCH_MOVE: TouchIndex = %d; Pos = (%d, %d); CursorPos = (%d, %d); Force = %.3f"), TouchIndex, PosX, PosY, static_cast<int>(CursorPos.X), static_cast<int>(CursorPos.Y), Force / 255.0f);
 			}
 			break;
 			case EventType::GAMEPAD_PRESS:
@@ -358,7 +364,7 @@ void FInputDevice::Tick(float DeltaTime)
 					break;
 				}
 				MessageHandler->OnControllerButtonPressed(ControllerButton, (int32_t)ControllerId, bIsRepeat);
-				UE_LOG(PixelStreamerInputDevice, VeryVerbose, TEXT("GAMEPAD_PRESS: ControllerId = %d; ButtonIndex = %d; IsRepeat = %d"), ControllerId, ButtonIndex, bIsRepeat);
+				UE_LOG(LogPixelStreamingInputDevice, VeryVerbose, TEXT("GAMEPAD_PRESS: ControllerId = %d; ButtonIndex = %d; IsRepeat = %d"), ControllerId, ButtonIndex, bIsRepeat);
 			}
 			break;
 			case EventType::GAMEPAD_RELEASE:
@@ -372,7 +378,7 @@ void FInputDevice::Tick(float DeltaTime)
 					break;
 				}
 				MessageHandler->OnControllerButtonReleased(ControllerButton, (int32_t)ControllerId, false);
-				UE_LOG(PixelStreamerInputDevice, VeryVerbose, TEXT("GAMEPAD_RELEASE: ControllerId = %d; ButtonIndex = %d;"), ControllerId, ButtonIndex);
+				UE_LOG(LogPixelStreamingInputDevice, VeryVerbose, TEXT("GAMEPAD_RELEASE: ControllerId = %d; ButtonIndex = %d;"), ControllerId, ButtonIndex);
 			}
 			break;
 			case EventType::GAMEPAD_ANALOG:
@@ -387,12 +393,12 @@ void FInputDevice::Tick(float DeltaTime)
 					break;
 				}
 				MessageHandler->OnControllerAnalog(ControllerAxis, (int32_t)ControllerId, AnalogValue);
-				UE_LOG(PixelStreamerInputDevice, VeryVerbose, TEXT("GAMEPAD_ANALOG: ControllerId = %d; AxisIndex = %d; AnalogValue = %.4f;"), ControllerId, AxisIndex, AnalogValue);
+				UE_LOG(LogPixelStreamingInputDevice, VeryVerbose, TEXT("GAMEPAD_ANALOG: ControllerId = %d; AxisIndex = %d; AnalogValue = %.4f;"), ControllerId, AxisIndex, AnalogValue);
 			}
 			break;
 			default:
 			{
-				UE_LOG(PixelStreamerInputDevice, Error, TEXT("Unknown Pixel Streaming event %d with word 0x%016llx"), static_cast<int>(Event.Event), Event.Data.Word);
+				UE_LOG(LogPixelStreamingInputDevice, Error, TEXT("Unknown Pixel Streaming event %d with word 0x%016llx"), static_cast<int>(Event.Event), Event.Data.Word);
 			}
 			break;
 		}
@@ -401,21 +407,22 @@ void FInputDevice::Tick(float DeltaTime)
 	FString UIInteraction;
 	while (UIInteractions.Dequeue(UIInteraction))
 	{
-		for (UPixelStreamerInputComponent* InputComponent : this->PixelStreamingModule->GetInputComponents())
+		for (UPixelStreamingInput* InputComponent : PixelStreamingModule->GetInputComponents())
 		{
 			InputComponent->OnInputEvent.Broadcast(UIInteraction);
-			UE_LOG(PixelStreamerInputDevice, Verbose, TEXT("UIInteraction = %s"), *UIInteraction);
+			UE_LOG(LogPixelStreamingInputDevice, Verbose, TEXT("UIInteraction = %s"), *UIInteraction);
 		}
 	}
 
+	
 	FString Command;
 	while (Commands.Dequeue(Command))
 	{
-		UPixelStreamerInputComponent::OnCommand(Command);
+		UPixelStreamingInput::OnCommand(Command);
 	}
 }
 
-FGamepadKeyNames::Type FInputDevice::ConvertAxisIndexToGamepadAxis(uint8 AnalogAxis)
+FGamepadKeyNames::Type UE::PixelStreaming::FInputDevice::ConvertAxisIndexToGamepadAxis(uint8 AnalogAxis)
 {
 	switch (AnalogAxis)
 	{
@@ -457,7 +464,7 @@ FGamepadKeyNames::Type FInputDevice::ConvertAxisIndexToGamepadAxis(uint8 AnalogA
 	}
 }
 
-FGamepadKeyNames::Type FInputDevice::ConvertButtonIndexToGamepadButton(uint8 ButtonIndex)
+FGamepadKeyNames::Type UE::PixelStreaming::FInputDevice::ConvertButtonIndexToGamepadButton(uint8 ButtonIndex)
 {
 	switch (ButtonIndex)
 	{
@@ -537,52 +544,52 @@ FGamepadKeyNames::Type FInputDevice::ConvertButtonIndexToGamepadButton(uint8 But
 	}
 }
 
-void FInputDevice::SendControllerEvents()
+void UE::PixelStreaming::FInputDevice::SendControllerEvents()
 {
 }
 
-void FInputDevice::SetMessageHandler(const TSharedRef<FGenericApplicationMessageHandler>& InMessageHandler)
+void UE::PixelStreaming::FInputDevice::SetMessageHandler(const TSharedRef<FGenericApplicationMessageHandler>& InMessageHandler)
 {
 	MessageHandler = InMessageHandler;
 }
 
-bool FInputDevice::Exec(UWorld* InWorld, const TCHAR* Cmd, FOutputDevice& Ar)
+bool UE::PixelStreaming::FInputDevice::Exec(UWorld* InWorld, const TCHAR* Cmd, FOutputDevice& Ar)
 {
 	return true;
 }
 
-void FInputDevice::SetChannelValue(int32 ControllerId, FForceFeedbackChannelType ChannelType, float Value)
+void UE::PixelStreaming::FInputDevice::SetChannelValue(int32 ControllerId, FForceFeedbackChannelType ChannelType, float Value)
 {
 }
 
-void FInputDevice::SetChannelValues(int32 ControllerId, const FForceFeedbackValues& values)
+void UE::PixelStreaming::FInputDevice::SetChannelValues(int32 ControllerId, const FForceFeedbackValues& values)
 {
 }
 
-void FInputDevice::ProcessEvent(const FEvent& InEvent)
+void UE::PixelStreaming::FInputDevice::ProcessEvent(const FEvent& InEvent)
 {
-	bool Success = Events.Enqueue(InEvent);
-	checkf(Success, TEXT("Unable to enqueue new event of type %d"), static_cast<int>(InEvent.Event));
+	bool bSuccess = Events.Enqueue(InEvent);
+	checkf(bSuccess, TEXT("Unable to enqueue new event of type %d"), static_cast<int>(InEvent.Event));
 }
 
-void FInputDevice::ProcessUIInteraction(const FString& InDescriptor)
+void UE::PixelStreaming::FInputDevice::ProcessUIInteraction(const FString& InDescriptor)
 {
-	bool Success = UIInteractions.Enqueue(InDescriptor);
-	checkf(Success, TEXT("Unable to enqueue new UI Interaction %s"), *InDescriptor);
+	bool bSuccess = UIInteractions.Enqueue(InDescriptor);
+	checkf(bSuccess, TEXT("Unable to enqueue new UI Interaction %s"), *InDescriptor);
 }
 
-void FInputDevice::ProcessCommand(const FString& InDescriptor)
+void UE::PixelStreaming::FInputDevice::ProcessCommand(const FString& InDescriptor)
 {
-	if (PixelStreamingSettings::CVarPixelStreamingAllowConsoleCommands.GetValueOnAnyThread())
+	if (UE::PixelStreaming::Settings::CVarPixelStreamingAllowConsoleCommands.GetValueOnAnyThread())
 	{
-		bool Success = Commands.Enqueue(InDescriptor);
-		checkf(Success, TEXT("Unable to enqueue new Command %s"), *InDescriptor);
+		bool bSuccess = Commands.Enqueue(InDescriptor);
+		checkf(bSuccess, TEXT("Unable to enqueue new Command %s"), *InDescriptor);
 	}
 }
 
-void FInputDevice::FindFocusedWidget()
+void UE::PixelStreaming::FInputDevice::FindFocusedWidget()
 {
-	FSlateApplication::Get().ForEachUser([&](FSlateUser& User) {
+	FSlateApplication::Get().ForEachUser([this](FSlateUser& User) {
 		TSharedPtr<SWidget> FocusedWidget = User.GetFocusedWidget();
 
 		static FName SEditableTextType(TEXT("SEditableText"));
@@ -624,7 +631,7 @@ void FInputDevice::FindFocusedWidget()
 namespace
 {
 
-#define GET(Type, Var) Type Var = PixelStreamingProtocol::ParseBuffer<Type>(Data, Size)
+#define GET(Type, Var) Type Var = UE::PixelStreaming::Protocol::ParseBuffer<Type>(Data, Size)
 
 	// XY positions are the ratio (0.0..1.0) along a viewport axis, quantized
 	// into an uint16 (0..65536). This allows the browser viewport and player
@@ -733,9 +740,9 @@ namespace
 	};
 } // namespace
 
-void FInputDevice::OnMessage(const webrtc::DataBuffer& Buffer)
+void UE::PixelStreaming::FInputDevice::OnMessage(const webrtc::DataBuffer& Buffer)
 {
-	using namespace PixelStreamingProtocol;
+	using namespace UE::PixelStreaming::Protocol;
 
 	const uint8* Data = Buffer.data.data();
 	uint32 Size = (uint32)Buffer.data.size();
@@ -746,15 +753,15 @@ void FInputDevice::OnMessage(const webrtc::DataBuffer& Buffer)
 	{
 		case EToStreamerMsg::UIInteraction:
 		{
-			FString Descriptor = PixelStreamingProtocol::ParseString(Buffer, FInputDevice::GetMessageHeaderOffset());
-			UE_LOG(PixelStreamerInput, Verbose, TEXT("UIInteraction: %s"), *Descriptor);
+			FString Descriptor = UE::PixelStreaming::Protocol::ParseString(Buffer, UE::PixelStreaming::FInputDevice::GetMessageHeaderOffset());
+			UE_LOG(LogPixelStreamingInput, Verbose, TEXT("UIInteraction: %s"), *Descriptor);
 			ProcessUIInteraction(Descriptor);
 			break;
 		}
 		case EToStreamerMsg::Command:
 		{
-			FString Descriptor = PixelStreamingProtocol::ParseString(Buffer, FInputDevice::GetMessageHeaderOffset());
-			UE_LOG(PixelStreamerInput, Verbose, TEXT("Command: %s"), *Descriptor);
+			FString Descriptor = UE::PixelStreaming::Protocol::ParseString(Buffer, UE::PixelStreaming::FInputDevice::GetMessageHeaderOffset());
+			UE_LOG(LogPixelStreamingInput, Verbose, TEXT("Command: %s"), *Descriptor);
 			ProcessCommand(Descriptor);
 			break;
 		}
@@ -763,7 +770,7 @@ void FInputDevice::OnMessage(const webrtc::DataBuffer& Buffer)
 			GET(FKeyCodeType, KeyCode);
 			GET(FRepeatType, Repeat);
 			checkf(Size == 0, TEXT("%d"), Size);
-			UE_LOG(PixelStreamerInput, Verbose, TEXT("key down: %d, repeat: %d"), KeyCode, Repeat);
+			UE_LOG(LogPixelStreamingInput, Verbose, TEXT("key down: %d, repeat: %d"), KeyCode, Repeat);
 
 			FEvent KeyDownEvent(EventType::KEY_DOWN);
 			KeyDownEvent.SetKeyDown(KeyCode, Repeat != 0);
@@ -774,7 +781,7 @@ void FInputDevice::OnMessage(const webrtc::DataBuffer& Buffer)
 		{
 			GET(FKeyCodeType, KeyCode);
 			checkf(Size == 0, TEXT("%d"), Size);
-			UE_LOG(PixelStreamerInput, Verbose, TEXT("key up: %d"), KeyCode);
+			UE_LOG(LogPixelStreamingInput, Verbose, TEXT("key up: %d"), KeyCode);
 
 			FEvent KeyUpEvent(EventType::KEY_UP);
 			KeyUpEvent.SetKeyUp(KeyCode);
@@ -785,7 +792,7 @@ void FInputDevice::OnMessage(const webrtc::DataBuffer& Buffer)
 		{
 			GET(FCharacterType, Character);
 			checkf(Size == 0, TEXT("%d"), Size);
-			UE_LOG(PixelStreamerInput, Verbose, TEXT("key press: '%c'"), Character);
+			UE_LOG(LogPixelStreamingInput, Verbose, TEXT("key press: '%c'"), Character);
 
 			FEvent KeyPressEvent(EventType::KEY_PRESS);
 			KeyPressEvent.SetCharCode(Character);
@@ -795,15 +802,15 @@ void FInputDevice::OnMessage(const webrtc::DataBuffer& Buffer)
 		case EToStreamerMsg::MouseEnter:
 		{
 			checkf(Size == 0, TEXT("%d"), Size);
-			UE_LOG(PixelStreamerInput, Verbose, TEXT("mouseEnter"));
-			ProcessEvent(FEvent(FInputDevice::EventType::MOUSE_ENTER));
+			UE_LOG(LogPixelStreamingInput, Verbose, TEXT("mouseEnter"));
+			ProcessEvent(FEvent(UE::PixelStreaming::FInputDevice::EventType::MOUSE_ENTER));
 			break;
 		}
 		case EToStreamerMsg::MouseLeave:
 		{
 			checkf(Size == 0, TEXT("%d"), Size);
-			ProcessEvent(FEvent(FInputDevice::EventType::MOUSE_LEAVE));
-			UE_LOG(PixelStreamerInput, Verbose, TEXT("mouseLeave"));
+			ProcessEvent(FEvent(UE::PixelStreaming::FInputDevice::EventType::MOUSE_LEAVE));
+			UE_LOG(LogPixelStreamingInput, Verbose, TEXT("mouseLeave"));
 			break;
 		}
 		case EToStreamerMsg::MouseDown:
@@ -812,7 +819,7 @@ void FInputDevice::OnMessage(const webrtc::DataBuffer& Buffer)
 			GET(FPosType, PosX);
 			GET(FPosType, PosY);
 			checkf(Size == 0, TEXT("%d"), Size);
-			UE_LOG(PixelStreamerInput, Verbose, TEXT("mouseDown at (%d, %d), button %d"), PosX, PosY, Button);
+			UE_LOG(LogPixelStreamingInput, Verbose, TEXT("mouseDown at (%d, %d), button %d"), PosX, PosY, Button);
 
 			UnquantizeAndDenormalize(PosX, PosY);
 
@@ -827,7 +834,7 @@ void FInputDevice::OnMessage(const webrtc::DataBuffer& Buffer)
 			GET(FPosType, PosX);
 			GET(FPosType, PosY);
 			checkf(Size == 0, TEXT("%d"), Size);
-			UE_LOG(PixelStreamerInput, Verbose, TEXT("mouseUp at (%d, %d), button %d"), PosX, PosY, Button);
+			UE_LOG(LogPixelStreamingInput, Verbose, TEXT("mouseUp at (%d, %d), button %d"), PosX, PosY, Button);
 
 			UnquantizeAndDenormalize(PosX, PosY);
 
@@ -843,7 +850,7 @@ void FInputDevice::OnMessage(const webrtc::DataBuffer& Buffer)
 			GET(FDeltaType, DeltaX);
 			GET(FDeltaType, DeltaY);
 			checkf(Size == 0, TEXT("%d"), Size);
-			UE_LOG(PixelStreamerInput, Verbose, TEXT("mouseMove to (%d, %d), delta (%d, %d)"), PosX, PosY, DeltaX, DeltaY);
+			UE_LOG(LogPixelStreamingInput, Verbose, TEXT("mouseMove to (%d, %d), delta (%d, %d)"), PosX, PosY, DeltaX, DeltaY);
 
 			UnquantizeAndDenormalize(PosX, PosY);
 			UnquantizeAndDenormalize(DeltaX, DeltaY);
@@ -859,7 +866,7 @@ void FInputDevice::OnMessage(const webrtc::DataBuffer& Buffer)
 			GET(FPosType, PosX);
 			GET(FPosType, PosY);
 			checkf(Size == 0, TEXT("%d"), Size);
-			UE_LOG(PixelStreamerInput, Verbose, TEXT("mouseWheel, delta %d"), Delta);
+			UE_LOG(LogPixelStreamingInput, Verbose, TEXT("mouseWheel, delta %d"), Delta);
 
 			UnquantizeAndDenormalize(PosX, PosY);
 
@@ -872,7 +879,7 @@ void FInputDevice::OnMessage(const webrtc::DataBuffer& Buffer)
 		{
 			FTouchesType Touches = GetTouches(Data, Size);
 			checkf(Size == 0, TEXT("%d"), Size);
-			UE_LOG(PixelStreamerInput, Verbose, TEXT("TouchStart: %s"), *TouchesToString(Touches));
+			UE_LOG(LogPixelStreamingInput, Verbose, TEXT("TouchStart: %s"), *TouchesToString(Touches));
 
 			for (const FTouch& Touch : Touches)
 			{
@@ -889,7 +896,7 @@ void FInputDevice::OnMessage(const webrtc::DataBuffer& Buffer)
 		{
 			FTouchesType Touches = GetTouches(Data, Size);
 			checkf(Size == 0, TEXT("%d"), Size);
-			UE_LOG(PixelStreamerInput, Verbose, TEXT("TouchEnd: %s"), *TouchesToString(Touches));
+			UE_LOG(LogPixelStreamingInput, Verbose, TEXT("TouchEnd: %s"), *TouchesToString(Touches));
 
 			for (const FTouch& Touch : Touches)
 			{
@@ -908,7 +915,7 @@ void FInputDevice::OnMessage(const webrtc::DataBuffer& Buffer)
 		{
 			FTouchesType Touches = GetTouches(Data, Size);
 			checkf(Size == 0, TEXT("%d"), Size);
-			UE_LOG(PixelStreamerInput, Verbose, TEXT("TouchMove: %s"), *TouchesToString(Touches));
+			UE_LOG(LogPixelStreamingInput, Verbose, TEXT("TouchMove: %s"), *TouchesToString(Touches));
 
 			for (const FTouch& Touch : Touches)
 			{
@@ -961,7 +968,7 @@ void FInputDevice::OnMessage(const webrtc::DataBuffer& Buffer)
 
 #undef GET
 
-void FInputDevice::FEvent::GetMouseClick(EMouseButtons::Type& OutButton, uint16& OutPosX, uint16& OutPosY)
+void UE::PixelStreaming::FInputDevice::FEvent::GetMouseClick(EMouseButtons::Type& OutButton, uint16& OutPosX, uint16& OutPosY)
 {
 	check(Event == EventType::MOUSE_DOWN || Event == EventType::MOUSE_UP);
 	// https://developer.mozilla.org/en-US/docs/Web/Events/mousedown
@@ -995,7 +1002,7 @@ void FInputDevice::FEvent::GetMouseClick(EMouseButtons::Type& OutButton, uint16&
 		break;
 		default:
 		{
-			UE_LOG(PixelStreamerInputDevice, Error, TEXT("Unknown Pixel Streaming mouse click with button %d and word 0x%016llx"), Button, Data.Word);
+			UE_LOG(LogPixelStreamingInputDevice, Error, TEXT("Unknown Pixel Streaming mouse click with button %d and word 0x%016llx"), Button, Data.Word);
 		}
 		break;
 	}
diff --git a/Engine/Plugins/Media/PixelStreaming/Source/PixelStreaming/Private/InputDevice.h b/Engine/Plugins/Media/PixelStreaming/Source/PixelStreaming/Private/InputDevice.h
--- a/Engine/Plugins/Media/PixelStreaming/Source/PixelStreaming/Private/InputDevice.h
+++ b/Engine/Plugins/Media/PixelStreaming/Source/PixelStreaming/Private/InputDevice.h
@@ -9,517 +9,522 @@
 
 class IPixelStreamingModule;
 
-class FInputDevice : public IInputDevice
-{
-public:
-	/**
-	 * The pixel streaming input device takes input events and passes them to
-	 * a message handler.
-	 * @param InMessageHandler - The message handler which events should be passed to.
-	 * @param InInputComponents - The preexisting list of pixel streaming input components.
-	 */
-	FInputDevice(const TSharedRef<FGenericApplicationMessageHandler>& InMessageHandler);
-
-	/** Tick the interface (e.g. check for new controllers) */
-	virtual void Tick(float DeltaTime) override;
-
-	/** Poll for controller state and send events if needed */
-	virtual void SendControllerEvents() override;
-
-	/** Set which MessageHandler will get the events from SendControllerEvents. */
-	virtual void SetMessageHandler(const TSharedRef<FGenericApplicationMessageHandler>& InMessageHandler) override;
-
-	/** Exec handler to allow console commands to be passed through for debugging */
-	virtual bool Exec(UWorld* InWorld, const TCHAR* Cmd, FOutputDevice& Ar) override;
-
-	/**
-	* IForceFeedbackSystem pass through functions
-	*/
-	virtual void SetChannelValue(int32 ControllerId, FForceFeedbackChannelType ChannelType, float Value) override;
-
-	virtual void SetChannelValues(int32 ControllerId, const FForceFeedbackValues& values) override;
-
-	/** The types of event which can be processed by the device. */
-	enum class EventType
-	{
-		UNDEFINED,		 /** No value. */
-		KEY_DOWN,		 /** A key has been pushed down. */
-		KEY_UP,			 /** A key has been released. */
-		KEY_PRESS,		 /** A key has been pressed and a character has been input. */
-		MOUSE_ENTER,	 /** The mouse has entered canvas. */
-		MOUSE_LEAVE,	 /** The mouse has left the canvas. */
-		MOUSE_MOVE,		 /** The mouse has been moved. */
-		MOUSE_DOWN,		 /** A mouse button has been clicked. */
-		MOUSE_UP,		 /** A mouse button has been released. */
-		MOUSE_WHEEL,	 /** The mouse wheel was scrolled. */
-		TOUCH_START,	 /** A finger is put down onto the canvas. */
-		TOUCH_END,		 /** A finger is lifted from the canvas. */
-		TOUCH_MOVE,		 /** A finger is being dragged along the surface of the canvas. */
-		GAMEPAD_PRESS,	 /** A gamepad button has been pressed. */
-		GAMEPAD_RELEASE, /** A gamepad button has been released. */
-		GAMEPAD_ANALOG	 /** A gamepad analog stick has been moved. */
-	};
-
-	/** A general input event. */
-	struct FEvent
-	{
-		/** The type of the general event. */
-		EventType Event;
-
-		/** A generic piece of data which is used to hold information about the
-		 * event, specialized by making a union with an appropriate struct. */
-		union
+namespace UE {
+	namespace PixelStreaming {
+
+		class FInputDevice : public IInputDevice
 		{
-			uint64 Word;
+		public:
+			/**
+			 * The pixel streaming input device takes input events and passes them to
+			 * a message handler.
+			 * @param InMessageHandler - The message handler which events should be passed to.
+			 * @param InInputComponents - The preexisting list of pixel streaming input components.
+			 */
+			FInputDevice(const TSharedRef<FGenericApplicationMessageHandler>& InMessageHandler);
 
-			struct /** KEY_DOWN */
-			{
-				uint8 KeyCode;
-				bool bIsRepeat;
-			} KeyDown;
+			/** Tick the interface (e.g. check for new controllers) */
+			virtual void Tick(float DeltaTime) override;
 
-			struct /* KEY_UP */
-			{
-				uint8 KeyCode;
-			} KeyUp;
+			/** Poll for controller state and send events if needed */
+			virtual void SendControllerEvents() override;
 
-			struct /** KEY_PRESSED */
-			{
-				TCHAR Character;
-			} Character;
+			/** Set which MessageHandler will get the events from SendControllerEvents. */
+			virtual void SetMessageHandler(const TSharedRef<FGenericApplicationMessageHandler>& InMessageHandler) override;
 
-			struct /** MOUSE_MOVE */
-			{
-				int16 DeltaX;
-				int16 DeltaY;
-				uint16 PosX;
-				uint16 PosY;
-			} MouseMove;
+			/** Exec handler to allow console commands to be passed through for debugging */
+			virtual bool Exec(UWorld* InWorld, const TCHAR* Cmd, FOutputDevice& Ar) override;
 
-			struct /** MOUSE_DOWN, MOUSE_UP */
-			{
-				uint8 Button;
-				uint16 PosX;
-				uint16 PosY;
-			} MouseButton;
+			/**
+			* IForceFeedbackSystem pass through functions
+			*/
+			virtual void SetChannelValue(int32 ControllerId, FForceFeedbackChannelType ChannelType, float Value) override;
 
-			struct /** MOUSE_WHEEL */
-			{
-				int16 Delta;
-				uint16 PosX;
-				uint16 PosY;
-			} MouseWheel;
+			virtual void SetChannelValues(int32 ControllerId, const FForceFeedbackValues& values) override;
 
-			struct /** TOUCH_START, TOUCH_END, TOUCH_MOVE */
+			/** The types of event which can be processed by the device. */
+			enum class EventType
 			{
-				uint8 TouchIndex;
-				uint16 PosX;
-				uint16 PosY;
-				uint8 Force;
-			} Touch;
-
-			struct /** GAMEPAD_PRESSED, GAMEPAD_RELEASED */
+				UNDEFINED,		 /** No value. */
+				KEY_DOWN,		 /** A key has been pushed down. */
+				KEY_UP,			 /** A key has been released. */
+				KEY_PRESS,		 /** A key has been pressed and a character has been input. */
+				MOUSE_ENTER,	 /** The mouse has entered canvas. */
+				MOUSE_LEAVE,	 /** The mouse has left the canvas. */
+				MOUSE_MOVE,		 /** The mouse has been moved. */
+				MOUSE_DOWN,		 /** A mouse button has been clicked. */
+				MOUSE_UP,		 /** A mouse button has been released. */
+				MOUSE_WHEEL,	 /** The mouse wheel was scrolled. */
+				TOUCH_START,	 /** A finger is put down onto the canvas. */
+				TOUCH_END,		 /** A finger is lifted from the canvas. */
+				TOUCH_MOVE,		 /** A finger is being dragged along the surface of the canvas. */
+				GAMEPAD_PRESS,	 /** A gamepad button has been pressed. */
+				GAMEPAD_RELEASE, /** A gamepad button has been released. */
+				GAMEPAD_ANALOG	 /** A gamepad analog stick has been moved. */
+			};
+
+			/** A general input event. */
+			struct FEvent
 			{
-				uint8 ControllerIndex;
-				uint8 ButtonIndex;
-				bool bIsRepeat;
-
-			} GamepadButton;
-
-			struct /** GAMEPAD_PRESSED, GAMEPAD_RELEASED */
+				/** The type of the general event. */
+				EventType Event;
+
+				/** A generic piece of data which is used to hold information about the
+				* event, specialized by making a union with an appropriate struct. */
+				union
+				{
+					uint64 Word;
+
+					struct /** KEY_DOWN */
+					{
+						uint8 KeyCode;
+						bool bIsRepeat;
+					} KeyDown;
+
+					struct /* KEY_UP */
+					{
+						uint8 KeyCode;
+					} KeyUp;
+
+					struct /** KEY_PRESSED */
+					{
+						TCHAR Character;
+					} Character;
+
+					struct /** MOUSE_MOVE */
+					{
+						int16 DeltaX;
+						int16 DeltaY;
+						uint16 PosX;
+						uint16 PosY;
+					} MouseMove;
+
+					struct /** MOUSE_DOWN, MOUSE_UP */
+					{
+						uint8 Button;
+						uint16 PosX;
+						uint16 PosY;
+					} MouseButton;
+
+					struct /** MOUSE_WHEEL */
+					{
+						int16 Delta;
+						uint16 PosX;
+						uint16 PosY;
+					} MouseWheel;
+
+					struct /** TOUCH_START, TOUCH_END, TOUCH_MOVE */
+					{
+						uint8 TouchIndex;
+						uint16 PosX;
+						uint16 PosY;
+						uint8 Force;
+					} Touch;
+
+					struct /** GAMEPAD_PRESSED, GAMEPAD_RELEASED */
+					{
+						uint8 ControllerIndex;
+						uint8 ButtonIndex;
+						bool bIsRepeat;
+
+					} GamepadButton;
+
+					struct /** GAMEPAD_PRESSED, GAMEPAD_RELEASED */
+					{
+						uint8 ControllerIndex;
+						float AnalogValue;
+						uint8 AxisIndex;
+					} GamepadAnalog;
+				} Data;
+
+				/**
+				* Create a completely empty event.
+				*/
+				FEvent()
+					: Event(EventType::UNDEFINED)
+				{
+				}
+
+				/**
+				* Create an event of the given type.
+				* @param InEvent - The type of the event.
+				*/
+				FEvent(EventType InEvent)
+					: Event(InEvent)
+				{
+				}
+
+				/**
+				* An event related to gamepad analog stick being moved.
+				* @param InControllerIndex - Number representing the index of the controller.
+				* @param InAxisIndex - The axis that is non-centered.
+				* @param InAnalogValue - The value of the analog stick.
+				*/
+				void SetGamepadAnalog(uint8 InControllerIndex, uint8 InAxisIndex, float InAnalogValue)
+				{
+					check(Event == EventType::GAMEPAD_ANALOG);
+					Data.GamepadAnalog.ControllerIndex = InControllerIndex;
+					Data.GamepadAnalog.AxisIndex = InAxisIndex;
+					Data.GamepadAnalog.AnalogValue = InAnalogValue;
+				}
+
+				/**
+				* An event related to a gamepad button being pushed down.
+				* @param InControllerIndex - Number representing the index of the controller.
+				* @param InButtonIndex - Numerical code identifying the pushed down button.
+				* @param InIsRepeat - Whether the button is being kept down and is repeating.
+				*/
+				void SetGamepadButtonPressed(uint8 InControllerIndex, uint8 InButtonIndex, bool InIsRepeat)
+				{
+					check(Event == EventType::GAMEPAD_PRESS);
+					Data.GamepadButton.ControllerIndex = InControllerIndex;
+					Data.GamepadButton.ButtonIndex = InButtonIndex;
+					Data.GamepadButton.bIsRepeat = InIsRepeat;
+				}
+
+				/**
+				* An event related to a gamepad button being released.
+				* @param InControllerIndex - Number representing the index of the controller.
+				* @param InButtonIndex - Numerical code identifying the released button.
+				*/
+				void SetGamepadButtonReleased(uint8 InControllerIndex, uint8 InButtonIndex)
+				{
+					check(Event == EventType::GAMEPAD_RELEASE);
+					Data.GamepadButton.ControllerIndex = InControllerIndex;
+					Data.GamepadButton.ButtonIndex = InButtonIndex;
+				}
+
+				/**
+				* An event related to a key being pushed down.
+				* @param InKeyCode - Numerical code identifying the pushed down key.
+				* @param InIsRepeat - Whether the key is being kept down and is repeating.
+				*/
+				void SetKeyDown(uint8 InKeyCode, bool InIsRepeat)
+				{
+					check(Event == EventType::KEY_DOWN);
+					Data.KeyDown.KeyCode = InKeyCode;
+					Data.KeyDown.bIsRepeat = InIsRepeat;
+				}
+
+				/**
+				* An event related to a key being released.
+				* @param InKeyCode - Numerical code identifying the released key.
+				*/
+				void SetKeyUp(uint8 InKeyCode)
+				{
+					check(Event == EventType::KEY_UP);
+					Data.KeyUp.KeyCode = InKeyCode;
+				}
+
+				/**
+				* An event related to character input.
+				* @param InCharacter - The character being input.
+				*/
+				void SetCharCode(TCHAR InCharacter)
+				{
+					check(Event == EventType::KEY_PRESS);
+					Data.Character.Character = InCharacter;
+				}
+
+				/**
+				* An event related to mouse movement.
+				* @param InPoxX - The X position of the mouse pointer.
+				* @param InPosY - The Y position of the mouse pointer.
+				* @param InDeltaX - The change in the X position of the mouse pointer.
+				* @param InDeltaY - The change in the Y position of the mouse pointer.
+				*/
+				void SetMouseDelta(uint16 InPosX, uint16 InPosY, int16 InDeltaX, int16 InDeltaY)
+				{
+					check(Event == EventType::MOUSE_MOVE);
+					Data.MouseMove.DeltaX = InDeltaX;
+					Data.MouseMove.DeltaY = InDeltaY;
+					Data.MouseMove.PosX = InPosX;
+					Data.MouseMove.PosY = InPosY;
+				}
+
+				/**
+				* An event related to mouse buttons.
+				* @param InButton - The button number corresponding to left, middle, right, etc.
+				* @param InPoxX - The X position of the mouse pointer.
+				* @param InPosY - The Y position of the mouse pointer.
+				*/
+				void SetMouseClick(uint8 InButton, uint16 InPosX, uint16 InPosY)
+				{
+					check(Event == EventType::MOUSE_DOWN || Event == EventType::MOUSE_UP);
+					Data.MouseButton.Button = InButton;
+					Data.MouseButton.PosX = InPosX;
+					Data.MouseButton.PosY = InPosY;
+				}
+
+				/**
+				* An event related to the mouse scroll wheel.
+				* @param InButton - The amount by which the mouse wheel was scrolled.
+				* @param InPoxX - The X position of the mouse pointer when the wheel was scrolled.
+				* @param InPosY - The Y position of the mouse pointer when the wheel was scrolled.
+				*/
+				void SetMouseWheel(int16 InDelta, uint16 InPosX, uint16 InPosY)
+				{
+					check(Event == EventType::MOUSE_WHEEL);
+					Data.MouseWheel.Delta = InDelta;
+					Data.MouseWheel.PosX = InPosX;
+					Data.MouseWheel.PosY = InPosY;
+				}
+
+				/**
+				* An event related to a finger touching the canvas.
+				* @param InTouchIndex - The finger used in multi-touch.
+				* @param InPoxX - The X position of the finger.
+				* @param InPosY - The Y position of the finger.
+				* @param InForce - The amount of pressure being applied by the finger.
+				*/
+				void SetTouch(uint8 InTouchIndex, uint16 InPosX, uint16 InPosY, uint8 InForce)
+				{
+					check(Event == EventType::TOUCH_START || Event == EventType::TOUCH_END || Event == EventType::TOUCH_MOVE);
+					Data.Touch.TouchIndex = InTouchIndex;
+					Data.Touch.PosX = InPosX;
+					Data.Touch.PosY = InPosY;
+					Data.Touch.Force = InForce;
+				}
+
+				/**
+				* Get information about an event related to gamepad analog stick being moved.
+				* @param OutControllerIndex - Number representing the index of the controller.
+				* @param OutAxisIndex - Number representing the axis of the controller.
+				* @param OutAnalogValue - The value of the analog stick.
+				*/
+				void GetGamepadAnalog(uint8& OutControllerIndex, uint8& OutAxisIndex, float& OutAnalogValue)
+				{
+					check(Event == EventType::GAMEPAD_ANALOG);
+					OutControllerIndex = Data.GamepadAnalog.ControllerIndex;
+					OutAxisIndex = Data.GamepadAnalog.AxisIndex;
+					OutAnalogValue = Data.GamepadAnalog.AnalogValue;
+				}
+
+				/**
+				* Get information about an event related to a gamepad button being pushed down.
+				* @param OutControllerIndex - Number representing the index of the controller.
+				* @param OutButtonIndex - Numerical code identifying the pushed down button.
+				* @param OutIsRepeat - Whether the button is being kept down and is repeating.
+				*/
+				void GetGamepadButtonPressed(uint8& OutControllerIndex, uint8& OutButtonIndex, bool& OutIsRepeat)
+				{
+					check(Event == EventType::GAMEPAD_PRESS);
+					OutControllerIndex = Data.GamepadButton.ControllerIndex;
+					OutButtonIndex = Data.GamepadButton.ButtonIndex;
+					OutIsRepeat = Data.GamepadButton.bIsRepeat;
+				}
+
+				/**
+				* Get information about an event related to a gamepad button being released.
+				* @param OutControllerIndex - Number representing the index of the controller.
+				* @param OutButtonIndex - Numerical code identifying the pushed down button.
+				*/
+				void GetGamepadButtonReleased(uint8& OutControllerIndex, uint8& OutButtonIndex)
+				{
+					check(Event == EventType::GAMEPAD_RELEASE);
+					OutControllerIndex = Data.GamepadButton.ControllerIndex;
+					OutButtonIndex = Data.GamepadButton.ButtonIndex;
+				}
+
+				/**
+				* Get information about an event related to a key being pushed down.
+				* @param OutKeyCode - Numerical code identifying the pushed down key.
+				* @param OutIsRepeat - Whether the key is being kept down and is repeating.
+				*/
+				void GetKeyDown(uint8& OutKeyCode, bool& OutIsRepeat)
+				{
+					check(Event == EventType::KEY_DOWN);
+					OutKeyCode = Data.KeyDown.KeyCode;
+					OutIsRepeat = Data.KeyDown.bIsRepeat;
+				}
+
+				/**
+				* Get information about an event related to a key being released.
+				* @param OutKeyCode - Numerical code identifying the released key.
+				*/
+				void GetKeyUp(uint8& OutKeyCode)
+				{
+					check(Event == EventType::KEY_UP);
+					OutKeyCode = Data.KeyUp.KeyCode;
+				}
+
+				/**
+				* Get information about an event related to character input.
+				* @param OutCharacter - The character being input.
+				*/
+				void GetCharacterCode(TCHAR& OutCharacter)
+				{
+					check(Event == EventType::KEY_PRESS);
+					OutCharacter = Data.Character.Character;
+				}
+
+				/**
+				* Get information about an event related to mouse movement.
+				* @param OutPoxX - The X position of the mouse pointer.
+				* @param OutPosY - The Y position of the mouse pointer.
+				* @param OutDeltaX - The change in the X position of the mouse pointer.
+				* @param OutDeltaY - The change in the Y position of the mouse pointer.
+				*/
+				void GetMouseDelta(uint16& OutPosX, uint16& OutPosY, int16& OutDeltaX, int16& OutDeltaY)
+				{
+					check(Event == EventType::MOUSE_MOVE);
+					OutPosX = Data.MouseMove.PosX;
+					OutPosY = Data.MouseMove.PosY;
+					OutDeltaX = Data.MouseMove.DeltaX;
+					OutDeltaY = Data.MouseMove.DeltaY;
+				}
+
+				/**
+				* Get information about an event related to mouse buttons.
+				* @param OutButton - The button number corresponding to left, middle, right, etc.
+				* @param OutPosX - The X position of the mouse pointer.
+				* @param OutPosY - The Y position of the mouse pointer.
+				*/
+				void GetMouseClick(EMouseButtons::Type& OutButton, uint16& OutPosX, uint16& OutPosY);
+
+				/**
+				* Get information about an event related to the mouse wheel.
+				* @param OutDelta - The amount by which the mouse wheel was scrolled.
+				* @param PosX - The X position of the mouse pointer when the wheel was scrolled.
+				* @param PosY - The Y position of the mouse pointer when the wheel was scrolled.
+				*/
+				void GetMouseWheel(int16& OutDelta, uint16& OutPosX, uint16& OutPosY)
+				{
+					check(Event == EventType::MOUSE_WHEEL);
+					OutDelta = Data.MouseWheel.Delta;
+					OutPosX = Data.MouseWheel.PosX;
+					OutPosY = Data.MouseWheel.PosY;
+				}
+
+				/**
+				* Get information about an event related to a finger touching the canvas.
+				* @param OutTouchIndex - The finger used in multi-touch.
+				* @param OutPoxX - The X position of the finger.
+				* @param OutPosY - The Y position of the finger.
+				* @param OutForce - Amount of pressure being applied by the finger.
+				*/
+				void GetTouch(uint8& OutTouchIndex, uint16& OutPosX, uint16& OutPosY, uint8& OutForce)
+				{
+					check(Event == EventType::TOUCH_START || Event == EventType::TOUCH_END || Event == EventType::TOUCH_MOVE);
+					OutTouchIndex = Data.Touch.TouchIndex;
+					OutPosX = Data.Touch.PosX;
+					OutPosY = Data.Touch.PosY;
+					OutForce = Data.Touch.Force;
+				}
+			};
+
+			/**
+			* Add a new event to the input device for later processing.
+			* @param InEvent - The new event.
+			*/
+			void ProcessEvent(const FEvent& InEvent);
+
+			/**
+			* Add a new UI interaction descriptor to the input device for later processing.
+			* @param InDescriptor - The new UI interaction descriptor.
+			*/
+			void ProcessUIInteraction(const FString& InDescriptor);
+
+			/**
+			* Add a new command descriptor to the input device for later execution.
+			* @param InDescriptor - The new command descriptor.
+			*/
+			void ProcessCommand(const FString& InDescriptor);
+
+			/**
+			* Is the application faking touch events?
+			* @return True if the application is faking touch events.
+			*/
+			bool IsFakingTouchEvents() const
 			{
-				uint8 ControllerIndex;
-				float AnalogValue;
-				uint8 AxisIndex;
-			} GamepadAnalog;
-		} Data;
-
-		/**
-		 * Create a completely empty event.
-		 */
-		FEvent()
-			: Event(EventType::UNDEFINED)
-		{
-		}
-
-		/**
-		 * Create an event of the given type.
-		 * @param InEvent - The type of the event.
-		 */
-		FEvent(EventType InEvent)
-			: Event(InEvent)
-		{
-		}
-
-		/**
-		 * An event related to gamepad analog stick being moved.
-		 * @param InControllerIndex - Number representing the index of the controller.
-		 * @param InAxisIndex - The axis that is non-centered.
-		 * @param InAnalogValue - The value of the analog stick.
-		*/
-		void SetGamepadAnalog(uint8 InControllerIndex, uint8 InAxisIndex, float InAnalogValue)
-		{
-			check(Event == EventType::GAMEPAD_ANALOG);
-			Data.GamepadAnalog.ControllerIndex = InControllerIndex;
-			Data.GamepadAnalog.AxisIndex = InAxisIndex;
-			Data.GamepadAnalog.AnalogValue = InAnalogValue;
-		}
-
-		/**
-		 * An event related to a gamepad button being pushed down.
-		 * @param InControllerIndex - Number representing the index of the controller.
-		 * @param InButtonIndex - Numerical code identifying the pushed down button.
-		 * @param InIsRepeat - Whether the button is being kept down and is repeating.
-		*/
-		void SetGamepadButtonPressed(uint8 InControllerIndex, uint8 InButtonIndex, bool InIsRepeat)
-		{
-			check(Event == EventType::GAMEPAD_PRESS);
-			Data.GamepadButton.ControllerIndex = InControllerIndex;
-			Data.GamepadButton.ButtonIndex = InButtonIndex;
-			Data.GamepadButton.bIsRepeat = InIsRepeat;
-		}
-
-		/**
-		 * An event related to a gamepad button being released.
-		 * @param InControllerIndex - Number representing the index of the controller.
-		 * @param InButtonIndex - Numerical code identifying the released button.
-		*/
-		void SetGamepadButtonReleased(uint8 InControllerIndex, uint8 InButtonIndex)
-		{
-			check(Event == EventType::GAMEPAD_RELEASE);
-			Data.GamepadButton.ControllerIndex = InControllerIndex;
-			Data.GamepadButton.ButtonIndex = InButtonIndex;
-		}
-
-		/**
-		 * An event related to a key being pushed down.
-		 * @param InKeyCode - Numerical code identifying the pushed down key.
-		 * @param InIsRepeat - Whether the key is being kept down and is repeating.
-		 */
-		void SetKeyDown(uint8 InKeyCode, bool InIsRepeat)
-		{
-			check(Event == EventType::KEY_DOWN);
-			Data.KeyDown.KeyCode = InKeyCode;
-			Data.KeyDown.bIsRepeat = InIsRepeat;
-		}
-
-		/**
-		 * An event related to a key being released.
-		 * @param InKeyCode - Numerical code identifying the released key.
-		 */
-		void SetKeyUp(uint8 InKeyCode)
-		{
-			check(Event == EventType::KEY_UP);
-			Data.KeyUp.KeyCode = InKeyCode;
-		}
-
-		/**
-		 * An event related to character input.
-		 * @param InCharacter - The character being input.
-		 */
-		void SetCharCode(TCHAR InCharacter)
-		{
-			check(Event == EventType::KEY_PRESS);
-			Data.Character.Character = InCharacter;
-		}
-
-		/**
-		 * An event related to mouse movement.
-		 * @param InPoxX - The X position of the mouse pointer.
-		 * @param InPosY - The Y position of the mouse pointer.
-		 * @param InDeltaX - The change in the X position of the mouse pointer.
-		 * @param InDeltaY - The change in the Y position of the mouse pointer.
-		 */
-		void SetMouseDelta(uint16 InPosX, uint16 InPosY, int16 InDeltaX, int16 InDeltaY)
-		{
-			check(Event == EventType::MOUSE_MOVE);
-			Data.MouseMove.DeltaX = InDeltaX;
-			Data.MouseMove.DeltaY = InDeltaY;
-			Data.MouseMove.PosX = InPosX;
-			Data.MouseMove.PosY = InPosY;
-		}
-
-		/**
-		 * An event related to mouse buttons.
-		 * @param InButton - The button number corresponding to left, middle, right, etc.
-		 * @param InPoxX - The X position of the mouse pointer.
-		 * @param InPosY - The Y position of the mouse pointer.
-		 */
-		void SetMouseClick(uint8 InButton, uint16 InPosX, uint16 InPosY)
-		{
-			check(Event == EventType::MOUSE_DOWN || Event == EventType::MOUSE_UP);
-			Data.MouseButton.Button = InButton;
-			Data.MouseButton.PosX = InPosX;
-			Data.MouseButton.PosY = InPosY;
-		}
-
-		/**
-		 * An event related to the mouse scroll wheel.
-		 * @param InButton - The amount by which the mouse wheel was scrolled.
-		 * @param InPoxX - The X position of the mouse pointer when the wheel was scrolled.
-		 * @param InPosY - The Y position of the mouse pointer when the wheel was scrolled.
-		 */
-		void SetMouseWheel(int16 InDelta, uint16 InPosX, uint16 InPosY)
-		{
-			check(Event == EventType::MOUSE_WHEEL);
-			Data.MouseWheel.Delta = InDelta;
-			Data.MouseWheel.PosX = InPosX;
-			Data.MouseWheel.PosY = InPosY;
-		}
-
-		/**
-		 * An event related to a finger touching the canvas.
-		 * @param InTouchIndex - The finger used in multi-touch.
-		 * @param InPoxX - The X position of the finger.
-		 * @param InPosY - The Y position of the finger.
-		 * @param InForce - The amount of pressure being applied by the finger.
-		 */
-		void SetTouch(uint8 InTouchIndex, uint16 InPosX, uint16 InPosY, uint8 InForce)
-		{
-			check(Event == EventType::TOUCH_START || Event == EventType::TOUCH_END || Event == EventType::TOUCH_MOVE);
-			Data.Touch.TouchIndex = InTouchIndex;
-			Data.Touch.PosX = InPosX;
-			Data.Touch.PosY = InPosY;
-			Data.Touch.Force = InForce;
-		}
-
-		/**
-		 * Get information about an event related to gamepad analog stick being moved.
-		 * @param OutControllerIndex - Number representing the index of the controller.
-		 * @param OutAxisIndex - Number representing the axis of the controller.
-		 * @param OutAnalogValue - The value of the analog stick.
-		 */
-		void GetGamepadAnalog(uint8& OutControllerIndex, uint8& OutAxisIndex, float& OutAnalogValue)
-		{
-			check(Event == EventType::GAMEPAD_ANALOG);
-			OutControllerIndex = Data.GamepadAnalog.ControllerIndex;
-			OutAxisIndex = Data.GamepadAnalog.AxisIndex;
-			OutAnalogValue = Data.GamepadAnalog.AnalogValue;
-		}
-
-		/**
-		 * Get information about an event related to a gamepad button being pushed down.
-		 * @param OutControllerIndex - Number representing the index of the controller.
-		 * @param OutButtonIndex - Numerical code identifying the pushed down button.
-		 * @param OutIsRepeat - Whether the button is being kept down and is repeating.
-		*/
-		void GetGamepadButtonPressed(uint8& OutControllerIndex, uint8& OutButtonIndex, bool& OutIsRepeat)
-		{
-			check(Event == EventType::GAMEPAD_PRESS);
-			OutControllerIndex = Data.GamepadButton.ControllerIndex;
-			OutButtonIndex = Data.GamepadButton.ButtonIndex;
-			OutIsRepeat = Data.GamepadButton.bIsRepeat;
-		}
-
-		/**
-		 * Get information about an event related to a gamepad button being released.
-		 * @param OutControllerIndex - Number representing the index of the controller.
-		 * @param OutButtonIndex - Numerical code identifying the pushed down button.
-		*/
-		void GetGamepadButtonReleased(uint8& OutControllerIndex, uint8& OutButtonIndex)
-		{
-			check(Event == EventType::GAMEPAD_RELEASE);
-			OutControllerIndex = Data.GamepadButton.ControllerIndex;
-			OutButtonIndex = Data.GamepadButton.ButtonIndex;
-		}
-
-		/**
-		 * Get information about an event related to a key being pushed down.
-		 * @param OutKeyCode - Numerical code identifying the pushed down key.
-		 * @param OutIsRepeat - Whether the key is being kept down and is repeating.
-		 */
-		void GetKeyDown(uint8& OutKeyCode, bool& OutIsRepeat)
-		{
-			check(Event == EventType::KEY_DOWN);
-			OutKeyCode = Data.KeyDown.KeyCode;
-			OutIsRepeat = Data.KeyDown.bIsRepeat;
-		}
-
-		/**
-		 * Get information about an event related to a key being released.
-		 * @param OutKeyCode - Numerical code identifying the released key.
-		 */
-		void GetKeyUp(uint8& OutKeyCode)
-		{
-			check(Event == EventType::KEY_UP);
-			OutKeyCode = Data.KeyUp.KeyCode;
-		}
-
-		/**
-		 * Get information about an event related to character input.
-		 * @param OutCharacter - The character being input.
-		 */
-		void GetCharacterCode(TCHAR& OutCharacter)
-		{
-			check(Event == EventType::KEY_PRESS);
-			OutCharacter = Data.Character.Character;
-		}
-
-		/**
-		 * Get information about an event related to mouse movement.
-		 * @param OutPoxX - The X position of the mouse pointer.
-		 * @param OutPosY - The Y position of the mouse pointer.
-		 * @param OutDeltaX - The change in the X position of the mouse pointer.
-		 * @param OutDeltaY - The change in the Y position of the mouse pointer.
-		 */
-		void GetMouseDelta(uint16& OutPosX, uint16& OutPosY, int16& OutDeltaX, int16& OutDeltaY)
-		{
-			check(Event == EventType::MOUSE_MOVE);
-			OutPosX = Data.MouseMove.PosX;
-			OutPosY = Data.MouseMove.PosY;
-			OutDeltaX = Data.MouseMove.DeltaX;
-			OutDeltaY = Data.MouseMove.DeltaY;
-		}
-
-		/**
-		 * Get information about an event related to mouse buttons.
-		 * @param OutButton - The button number corresponding to left, middle, right, etc.
-		 * @param OutPosX - The X position of the mouse pointer.
-		 * @param OutPosY - The Y position of the mouse pointer.
-		 */
-		void GetMouseClick(EMouseButtons::Type& OutButton, uint16& OutPosX, uint16& OutPosY);
-
-		/**
-		 * Get information about an event related to the mouse wheel.
-		 * @param OutDelta - The amount by which the mouse wheel was scrolled.
-		 * @param PosX - The X position of the mouse pointer when the wheel was scrolled.
-		 * @param PosY - The Y position of the mouse pointer when the wheel was scrolled.
-		 */
-		void GetMouseWheel(int16& OutDelta, uint16& OutPosX, uint16& OutPosY)
-		{
-			check(Event == EventType::MOUSE_WHEEL);
-			OutDelta = Data.MouseWheel.Delta;
-			OutPosX = Data.MouseWheel.PosX;
-			OutPosY = Data.MouseWheel.PosY;
-		}
-
-		/**
-		 * Get information about an event related to a finger touching the canvas.
-		 * @param OutTouchIndex - The finger used in multi-touch.
-		 * @param OutPoxX - The X position of the finger.
-		 * @param OutPosY - The Y position of the finger.
-		 * @param OutForce - Amount of pressure being applied by the finger.
-		 */
-		void GetTouch(uint8& OutTouchIndex, uint16& OutPosX, uint16& OutPosY, uint8& OutForce)
-		{
-			check(Event == EventType::TOUCH_START || Event == EventType::TOUCH_END || Event == EventType::TOUCH_MOVE);
-			OutTouchIndex = Data.Touch.TouchIndex;
-			OutPosX = Data.Touch.PosX;
-			OutPosY = Data.Touch.PosY;
-			OutForce = Data.Touch.Force;
-		}
-	};
-
-	/**
-	 * Add a new event to the input device for later processing.
-	 * @param InEvent - The new event.
-	 */
-	void ProcessEvent(const FEvent& InEvent);
-
-	/**
-	 * Add a new UI interaction descriptor to the input device for later processing.
-	 * @param InDescriptor - The new UI interaction descriptor.
-	 */
-	void ProcessUIInteraction(const FString& InDescriptor);
-
-	/**
-	 * Add a new command descriptor to the input device for later execution.
-	 * @param InDescriptor - The new command descriptor.
-	 */
-	void ProcessCommand(const FString& InDescriptor);
-
-	/**
-	 * Is the application faking touch events?
-	 * @return True if the application is faking touch events.
-	*/
-	bool IsFakingTouchEvents() const
-	{
-		return bFakingTouchEvents;
-	}
-
-	/**
-     * Handle incoming messages
-     * @param Buffer - Incoming data buffer.
-     */
-	void OnMessage(const webrtc::DataBuffer& Buffer);
-
-	/**
-     * How long is the message type specifier and the length specifier in bytes?
-     * @return The number of header bytes in an incoming message what is ignored when it is parsed as string.
-     */
-	static const size_t GetMessageHeaderOffset()
-	{
-		return MessageHeaderOffset;
+				return bFakingTouchEvents;
+			}
+
+			/**
+			* Handle incoming messages
+			* @param Buffer - Incoming data buffer.
+			*/
+			void OnMessage(const webrtc::DataBuffer& Buffer);
+
+			/**
+			* How long is the message type specifier and the length specifier in bytes?
+			* @return The number of header bytes in an incoming message what is ignored when it is parsed as string.
+			*/
+			static const size_t GetMessageHeaderOffset()
+			{
+				return MessageHeaderOffset;
+			}
+
+		private:
+			/**
+			* Touch only. If the user touches the UI and focuses an editable widget
+			* then we need to tell the browser to allow it to display an on-screen
+			* keyboard.
+			*/
+			void FindFocusedWidget();
+
+			/**
+			* Convert the controller analog axis index to its corresponding FType
+			*/
+			FGamepadKeyNames::Type ConvertAxisIndexToGamepadAxis(uint8 AnalogAxis);
+
+			/**
+			* Convert the controller button index to its corresponding FType
+			*/
+			FGamepadKeyNames::Type ConvertButtonIndexToGamepadButton(uint8 ButtonIndex);
+
+			/**
+			* A special wrapper over the GenericApplication layer which allows us to
+			* override certain behavior.
+			*/
+			TSharedPtr<class FApplicationWrapper> PixelStreamerApplicationWrapper;
+
+			/** Reference to the message handler which events should be passed to. */
+			TSharedRef<FGenericApplicationMessageHandler> MessageHandler;
+
+			/** The queue of events which are awaiting processing. */
+			TQueue<FEvent> Events;
+
+			/**
+			* A queue of UI interaction descriptor strings which contain arbitrary
+			* information related to the interaction.
+			*/
+			TQueue<FString> UIInteractions;
+
+			/**
+			* A queue of command descriptor strings which contain the command to
+			* execute and its arguments.
+			*/
+			TQueue<FString> Commands;
+
+			/**
+			* Is the application faking touch events by dragging the mouse along
+			* the canvas? If so then we must put the browser canvas in a special
+			* state to replicate the behavior of the application.
+			*/
+			bool bFakingTouchEvents;
+
+			/**
+			* Touch only. Location of the focused UI widget. If no UI widget is focused
+			* then this has the UnfocusedPos value.
+			*/
+			FVector2D FocusedPos;
+
+			/**
+			* Touch only. A special position which indicates that no UI widget is
+			* focused.
+			*/
+			static const FVector2D UnfocusedPos;
+
+			/** For convenience we keep a reference to the Pixel Streaming plugin. */
+			IPixelStreamingModule* PixelStreamingModule;
+
+			/*
+			* Padding for string parsing when handling messages.
+			* 1 character for the actual message and then
+			* 2 characters for the length which are skipped
+			*/
+			static const size_t MessageHeaderOffset;
+		};
 	}
-
-private:
-	/**
-	 * Touch only. If the user touches the UI and focuses an editable widget
-	 * then we need to tell the browser to allow it to display an on-screen
-	 * keyboard.
-	 */
-	void FindFocusedWidget();
-
-	/**
-	 * Convert the controller analog axis index to its corresponding FType
-	 */
-	FGamepadKeyNames::Type ConvertAxisIndexToGamepadAxis(uint8 AnalogAxis);
-
-	/**
-	 * Convert the controller button index to its corresponding FType
-	 */
-	FGamepadKeyNames::Type ConvertButtonIndexToGamepadButton(uint8 ButtonIndex);
-
-	/**
-	 * A special wrapper over the GenericApplication layer which allows us to
-	 * override certain behavior.
-	 */
-	TSharedPtr<class FApplicationWrapper> PixelStreamerApplicationWrapper;
-
-	/** Reference to the message handler which events should be passed to. */
-	TSharedRef<FGenericApplicationMessageHandler> MessageHandler;
-
-	/** The queue of events which are awaiting processing. */
-	TQueue<FEvent> Events;
-
-	/**
-	 * A queue of UI interaction descriptor strings which contain arbitrary
-	 * information related to the interaction.
-	 */
-	TQueue<FString> UIInteractions;
-
-	/**
-	 * A queue of command descriptor strings which contain the command to
-	 * execute and its arguments.
-	 */
-	TQueue<FString> Commands;
-
-	/**
-	 * Is the application faking touch events by dragging the mouse along
-	 * the canvas? If so then we must put the browser canvas in a special
-	 * state to replicate the behavior of the application.
-	 */
-	bool bFakingTouchEvents;
-
-	/**
-	 * Touch only. Location of the focused UI widget. If no UI widget is focused
-	 * then this has the UnfocusedPos value.
-	 */
-	FVector2D FocusedPos;
-
-	/**
-	 * Touch only. A special position which indicates that no UI widget is
-	 * focused.
-	 */
-	static const FVector2D UnfocusedPos;
-
-	/** For convenience we keep a reference to the Pixel Streaming plugin. */
-	IPixelStreamingModule* PixelStreamingModule;
-
-	/*
-	 * Padding for string parsing when handling messages.
-	 * 1 character for the actual message and then
-	 * 2 characters for the length which are skipped
-	 */
-	static const size_t MessageHeaderOffset;
-};
+}
diff --git a/Engine/Plugins/Media/PixelStreaming/Source/PixelStreaming/Private/JavaScriptKeyCodes.inl b/Engine/Plugins/Media/PixelStreaming/Source/PixelStreaming/Private/JavaScriptKeyCodes.inl
--- a/Engine/Plugins/Media/PixelStreaming/Source/PixelStreaming/Private/JavaScriptKeyCodes.inl
+++ b/Engine/Plugins/Media/PixelStreaming/Source/PixelStreaming/Private/JavaScriptKeyCodes.inl
@@ -4,6 +4,8 @@
 
 #include "InputCoreTypes.h"
 
+namespace UE {
+	namespace PixelStreaming {
 /**
  * An array which will convert JavaScript key codes into FKeys. The JavaScript
  * key code should be used as an index into the array to perform the lookup.
@@ -13,265 +15,267 @@
  * have been choosen specially for 'right' versions of the same key.
  */
 
-static const FKey* JavaScriptKeyCodeToFKey[256] =
-{
-	/*   0 */ &EKeys::Invalid,
-	/*   1 */ &EKeys::Invalid,
-	/*   2 */ &EKeys::Invalid,
-	/*   3 */ &EKeys::Invalid,
-	/*   4 */ &EKeys::Invalid,
-	/*   5 */ &EKeys::Invalid,
-	/*   6 */ &EKeys::Invalid,
-	/*   7 */ &EKeys::Invalid,
-	/*   8 */ &EKeys::BackSpace,
-	/*   9 */ &EKeys::Tab,
-	/*  10 */ &EKeys::Invalid,
-	/*  11 */ &EKeys::Invalid,
-	/*  12 */ /*&EKeys::Clear*/ &EKeys::Invalid,
-	/*  13 */ &EKeys::Enter,
-	/*  14 */ &EKeys::Invalid,
-	/*  15 */ &EKeys::Invalid,
-	/*  16 */ &EKeys::LeftShift,
-	/*  17 */ &EKeys::LeftControl,
-	/*  18 */ &EKeys::LeftAlt,
-	/*  19 */ &EKeys::Pause,
-	/*  20 */ &EKeys::CapsLock,
-	/*  21 */ &EKeys::Invalid,
-	/*  22 */ &EKeys::Invalid,
-	/*  23 */ &EKeys::Invalid,
-	/*  24 */ &EKeys::Invalid,
-	/*  25 */ &EKeys::Invalid,
-	/*  26 */ &EKeys::Invalid,
-	/*  27 */ &EKeys::Escape,
-	/*  28 */ &EKeys::Invalid,
-	/*  29 */ &EKeys::Invalid,
-	/*  30 */ &EKeys::Invalid,
-	/*  31 */ &EKeys::Invalid,
-	/*  32 */ &EKeys::SpaceBar,
-	/*  33 */ &EKeys::PageUp,
-	/*  34 */ &EKeys::PageDown,
-	/*  35 */ &EKeys::End,
-	/*  36 */ &EKeys::Home,
-	/*  37 */ &EKeys::Left,
-	/*  38 */ &EKeys::Up,
-	/*  39 */ &EKeys::Right,
-	/*  40 */ &EKeys::Down,
-	/*  41 */ &EKeys::Invalid,
-	/*  42 */ &EKeys::Invalid,
-	/*  43 */ &EKeys::Invalid,
-	/*  44 */ /*&EKeys::PrintScreen*/ &EKeys::Invalid,
-	/*  45 */ &EKeys::Insert,
-	/*  46 */ &EKeys::Delete,
-	/*  47 */ &EKeys::Invalid,
-	/*  48 */ &EKeys::Zero,
-	/*  49 */ &EKeys::One,
-	/*  50 */ &EKeys::Two,
-	/*  51 */ &EKeys::Three,
-	/*  52 */ &EKeys::Four,
-	/*  53 */ &EKeys::Five,
-	/*  54 */ &EKeys::Six,
-	/*  55 */ &EKeys::Seven,
-	/*  56 */ &EKeys::Eight,
-	/*  57 */ &EKeys::Nine,
-	/*  58 */ &EKeys::Invalid,
-	/*  59 */ &EKeys::Invalid,
-	/*  60 */ &EKeys::Invalid,
-	/*  61 */ &EKeys::Invalid,
-	/*  62 */ &EKeys::Invalid,
-	/*  63 */ &EKeys::Invalid,
-	/*  64 */ &EKeys::Invalid,
-	/*  65 */ &EKeys::A,
-	/*  66 */ &EKeys::B,
-	/*  67 */ &EKeys::C,
-	/*  68 */ &EKeys::D,
-	/*  69 */ &EKeys::E,
-	/*  70 */ &EKeys::F,
-	/*  71 */ &EKeys::G,
-	/*  72 */ &EKeys::H,
-	/*  73 */ &EKeys::I,
-	/*  74 */ &EKeys::J,
-	/*  75 */ &EKeys::K,
-	/*  76 */ &EKeys::L,
-	/*  77 */ &EKeys::M,
-	/*  78 */ &EKeys::N,
-	/*  79 */ &EKeys::O,
-	/*  80 */ &EKeys::P,
-	/*  81 */ &EKeys::Q,
-	/*  82 */ &EKeys::R,
-	/*  83 */ &EKeys::S,
-	/*  84 */ &EKeys::T,
-	/*  85 */ &EKeys::U,
-	/*  86 */ &EKeys::V,
-	/*  87 */ &EKeys::W,
-	/*  88 */ &EKeys::X,
-	/*  89 */ &EKeys::Y,
-	/*  90 */ &EKeys::Z,
-	/*  91 */ /*&EKeys::LeftWindowKey*/ &EKeys::Invalid,
-	/*  92 */ /*&EKeys::RightWindowKey*/ &EKeys::Invalid,
-	/*  93 */ /*&EKeys::SelectKey*/ &EKeys::Invalid,
-	/*  94 */ &EKeys::Invalid,
-	/*  95 */ &EKeys::Invalid,
-	/*  96 */ &EKeys::NumPadZero,
-	/*  97 */ &EKeys::NumPadOne,
-	/*  98 */ &EKeys::NumPadTwo,
-	/*  99 */ &EKeys::NumPadThree,
-	/* 100 */ &EKeys::NumPadFour,
-	/* 101 */ &EKeys::NumPadFive,
-	/* 102 */ &EKeys::NumPadSix,
-	/* 103 */ &EKeys::NumPadSeven,
-	/* 104 */ &EKeys::NumPadEight,
-	/* 105 */ &EKeys::NumPadNine,
-	/* 106 */ &EKeys::Multiply,
-	/* 107 */ &EKeys::Add,
-	/* 108 */ &EKeys::Invalid,
-	/* 109 */ &EKeys::Subtract,
-	/* 110 */ &EKeys::Decimal,
-	/* 111 */ &EKeys::Divide,
-	/* 112 */ &EKeys::F1,
-	/* 113 */ &EKeys::F2,
-	/* 114 */ &EKeys::F3,
-	/* 115 */ &EKeys::F4,
-	/* 116 */ &EKeys::F5,
-	/* 117 */ &EKeys::F6,
-	/* 118 */ &EKeys::F7,
-	/* 119 */ &EKeys::F8,
-	/* 120 */ &EKeys::F9,
-	/* 121 */ &EKeys::F10,
-	/* 122 */ &EKeys::F11,
-	/* 123 */ &EKeys::F12,
-	/* 124 */ &EKeys::Invalid,
-	/* 125 */ &EKeys::Invalid,
-	/* 126 */ &EKeys::Invalid,
-	/* 127 */ &EKeys::Invalid,
-	/* 128 */ &EKeys::Invalid,
-	/* 129 */ &EKeys::Invalid,
-	/* 130 */ &EKeys::Invalid,
-	/* 131 */ &EKeys::Invalid,
-	/* 132 */ &EKeys::Invalid,
-	/* 133 */ &EKeys::Invalid,
-	/* 134 */ &EKeys::Invalid,
-	/* 135 */ &EKeys::Invalid,
-	/* 136 */ &EKeys::Invalid,
-	/* 137 */ &EKeys::Invalid,
-	/* 138 */ &EKeys::Invalid,
-	/* 139 */ &EKeys::Invalid,
-	/* 140 */ &EKeys::Invalid,
-	/* 141 */ &EKeys::Invalid,
-	/* 142 */ &EKeys::Invalid,
-	/* 143 */ &EKeys::Invalid,
-	/* 144 */ &EKeys::NumLock,
-	/* 145 */ &EKeys::ScrollLock,
-	/* 146 */ &EKeys::Invalid,
-	/* 147 */ &EKeys::Invalid,
-	/* 148 */ &EKeys::Invalid,
-	/* 149 */ &EKeys::Invalid,
-	/* 150 */ &EKeys::Invalid,
-	/* 151 */ &EKeys::Invalid,
-	/* 152 */ &EKeys::Invalid,
-	/* 153 */ &EKeys::Invalid,
-	/* 154 */ &EKeys::Invalid,
-	/* 155 */ &EKeys::Invalid,
-	/* 156 */ &EKeys::Invalid,
-	/* 157 */ &EKeys::Invalid,
-	/* 158 */ &EKeys::Invalid,
-	/* 159 */ &EKeys::Invalid,
-	/* 160 */ &EKeys::Invalid,
-	/* 161 */ &EKeys::Invalid,
-	/* 162 */ &EKeys::Invalid,
-	/* 163 */ &EKeys::Invalid,
-	/* 164 */ &EKeys::Invalid,
-	/* 165 */ &EKeys::Invalid,
-	/* 166 */ &EKeys::Invalid,
-	/* 167 */ &EKeys::Invalid,
-	/* 168 */ &EKeys::Invalid,
-	/* 169 */ &EKeys::Invalid,
-	/* 170 */ &EKeys::Invalid,
-	/* 171 */ &EKeys::Invalid,
-	/* 172 */ &EKeys::Invalid,
-	/* 173 */ &EKeys::Invalid,
-	/* 174 */ &EKeys::Invalid,
-	/* 175 */ &EKeys::Invalid,
-	/* 176 */ &EKeys::Invalid,
-	/* 177 */ &EKeys::Invalid,
-	/* 178 */ &EKeys::Invalid,
-	/* 179 */ &EKeys::Invalid,
-	/* 180 */ &EKeys::Invalid,
-	/* 181 */ &EKeys::Invalid,
-	/* 182 */ &EKeys::Invalid,
-	/* 183 */ &EKeys::Invalid,
-	/* 184 */ &EKeys::Invalid,
-	/* 185 */ &EKeys::Invalid,
-	/* 186 */ &EKeys::Semicolon,
-	/* 187 */ &EKeys::Equals,
-	/* 188 */ &EKeys::Comma,
-	/* 189 */ &EKeys::Hyphen,
-	/* 190 */ &EKeys::Period,
-	/* 191 */ &EKeys::Slash,
-	/* 192 */ &EKeys::Tilde,
-	/* 193 */ &EKeys::Invalid,
-	/* 194 */ &EKeys::Invalid,
-	/* 195 */ &EKeys::Invalid,
-	/* 196 */ &EKeys::Invalid,
-	/* 197 */ &EKeys::Invalid,
-	/* 198 */ &EKeys::Invalid,
-	/* 199 */ &EKeys::Invalid,
-	/* 200 */ &EKeys::Invalid,
-	/* 201 */ &EKeys::Invalid,
-	/* 202 */ &EKeys::Invalid,
-	/* 203 */ &EKeys::Invalid,
-	/* 204 */ &EKeys::Invalid,
-	/* 205 */ &EKeys::Invalid,
-	/* 206 */ &EKeys::Invalid,
-	/* 207 */ &EKeys::Invalid,
-	/* 208 */ &EKeys::Invalid,
-	/* 209 */ &EKeys::Invalid,
-	/* 210 */ &EKeys::Invalid,
-	/* 211 */ &EKeys::Invalid,
-	/* 212 */ &EKeys::Invalid,
-	/* 213 */ &EKeys::Invalid,
-	/* 214 */ &EKeys::Invalid,
-	/* 215 */ &EKeys::Invalid,
-	/* 216 */ &EKeys::Invalid,
-	/* 217 */ &EKeys::Invalid,
-	/* 218 */ &EKeys::Invalid,
-	/* 219 */ &EKeys::LeftBracket,
-	/* 220 */ &EKeys::Backslash,
-	/* 221 */ &EKeys::RightBracket,
-	/* 222 */ &EKeys::Apostrophe,
-	/* 223 */ &EKeys::Tilde,
-	/* 224 */ &EKeys::Invalid,
-	/* 225 */ &EKeys::Invalid,
-	/* 226 */ &EKeys::Invalid,
-	/* 227 */ &EKeys::Invalid,
-	/* 228 */ &EKeys::Invalid,
-	/* 229 */ &EKeys::Invalid,
-	/* 230 */ &EKeys::Invalid,
-	/* 231 */ &EKeys::Invalid,
-	/* 232 */ &EKeys::Invalid,
-	/* 233 */ &EKeys::Invalid,
-	/* 234 */ &EKeys::Invalid,
-	/* 235 */ &EKeys::Invalid,
-	/* 236 */ &EKeys::Invalid,
-	/* 237 */ &EKeys::Invalid,
-	/* 238 */ &EKeys::Invalid,
-	/* 239 */ &EKeys::Invalid,
-	/* 240 */ &EKeys::Invalid,
-	/* 241 */ &EKeys::Invalid,
-	/* 242 */ &EKeys::Invalid,
-	/* 243 */ &EKeys::Invalid,
-	/* 244 */ &EKeys::Invalid,
-	/* 245 */ &EKeys::Invalid,
-	/* 246 */ &EKeys::Invalid,
-	/* 247 */ &EKeys::Invalid,
-	/* 248 */ &EKeys::Invalid,
-	/* 249 */ &EKeys::Invalid,
-	/* 250 */ &EKeys::Invalid,
-	/* 251 */ &EKeys::Invalid,
-	/* 252 */ &EKeys::Invalid,
+		static const FKey* JavaScriptKeyCodeToFKey[256] =
+		{
+			/*   0 */ &EKeys::Invalid,
+			/*   1 */ &EKeys::Invalid,
+			/*   2 */ &EKeys::Invalid,
+			/*   3 */ &EKeys::Invalid,
+			/*   4 */ &EKeys::Invalid,
+			/*   5 */ &EKeys::Invalid,
+			/*   6 */ &EKeys::Invalid,
+			/*   7 */ &EKeys::Invalid,
+			/*   8 */ &EKeys::BackSpace,
+			/*   9 */ &EKeys::Tab,
+			/*  10 */ &EKeys::Invalid,
+			/*  11 */ &EKeys::Invalid,
+			/*  12 */ /*&EKeys::Clear*/ &EKeys::Invalid,
+			/*  13 */ &EKeys::Enter,
+			/*  14 */ &EKeys::Invalid,
+			/*  15 */ &EKeys::Invalid,
+			/*  16 */ &EKeys::LeftShift,
+			/*  17 */ &EKeys::LeftControl,
+			/*  18 */ &EKeys::LeftAlt,
+			/*  19 */ &EKeys::Pause,
+			/*  20 */ &EKeys::CapsLock,
+			/*  21 */ &EKeys::Invalid,
+			/*  22 */ &EKeys::Invalid,
+			/*  23 */ &EKeys::Invalid,
+			/*  24 */ &EKeys::Invalid,
+			/*  25 */ &EKeys::Invalid,
+			/*  26 */ &EKeys::Invalid,
+			/*  27 */ &EKeys::Escape,
+			/*  28 */ &EKeys::Invalid,
+			/*  29 */ &EKeys::Invalid,
+			/*  30 */ &EKeys::Invalid,
+			/*  31 */ &EKeys::Invalid,
+			/*  32 */ &EKeys::SpaceBar,
+			/*  33 */ &EKeys::PageUp,
+			/*  34 */ &EKeys::PageDown,
+			/*  35 */ &EKeys::End,
+			/*  36 */ &EKeys::Home,
+			/*  37 */ &EKeys::Left,
+			/*  38 */ &EKeys::Up,
+			/*  39 */ &EKeys::Right,
+			/*  40 */ &EKeys::Down,
+			/*  41 */ &EKeys::Invalid,
+			/*  42 */ &EKeys::Invalid,
+			/*  43 */ &EKeys::Invalid,
+			/*  44 */ /*&EKeys::PrintScreen*/ &EKeys::Invalid,
+			/*  45 */ &EKeys::Insert,
+			/*  46 */ &EKeys::Delete,
+			/*  47 */ &EKeys::Invalid,
+			/*  48 */ &EKeys::Zero,
+			/*  49 */ &EKeys::One,
+			/*  50 */ &EKeys::Two,
+			/*  51 */ &EKeys::Three,
+			/*  52 */ &EKeys::Four,
+			/*  53 */ &EKeys::Five,
+			/*  54 */ &EKeys::Six,
+			/*  55 */ &EKeys::Seven,
+			/*  56 */ &EKeys::Eight,
+			/*  57 */ &EKeys::Nine,
+			/*  58 */ &EKeys::Invalid,
+			/*  59 */ &EKeys::Invalid,
+			/*  60 */ &EKeys::Invalid,
+			/*  61 */ &EKeys::Invalid,
+			/*  62 */ &EKeys::Invalid,
+			/*  63 */ &EKeys::Invalid,
+			/*  64 */ &EKeys::Invalid,
+			/*  65 */ &EKeys::A,
+			/*  66 */ &EKeys::B,
+			/*  67 */ &EKeys::C,
+			/*  68 */ &EKeys::D,
+			/*  69 */ &EKeys::E,
+			/*  70 */ &EKeys::F,
+			/*  71 */ &EKeys::G,
+			/*  72 */ &EKeys::H,
+			/*  73 */ &EKeys::I,
+			/*  74 */ &EKeys::J,
+			/*  75 */ &EKeys::K,
+			/*  76 */ &EKeys::L,
+			/*  77 */ &EKeys::M,
+			/*  78 */ &EKeys::N,
+			/*  79 */ &EKeys::O,
+			/*  80 */ &EKeys::P,
+			/*  81 */ &EKeys::Q,
+			/*  82 */ &EKeys::R,
+			/*  83 */ &EKeys::S,
+			/*  84 */ &EKeys::T,
+			/*  85 */ &EKeys::U,
+			/*  86 */ &EKeys::V,
+			/*  87 */ &EKeys::W,
+			/*  88 */ &EKeys::X,
+			/*  89 */ &EKeys::Y,
+			/*  90 */ &EKeys::Z,
+			/*  91 */ /*&EKeys::LeftWindowKey*/ &EKeys::Invalid,
+			/*  92 */ /*&EKeys::RightWindowKey*/ &EKeys::Invalid,
+			/*  93 */ /*&EKeys::SelectKey*/ &EKeys::Invalid,
+			/*  94 */ &EKeys::Invalid,
+			/*  95 */ &EKeys::Invalid,
+			/*  96 */ &EKeys::NumPadZero,
+			/*  97 */ &EKeys::NumPadOne,
+			/*  98 */ &EKeys::NumPadTwo,
+			/*  99 */ &EKeys::NumPadThree,
+			/* 100 */ &EKeys::NumPadFour,
+			/* 101 */ &EKeys::NumPadFive,
+			/* 102 */ &EKeys::NumPadSix,
+			/* 103 */ &EKeys::NumPadSeven,
+			/* 104 */ &EKeys::NumPadEight,
+			/* 105 */ &EKeys::NumPadNine,
+			/* 106 */ &EKeys::Multiply,
+			/* 107 */ &EKeys::Add,
+			/* 108 */ &EKeys::Invalid,
+			/* 109 */ &EKeys::Subtract,
+			/* 110 */ &EKeys::Decimal,
+			/* 111 */ &EKeys::Divide,
+			/* 112 */ &EKeys::F1,
+			/* 113 */ &EKeys::F2,
+			/* 114 */ &EKeys::F3,
+			/* 115 */ &EKeys::F4,
+			/* 116 */ &EKeys::F5,
+			/* 117 */ &EKeys::F6,
+			/* 118 */ &EKeys::F7,
+			/* 119 */ &EKeys::F8,
+			/* 120 */ &EKeys::F9,
+			/* 121 */ &EKeys::F10,
+			/* 122 */ &EKeys::F11,
+			/* 123 */ &EKeys::F12,
+			/* 124 */ &EKeys::Invalid,
+			/* 125 */ &EKeys::Invalid,
+			/* 126 */ &EKeys::Invalid,
+			/* 127 */ &EKeys::Invalid,
+			/* 128 */ &EKeys::Invalid,
+			/* 129 */ &EKeys::Invalid,
+			/* 130 */ &EKeys::Invalid,
+			/* 131 */ &EKeys::Invalid,
+			/* 132 */ &EKeys::Invalid,
+			/* 133 */ &EKeys::Invalid,
+			/* 134 */ &EKeys::Invalid,
+			/* 135 */ &EKeys::Invalid,
+			/* 136 */ &EKeys::Invalid,
+			/* 137 */ &EKeys::Invalid,
+			/* 138 */ &EKeys::Invalid,
+			/* 139 */ &EKeys::Invalid,
+			/* 140 */ &EKeys::Invalid,
+			/* 141 */ &EKeys::Invalid,
+			/* 142 */ &EKeys::Invalid,
+			/* 143 */ &EKeys::Invalid,
+			/* 144 */ &EKeys::NumLock,
+			/* 145 */ &EKeys::ScrollLock,
+			/* 146 */ &EKeys::Invalid,
+			/* 147 */ &EKeys::Invalid,
+			/* 148 */ &EKeys::Invalid,
+			/* 149 */ &EKeys::Invalid,
+			/* 150 */ &EKeys::Invalid,
+			/* 151 */ &EKeys::Invalid,
+			/* 152 */ &EKeys::Invalid,
+			/* 153 */ &EKeys::Invalid,
+			/* 154 */ &EKeys::Invalid,
+			/* 155 */ &EKeys::Invalid,
+			/* 156 */ &EKeys::Invalid,
+			/* 157 */ &EKeys::Invalid,
+			/* 158 */ &EKeys::Invalid,
+			/* 159 */ &EKeys::Invalid,
+			/* 160 */ &EKeys::Invalid,
+			/* 161 */ &EKeys::Invalid,
+			/* 162 */ &EKeys::Invalid,
+			/* 163 */ &EKeys::Invalid,
+			/* 164 */ &EKeys::Invalid,
+			/* 165 */ &EKeys::Invalid,
+			/* 166 */ &EKeys::Invalid,
+			/* 167 */ &EKeys::Invalid,
+			/* 168 */ &EKeys::Invalid,
+			/* 169 */ &EKeys::Invalid,
+			/* 170 */ &EKeys::Invalid,
+			/* 171 */ &EKeys::Invalid,
+			/* 172 */ &EKeys::Invalid,
+			/* 173 */ &EKeys::Invalid,
+			/* 174 */ &EKeys::Invalid,
+			/* 175 */ &EKeys::Invalid,
+			/* 176 */ &EKeys::Invalid,
+			/* 177 */ &EKeys::Invalid,
+			/* 178 */ &EKeys::Invalid,
+			/* 179 */ &EKeys::Invalid,
+			/* 180 */ &EKeys::Invalid,
+			/* 181 */ &EKeys::Invalid,
+			/* 182 */ &EKeys::Invalid,
+			/* 183 */ &EKeys::Invalid,
+			/* 184 */ &EKeys::Invalid,
+			/* 185 */ &EKeys::Invalid,
+			/* 186 */ &EKeys::Semicolon,
+			/* 187 */ &EKeys::Equals,
+			/* 188 */ &EKeys::Comma,
+			/* 189 */ &EKeys::Hyphen,
+			/* 190 */ &EKeys::Period,
+			/* 191 */ &EKeys::Slash,
+			/* 192 */ &EKeys::Tilde,
+			/* 193 */ &EKeys::Invalid,
+			/* 194 */ &EKeys::Invalid,
+			/* 195 */ &EKeys::Invalid,
+			/* 196 */ &EKeys::Invalid,
+			/* 197 */ &EKeys::Invalid,
+			/* 198 */ &EKeys::Invalid,
+			/* 199 */ &EKeys::Invalid,
+			/* 200 */ &EKeys::Invalid,
+			/* 201 */ &EKeys::Invalid,
+			/* 202 */ &EKeys::Invalid,
+			/* 203 */ &EKeys::Invalid,
+			/* 204 */ &EKeys::Invalid,
+			/* 205 */ &EKeys::Invalid,
+			/* 206 */ &EKeys::Invalid,
+			/* 207 */ &EKeys::Invalid,
+			/* 208 */ &EKeys::Invalid,
+			/* 209 */ &EKeys::Invalid,
+			/* 210 */ &EKeys::Invalid,
+			/* 211 */ &EKeys::Invalid,
+			/* 212 */ &EKeys::Invalid,
+			/* 213 */ &EKeys::Invalid,
+			/* 214 */ &EKeys::Invalid,
+			/* 215 */ &EKeys::Invalid,
+			/* 216 */ &EKeys::Invalid,
+			/* 217 */ &EKeys::Invalid,
+			/* 218 */ &EKeys::Invalid,
+			/* 219 */ &EKeys::LeftBracket,
+			/* 220 */ &EKeys::Backslash,
+			/* 221 */ &EKeys::RightBracket,
+			/* 222 */ &EKeys::Apostrophe,
+			/* 223 */ &EKeys::Tilde,
+			/* 224 */ &EKeys::Invalid,
+			/* 225 */ &EKeys::Invalid,
+			/* 226 */ &EKeys::Invalid,
+			/* 227 */ &EKeys::Invalid,
+			/* 228 */ &EKeys::Invalid,
+			/* 229 */ &EKeys::Invalid,
+			/* 230 */ &EKeys::Invalid,
+			/* 231 */ &EKeys::Invalid,
+			/* 232 */ &EKeys::Invalid,
+			/* 233 */ &EKeys::Invalid,
+			/* 234 */ &EKeys::Invalid,
+			/* 235 */ &EKeys::Invalid,
+			/* 236 */ &EKeys::Invalid,
+			/* 237 */ &EKeys::Invalid,
+			/* 238 */ &EKeys::Invalid,
+			/* 239 */ &EKeys::Invalid,
+			/* 240 */ &EKeys::Invalid,
+			/* 241 */ &EKeys::Invalid,
+			/* 242 */ &EKeys::Invalid,
+			/* 243 */ &EKeys::Invalid,
+			/* 244 */ &EKeys::Invalid,
+			/* 245 */ &EKeys::Invalid,
+			/* 246 */ &EKeys::Invalid,
+			/* 247 */ &EKeys::Invalid,
+			/* 248 */ &EKeys::Invalid,
+			/* 249 */ &EKeys::Invalid,
+			/* 250 */ &EKeys::Invalid,
+			/* 251 */ &EKeys::Invalid,
+			/* 252 */ &EKeys::Invalid,
 
-	// Special key codes.
+			// Special key codes.
 
-	/* 253 */ &EKeys::RightShift,
-	/* 254 */ &EKeys::RightControl,
-	/* 255 */ &EKeys::RightAlt
-};
+			/* 253 */ &EKeys::RightShift,
+			/* 254 */ &EKeys::RightControl,
+			/* 255 */ &EKeys::RightAlt
+		};
+	}
+}
diff --git a/Engine/Plugins/Media/PixelStreaming/Source/PixelStreaming/Private/PCH.h b/Engine/Plugins/Media/PixelStreaming/Source/PixelStreaming/Private/PCH.h
--- a/Engine/Plugins/Media/PixelStreaming/Source/PixelStreaming/Private/PCH.h
+++ b/Engine/Plugins/Media/PixelStreaming/Source/PixelStreaming/Private/PCH.h
@@ -16,6 +16,7 @@
 
 #include "WebRTCLogging.h"
 #include "Utils.h"
+#include "ToStringExtensions.h"
 
 // Engine
 
