Merging Pixel Streaming refactor using //UE5/Dev-TensorWorks-to-//UE5/Release-5.0

diff --git a/Engine/Plugins/Media/PixelStreaming/Source/PixelStreaming/Private/PixelStreamingAudioComponent.cpp b/Engine/Plugins/Media/PixelStreaming/Source/PixelStreaming/Private/PixelStreamingAudioComponent.cpp
--- a/Engine/Plugins/Media/PixelStreaming/Source/PixelStreaming/Private/PixelStreamingAudioComponent.cpp
+++ b/Engine/Plugins/Media/PixelStreaming/Source/PixelStreaming/Private/PixelStreamingAudioComponent.cpp
@@ -7,7 +7,7 @@
 
 UPixelStreamingAudioComponent::UPixelStreamingAudioComponent(const FObjectInitializer& ObjectInitializer)
 	: Super(ObjectInitializer)
-	, PlayerToHear(FPlayerId())
+	, PlayerToHear(FPixelStreamingPlayerId())
 	, bAutoFindPeer(true)
 	, Buffer()
 	, bIsListeningToPeer(false)
@@ -22,19 +22,19 @@ UPixelStreamingAudioComponent::UPixelStreamingAudioComponent(const FObjectInitia
 	// Only output this warning if we are actually running this component (not in commandlet).
 	if (!bPixelStreamingLoaded && !IsRunningCommandlet())
 	{
-		UE_LOG(PixelStreamer, Warning, TEXT("Pixel Streaming audio component will not tick because Pixel Streaming module is not loaded. This is expected on dedicated servers."));
+		UE_LOG(LogPixelStreaming, Warning, TEXT("Pixel Streaming audio component will not tick because Pixel Streaming module is not loaded. This is expected on dedicated servers."));
 	}
 
-	//this->NumChannels = 2; //2 channels seem to cause problems
-	this->PrimaryComponentTick.bCanEverTick = bPixelStreamingLoaded;
-	this->SetComponentTickEnabled(bPixelStreamingLoaded);
-	this->bAutoActivate = true;
+	//NumChannels = 2; //2 channels seem to cause problems
+	PrimaryComponentTick.bCanEverTick = bPixelStreamingLoaded;
+	SetComponentTickEnabled(bPixelStreamingLoaded);
+	bAutoActivate = true;
 };
 
 int32 UPixelStreamingAudioComponent::OnGenerateAudio(float* OutAudio, int32 NumSamples)
 {
 	// Not listening to peer, return zero'd buffer.
-	if (!this->bIsListeningToPeer)
+	if (!bIsListeningToPeer)
 	{
 		return 0;
 	}
@@ -43,7 +43,7 @@ int32 UPixelStreamingAudioComponent::OnGenerateAudio(float* OutAudio, int32 NumS
 	{
 		FScopeLock Lock(&CriticalSection);
 
-		int32 NumSamplesToCopy = FGenericPlatformMath::Min(NumSamples, this->Buffer.Num());
+		int32 NumSamplesToCopy = FGenericPlatformMath::Min(NumSamples, Buffer.Num());
 
 		// Not enough samples to copy anything across
 		if (NumSamplesToCopy <= 0)
@@ -55,12 +55,12 @@ int32 UPixelStreamingAudioComponent::OnGenerateAudio(float* OutAudio, int32 NumS
 		for (int SampleIndex = 0; SampleIndex < NumSamplesToCopy; SampleIndex++)
 		{
 			// Convert from int16 to float audio
-			*OutAudio = ((float)this->Buffer[SampleIndex]) / 32767.0f;
+			*OutAudio = ((float)Buffer[SampleIndex]) / 32767.0f;
 			OutAudio++;
 		}
 
 		// Remove front NumSamples from the local buffer
-		this->Buffer.RemoveAt(0, NumSamplesToCopy, false);
+		Buffer.RemoveAt(0, NumSamplesToCopy, false);
 		return NumSamplesToCopy;
 	}
 }
@@ -68,18 +68,18 @@ int32 UPixelStreamingAudioComponent::OnGenerateAudio(float* OutAudio, int32 NumS
 bool UPixelStreamingAudioComponent::UpdateChannelsAndSampleRate(int InNumChannels, int InSampleRate)
 {
 
-	if (InNumChannels != this->NumChannels || InSampleRate != this->SampleRate)
+	if (InNumChannels != NumChannels || InSampleRate != SampleRate)
 	{
 
 		// Critical Section - empty buffer because sample rate/num channels changed
 		FScopeLock Lock(&CriticalSection);
-		this->Buffer.Empty();
+		Buffer.Empty();
 
 		//this is the smallest buffer size we can set without triggering internal checks to fire
-		this->PreferredBufferLength = FGenericPlatformMath::Max(512, InSampleRate * InNumChannels / 100);
-		this->NumChannels = InNumChannels;
-		this->SampleRate = InSampleRate;
-		this->Initialize(this->SampleRate);
+		PreferredBufferLength = FGenericPlatformMath::Max(512, InSampleRate * InNumChannels / 100);
+		NumChannels = InNumChannels;
+		SampleRate = InSampleRate;
+		Initialize(SampleRate);
 
 		return true;
 	}
@@ -89,18 +89,18 @@ bool UPixelStreamingAudioComponent::UpdateChannelsAndSampleRate(int InNumChannel
 
 void UPixelStreamingAudioComponent::OnBeginGenerate()
 {
-	this->bComponentWantsAudio = true;
+	bComponentWantsAudio = true;
 }
 
 void UPixelStreamingAudioComponent::OnEndGenerate()
 {
-	this->bComponentWantsAudio = false;
+	bComponentWantsAudio = false;
 }
 
 void UPixelStreamingAudioComponent::BeginDestroy()
 {
 	Super::BeginDestroy();
-	this->Reset();
+	Reset();
 }
 
 bool UPixelStreamingAudioComponent::ListenTo(FString PlayerToListenTo)
@@ -108,64 +108,68 @@ bool UPixelStreamingAudioComponent::ListenTo(FString PlayerToListenTo)
 
 	if (!IPixelStreamingModule::IsAvailable())
 	{
-		UE_LOG(PixelStreamer, Verbose, TEXT("Pixel Streaming audio component could not listen to anything because Pixel Streaming module is not loaded. This is expected on dedicated servers."));
+		UE_LOG(LogPixelStreaming, Verbose, TEXT("Pixel Streaming audio component could not listen to anything because Pixel Streaming module is not loaded. This is expected on dedicated servers."));
 		return false;
 	}
 
 	IPixelStreamingModule& PixelStreamingModule = IPixelStreamingModule::Get();
+	if (!PixelStreamingModule.IsReady())
+	{
+		return false;
+	}
 
-	this->PlayerToHear = PlayerToListenTo;
+	PlayerToHear = PlayerToListenTo;
 
-	IPixelStreamingAudioSink* CandidateSink = this->WillListenToAnyPlayer() ? PixelStreamingModule.GetUnlistenedAudioSink() : PixelStreamingModule.GetPeerAudioSink(FPlayerId(this->PlayerToHear));
+	IPixelStreamingAudioSink* CandidateSink = WillListenToAnyPlayer() ? PixelStreamingModule.GetUnlistenedAudioSink() : PixelStreamingModule.GetPeerAudioSink(FPixelStreamingPlayerId(PlayerToHear));
 
 	if (CandidateSink == nullptr)
 	{
 		return false;
 	}
 
-	this->AudioSink = CandidateSink;
-	this->AudioSink->AddAudioConsumer(this);
+	AudioSink = CandidateSink;
+	AudioSink->AddAudioConsumer(this);
 
 	return true;
 }
 
 void UPixelStreamingAudioComponent::Reset()
 {
-	this->PlayerToHear = FString();
-	this->bIsListeningToPeer = false;
-	if (this->AudioSink)
+	PlayerToHear = FString();
+	bIsListeningToPeer = false;
+	if (AudioSink)
 	{
-		this->AudioSink->RemoveAudioConsumer(this);
+		AudioSink->RemoveAudioConsumer(this);
 	}
-	this->AudioSink = nullptr;
+	AudioSink = nullptr;
 
 	// Critical section
 	{
 		FScopeLock Lock(&CriticalSection);
-		this->Buffer.Empty();
+		Buffer.Empty();
 	}
 }
 
 bool UPixelStreamingAudioComponent::IsListeningToPlayer()
 {
-	return this->bIsListeningToPeer;
+	return bIsListeningToPeer;
 }
 
 bool UPixelStreamingAudioComponent::WillListenToAnyPlayer()
 {
-	return this->PlayerToHear == FString();
+	return PlayerToHear == FString();
 }
 
 void UPixelStreamingAudioComponent::ConsumeRawPCM(const int16_t* AudioData, int InSampleRate, size_t NChannels, size_t NFrames)
 {
 
-	if (!this->bComponentWantsAudio)
+	if (!bComponentWantsAudio)
 	{
 		return;
 	}
 
 	// Trigger a latent update for channels and sample rate, if required
-	bool bSampleRateOrChannelMismatch = this->UpdateChannelsAndSampleRate(NChannels, InSampleRate);
+	bool bSampleRateOrChannelMismatch = UpdateChannelsAndSampleRate(NChannels, InSampleRate);
 
 	// If there was a mismatch then skip this incoming audio until this the samplerate/channels update is completed on the gamethread
 	if (bSampleRateOrChannelMismatch)
@@ -179,39 +183,39 @@ void UPixelStreamingAudioComponent::ConsumeRawPCM(const int16_t* AudioData, int
 	// Critical Section
 	{
 		FScopeLock Lock(&CriticalSection);
-		this->Buffer.Append(AudioData, NSamples);
-		checkf((uint32)this->Buffer.Num() < this->SampleRate,
+		Buffer.Append(AudioData, NSamples);
+		checkf((uint32)Buffer.Num() < SampleRate,
 			TEXT("Pixel Streaming Audio Component internal buffer is getting too big, for some reason OnGenerateAudio is not consuming samples quickly enough."))
 	}
 }
 
 void UPixelStreamingAudioComponent::OnConsumerAdded()
 {
-	this->bIsListeningToPeer = true;
+	bIsListeningToPeer = true;
 }
 
 void UPixelStreamingAudioComponent::OnConsumerRemoved()
 {
-	this->Reset();
+	Reset();
 }
 
 void UPixelStreamingAudioComponent::TickComponent(float DeltaTime, ELevelTick TickType, FActorComponentTickFunction* ThisTickFunction)
 {
 
 	// if auto connect turned off don't bother
-	if (!this->bAutoFindPeer)
+	if (!bAutoFindPeer)
 	{
 		return;
 	}
 
 	// if listening to a peer don't auto connect
-	if (this->bIsListeningToPeer)
+	if (bIsListeningToPeer)
 	{
 		return;
 	}
 
-	if (this->ListenTo(this->PlayerToHear))
+	if (ListenTo(PlayerToHear))
 	{
-		UE_LOG(PixelStreamer, Log, TEXT("PixelStreaming audio component found a WebRTC peer to listen to."));
+		UE_LOG(LogPixelStreaming, Log, TEXT("PixelStreaming audio component found a WebRTC peer to listen to."));
 	}
 }
diff --git a/Engine/Plugins/Media/PixelStreaming/Source/PixelStreaming/Private/PixelStreamingBlueprints.cpp b/Engine/Plugins/Media/PixelStreaming/Source/PixelStreaming/Private/PixelStreamingBlueprints.cpp
--- a/Engine/Plugins/Media/PixelStreaming/Source/PixelStreaming/Private/PixelStreamingBlueprints.cpp
+++ b/Engine/Plugins/Media/PixelStreaming/Source/PixelStreaming/Private/PixelStreamingBlueprints.cpp
@@ -3,12 +3,12 @@
 #include "PixelStreamingBlueprints.h"
 #include "PixelStreamingPrivate.h"
 #include "Misc/FileHelper.h"
-#include "PixelStreamerInputComponent.h"
+#include "PixelStreamingInputComponent.h"
 #include "PixelStreamingAudioComponent.h"
 
 void UPixelStreamingBlueprints::SendFileAsByteArray(TArray<uint8> ByteArray, FString MimeType, FString FileExtension)
 {
-	IPixelStreamingModule* Module = FPixelStreamingModule::GetModule();
+	IPixelStreamingModule* Module = UE::PixelStreaming::FPixelStreamingModule::GetModule();
 	if (Module)
 	{
 		Module->SendFileData(ByteArray, MimeType, FileExtension);
@@ -17,7 +17,7 @@ void UPixelStreamingBlueprints::SendFileAsByteArray(TArray<uint8> ByteArray, FSt
 
 void UPixelStreamingBlueprints::SendFile(FString FilePath, FString MimeType, FString FileExtension)
 {
-	IPixelStreamingModule* Module = FPixelStreamingModule::GetModule();
+	IPixelStreamingModule* Module = UE::PixelStreaming::FPixelStreamingModule::GetModule();
 	if (Module)
 	{
 		TArray<uint8> ByteData;
@@ -28,14 +28,14 @@ void UPixelStreamingBlueprints::SendFile(FString FilePath, FString MimeType, FSt
 		}
 		else
 		{
-			UE_LOG(PixelStreamer, Error, TEXT("FileHelper failed to load file data"));
+			UE_LOG(LogPixelStreaming, Error, TEXT("FileHelper failed to load file data"));
 		}
 	}
 }
 
 void UPixelStreamingBlueprints::FreezeFrame(UTexture2D* Texture)
 {
-	IPixelStreamingModule* Module = FPixelStreamingModule::GetModule();
+	IPixelStreamingModule* Module = UE::PixelStreaming::FPixelStreamingModule::GetModule();
 	if (Module)
 	{
 		Module->FreezeFrame(Texture);
@@ -44,14 +44,14 @@ void UPixelStreamingBlueprints::FreezeFrame(UTexture2D* Texture)
 
 void UPixelStreamingBlueprints::UnfreezeFrame()
 {
-	IPixelStreamingModule* Module = FPixelStreamingModule::GetModule();
+	IPixelStreamingModule* Module = UE::PixelStreaming::FPixelStreamingModule::GetModule();
 	if (Module)
 	{
 		Module->UnfreezeFrame();
 	}
 }
 
-UPixelStreamerDelegates* UPixelStreamingBlueprints::GetPixelStreamerDelegates()
+UPixelStreamingDelegates* UPixelStreamingBlueprints::GetPixelStreamingDelegates()
 {
-	return UPixelStreamerDelegates::GetPixelStreamerDelegates();
+	return UPixelStreamingDelegates::GetPixelStreamingDelegates();
 }
\ No newline at end of file
diff --git a/Engine/Plugins/Media/PixelStreaming/Source/PixelStreaming/Private/PixelStreamingBlueprints.h b/Engine/Plugins/Media/PixelStreaming/Source/PixelStreaming/Private/PixelStreamingBlueprints.h
--- a/Engine/Plugins/Media/PixelStreaming/Source/PixelStreaming/Private/PixelStreamingBlueprints.h
+++ b/Engine/Plugins/Media/PixelStreaming/Source/PixelStreaming/Private/PixelStreamingBlueprints.h
@@ -7,7 +7,7 @@
 #include "PixelStreamingModule.h"
 #include "Kismet/BlueprintFunctionLibrary.h"
 #include "Containers/Array.h"
-#include "PixelStreamerDelegates.h"
+#include "PixelStreamingDelegates.h"
 #include "PixelStreamingBlueprints.generated.h"
 
 UCLASS()
@@ -54,11 +54,11 @@ public:
 	UFUNCTION(BlueprintCallable, Category = "Pixel Streaming Freeze Frame")
 	static void UnfreezeFrame();
 
-	// PixelStreamerDelegates
+	// PixelStreamingDelegates
 	/**
 	 * Get the singleton. This allows application-specific blueprints to bind
 	 * to delegates of interest.
 	 */
-	UFUNCTION(BlueprintCallable, Category = "Pixel Streamer Delegates")
-	static UPixelStreamerDelegates* GetPixelStreamerDelegates();
+	UFUNCTION(BlueprintCallable, Category = "Pixel Streaming Delegates")
+	static UPixelStreamingDelegates* GetPixelStreamingDelegates();
 };
\ No newline at end of file
diff --git a/Engine/Plugins/Media/PixelStreaming/Source/PixelStreaming/Private/PixelStreamingInputComponent.cpp b/Engine/Plugins/Media/PixelStreaming/Source/PixelStreaming/Private/PixelStreamingInputComponent.cpp
new file mode 100644
--- /dev/null
+++ b/Engine/Plugins/Media/PixelStreaming/Source/PixelStreaming/Private/PixelStreamingInputComponent.cpp
@@ -0,0 +1,160 @@
+// Copyright Epic Games, Inc. All Rights Reserved.
+
+#include "PixelStreamingInputComponent.h"
+#include "IPixelStreamingModule.h"
+#include "PixelStreamingModule.h"
+#include "InputDevice.h"
+#include "Engine/Engine.h"
+#include "Engine/World.h"
+#include "Policies/CondensedJsonPrintPolicy.h"
+#include "Serialization/JsonReader.h"
+#include "Serialization/JsonSerializer.h"
+#include "GameFramework/GameUserSettings.h"
+#include "PixelStreamingPrivate.h"
+
+UPixelStreamingInput::UPixelStreamingInput(const FObjectInitializer& ObjectInitializer)
+	: Super(ObjectInitializer)
+	, PixelStreamingModule(UE::PixelStreaming::FPixelStreamingModule::GetModule())
+{
+	bAutoActivate = true;
+	PrimaryComponentTick.bCanEverTick = true;
+	SetComponentTickEnabled(true);
+}
+
+void UPixelStreamingInput::BeginPlay()
+{
+	Super::BeginPlay();
+
+	if (PixelStreamingModule)
+	{
+		// When this component is initializing it registers itself with the Pixel Streaming module.
+		PixelStreamingModule->AddInputComponent(this);
+	}
+	else
+	{
+		UE_LOG(LogPixelStreaming, Warning, TEXT("Pixel Streaming input component not added because Pixel Streaming module is not loaded. This is expected on dedicated servers."));
+	}
+}
+
+void UPixelStreamingInput::EndPlay(const EEndPlayReason::Type EndPlayReason)
+{
+	Super::EndPlay(EndPlayReason);
+
+	if (PixelStreamingModule)
+	{
+		// When this component is destructing it unregisters itself with the Pixel Streaming module.
+		PixelStreamingModule->RemoveInputComponent(this);
+	}
+	else
+	{
+		UE_LOG(LogPixelStreaming, Warning, TEXT("Pixel Streaming input component not removed because Pixel Streaming module is not loaded. This is expected on dedicated servers."));
+	}
+}
+
+bool UPixelStreamingInput::OnCommand(const FString& Descriptor)
+{
+	FString ConsoleCommand;
+	bool bSuccess = false;
+	ExtractJsonFromDescriptor(Descriptor, TEXT("ConsoleCommand"), ConsoleCommand, bSuccess);
+	if (bSuccess)
+	{
+		return GEngine->Exec(GEngine->GetWorld(), *ConsoleCommand);
+	}
+
+	FString WidthString;
+	FString HeightString;
+	ExtractJsonFromDescriptor(Descriptor, TEXT("Resolution.Width"), WidthString, bSuccess);
+	if (bSuccess)
+	{
+		ExtractJsonFromDescriptor(Descriptor, TEXT("Resolution.Height"), HeightString, bSuccess);
+
+		int Width = FCString::Atoi(*WidthString);
+		int Height = FCString::Atoi(*HeightString);
+
+		if (Width < 1 || Height < 1)
+		{
+			return false;
+		}
+
+		FString ChangeResCommand = FString::Printf(TEXT("r.SetRes %dx%d"), Width, Height);
+		return GEngine->Exec(GEngine->GetWorld(), *ChangeResCommand);
+	}
+
+	return false;
+}
+
+void UPixelStreamingInput::SendPixelStreamingResponse(const FString& Descriptor)
+{
+	if (PixelStreamingModule)
+	{
+		PixelStreamingModule->SendResponse(Descriptor);
+	}
+	else
+	{
+		UE_LOG(LogPixelStreaming, Warning, TEXT("Pixel Streaming input component skipped sending response. This is expected on dedicated servers."));
+	}
+}
+
+void UPixelStreamingInput::GetJsonStringValue(FString Descriptor, FString FieldName, FString& StringValue, bool& Success)
+{
+	ExtractJsonFromDescriptor(Descriptor, FieldName, StringValue, Success);
+}
+
+void UPixelStreamingInput::ExtractJsonFromDescriptor(FString Descriptor, FString FieldName, FString& StringValue, bool& Success)
+{
+	TSharedPtr<FJsonObject> JsonObject = MakeShareable(new FJsonObject);
+
+	TSharedRef<TJsonReader<>> JsonReader = TJsonReaderFactory<>::Create(Descriptor);
+	if (FJsonSerializer::Deserialize(JsonReader, JsonObject) && JsonObject.IsValid())
+	{
+		const TSharedPtr<FJsonObject>* JsonObjectPtr = &JsonObject;
+
+		if (FieldName.Contains(TEXT(".")))
+		{
+			TArray<FString> FieldComponents;
+			FieldName.ParseIntoArray(FieldComponents, TEXT("."));
+			FieldName = FieldComponents.Pop();
+
+			for (const FString& FieldComponent : FieldComponents)
+			{
+				if (!(*JsonObjectPtr)->TryGetObjectField(FieldComponent, JsonObjectPtr))
+				{
+					Success = false;
+					return;
+				}
+			}
+		}
+
+		Success = (*JsonObjectPtr)->TryGetStringField(FieldName, StringValue);
+	}
+	else
+	{
+		Success = false;
+	}
+}
+
+void UPixelStreamingInput::AddJsonStringValue(const FString& Descriptor, FString FieldName, FString StringValue, FString& NewDescriptor, bool& Success)
+{
+	ExtendJsonWithField(Descriptor, FieldName, StringValue, NewDescriptor, Success);
+}
+
+void UPixelStreamingInput::ExtendJsonWithField(const FString& Descriptor, FString FieldName, FString StringValue, FString& NewDescriptor, bool& Success)
+{
+	TSharedPtr<FJsonObject> JsonObject = MakeShareable(new FJsonObject);
+
+	if (!Descriptor.IsEmpty())
+	{
+		TSharedRef<TJsonReader<>> JsonReader = TJsonReaderFactory<>::Create(Descriptor);
+		if (!FJsonSerializer::Deserialize(JsonReader, JsonObject) || !JsonObject.IsValid())
+		{
+			Success = false;
+			return;
+		}
+	}
+
+	TSharedRef<FJsonValueString> JsonValueObject = MakeShareable(new FJsonValueString(StringValue));
+	JsonObject->SetField(FieldName, JsonValueObject);
+
+	TSharedRef<TJsonWriter<TCHAR, TCondensedJsonPrintPolicy<TCHAR>>> JsonWriter = TJsonWriterFactory<TCHAR, TCondensedJsonPrintPolicy<TCHAR>>::Create(&NewDescriptor);
+	Success = FJsonSerializer::Serialize(JsonObject.ToSharedRef(), JsonWriter);
+}
diff --git a/Engine/Plugins/Media/PixelStreaming/Source/PixelStreaming/Private/PixelStreamingModule.cpp b/Engine/Plugins/Media/PixelStreaming/Source/PixelStreaming/Private/PixelStreamingModule.cpp
--- a/Engine/Plugins/Media/PixelStreaming/Source/PixelStreaming/Private/PixelStreamingModule.cpp
+++ b/Engine/Plugins/Media/PixelStreaming/Source/PixelStreaming/Private/PixelStreamingModule.cpp
@@ -3,13 +3,13 @@
 #include "PixelStreamingModule.h"
 #include "Streamer.h"
 #include "InputDevice.h"
-#include "PixelStreamerInputComponent.h"
-#include "PixelStreamerDelegates.h"
+#include "PixelStreamingInputComponent.h"
+#include "PixelStreamingDelegates.h"
 #include "SignallingServerConnection.h"
-#include "PixelStreamingSettings.h"
+#include "Settings.h"
 #include "PixelStreamingPrivate.h"
 #include "PlayerSession.h"
-#include "PixelStreamingAudioSink.h"
+#include "AudioSink.h"
 #include "CoreMinimal.h"
 #include "Modules/ModuleManager.h"
 #include "UObject/UObjectIterator.h"
@@ -23,6 +23,14 @@
 	#include "CudaModule.h"
 #endif
 
+#if PLATFORM_WINDOWS
+	#include "Windows/AllowWindowsPlatformTypes.h"
+THIRD_PARTY_INCLUDES_START
+	#include <VersionHelpers.h>
+THIRD_PARTY_INCLUDES_END
+	#include "Windows/HideWindowsPlatformTypes.h"
+#endif
+
 #include "RenderingThread.h"
 #include "Misc/CommandLine.h"
 #include "Misc/Parse.h"
@@ -43,9 +51,9 @@
 	#include "DrawDebugHelpers.h"
 #endif
 
-DEFINE_LOG_CATEGORY(PixelStreaming);
+DEFINE_LOG_CATEGORY(LogPixelStreaming);
 
-IPixelStreamingModule* FPixelStreamingModule::PixelStreamingModule = nullptr;
+IPixelStreamingModule* UE::PixelStreaming::FPixelStreamingModule::PixelStreamingModule = nullptr;
 
 namespace
 {
@@ -57,7 +65,7 @@ namespace
 	bool LoadMediaFoundationDLLs()
 	{
 		// Ensure that all required modules are preloaded so they are not loaded just-in-time, causing a hitch.
-		if (IsWindows8Plus())
+		if (IsWindows8OrGreater())
 		{
 			return FPlatformProcess::GetDllHandle(TEXT("mf.dll"))
 				&& FPlatformProcess::GetDllHandle(TEXT("mfplat.dll"))
@@ -75,7 +83,7 @@ namespace
 #endif
 } // namespace
 
-void FPixelStreamingModule::InitStreamer()
+void UE::PixelStreaming::FPixelStreamingModule::InitStreamer()
 {
 	FString StreamerId;
 	FParse::Value(FCommandLine::Get(), TEXT("PixelStreamingID="), StreamerId);
@@ -88,14 +96,14 @@ void FPixelStreamingModule::InitStreamer()
 		uint16 SignallingServerPort = 8888;
 		if (!FParse::Value(FCommandLine::Get(), TEXT("PixelStreamingIP="), SignallingServerIP) || !FParse::Value(FCommandLine::Get(), TEXT("PixelStreamingPort="), SignallingServerPort))
 		{
-			UE_LOG(PixelStreamer, Log, TEXT("PixelStreaming is disabled, provide `PixelStreamingIP` and `PixelStreamingPort` cmd-args to enable it"));
+			UE_LOG(LogPixelStreaming, Log, TEXT("PixelStreaming is disabled, provide `PixelStreamingIP` and `PixelStreamingPort` cmd-args to enable it"));
 			return;
 		}
-		UE_LOG(PixelStreamer, Warning, TEXT("PixelStreamingIP and PixelStreamingPort are deprecated flags. Use PixelStreamingURL instead. eg. -PixelStreamingURL=ws://%s:%d"), *SignallingServerIP, SignallingServerPort);
+		UE_LOG(LogPixelStreaming, Warning, TEXT("PixelStreamingIP and PixelStreamingPort are deprecated flags. Use PixelStreamingURL instead. eg. -PixelStreamingURL=ws://%s:%d"), *SignallingServerIP, SignallingServerPort);
 		SignallingServerURL = FString::Printf(TEXT("ws://%s:%d"), *SignallingServerIP, SignallingServerPort);
 	}
 
-	UE_LOG(PixelStreamer, Log, TEXT("PixelStreaming endpoint ID: %s"), *StreamerId);
+	UE_LOG(LogPixelStreaming, Log, TEXT("PixelStreaming endpoint ID: %s"), *StreamerId);
 
 	if (GIsEditor)
 	{
@@ -103,7 +111,7 @@ void FPixelStreamingModule::InitStreamer()
 		FString ErrorString = TEXT("Pixel Streaming Plugin is not supported in editor, but it was explicitly enabled by command-line arguments. Please remove `PixelStreamingIP` and `PixelStreamingPort` args from editor command line.");
 		FText ErrorText = FText::FromString(ErrorString);
 		FMessageDialog::Open(EAppMsgType::Ok, ErrorText, &TitleText);
-		UE_LOG(PixelStreamer, Error, TEXT("%s"), *ErrorString);
+		UE_LOG(LogPixelStreaming, Error, TEXT("%s"), *ErrorString);
 		return;
 	}
 
@@ -111,7 +119,7 @@ void FPixelStreamingModule::InitStreamer()
 	// If not then we avoid setting up our delegates to prevent access to the
 	// plugin. Note that Pixel Streaming is not currently performed in the
 	// Editor.
-	if (!CheckPlatformCompatibility())
+	if (!IsPlatformCompatible())
 	{
 		return;
 	}
@@ -125,11 +133,11 @@ void FPixelStreamingModule::InitStreamer()
 	// TODO check if there is a better callback to attach so that we can use with editor
 	if (FSlateApplication::IsInitialized())
 	{
-		FSlateApplication::Get().GetRenderer()->OnBackBufferReadyToPresent().AddRaw(this, &FPixelStreamingModule::OnBackBufferReady_RenderThread);
+		FSlateApplication::Get().GetRenderer()->OnBackBufferReadyToPresent().AddRaw(this, &UE::PixelStreaming::FPixelStreamingModule::OnBackBufferReady_RenderThread);
 	}
 
-	FGameModeEvents::GameModePostLoginEvent.AddRaw(this, &FPixelStreamingModule::OnGameModePostLogin);
-	FGameModeEvents::GameModeLogoutEvent.AddRaw(this, &FPixelStreamingModule::OnGameModeLogout);
+	FGameModeEvents::GameModePostLoginEvent.AddRaw(this, &UE::PixelStreaming::FPixelStreamingModule::OnGameModePostLogin);
+	FGameModeEvents::GameModeLogoutEvent.AddRaw(this, &UE::PixelStreaming::FPixelStreamingModule::OnGameModeLogout);
 
 	IModularFeatures::Get().RegisterModularFeature(GetModularFeatureName(), this);
 
@@ -137,18 +145,18 @@ void FPixelStreamingModule::InitStreamer()
 
 	// Allow Pixel Streaming to broadcast to various delegates bound in the
 	// application-specific blueprint.
-	UPixelStreamerDelegates::CreateInstance();
+	UPixelStreamingDelegates::CreateInstance();
 
 	verify(FModuleManager::Get().LoadModule(FName("ImageWrapper")));
 
-	Streamer = MakeUnique<FStreamer>(SignallingServerURL, StreamerId);
+	Streamer = MakeUnique<UE::PixelStreaming::FStreamer>(SignallingServerURL, StreamerId);
 
 	// Streamer has been created, so module is now "ready" for external use.
 	ReadyEvent.Broadcast(*this);
 }
 
 /** IModuleInterface implementation */
-void FPixelStreamingModule::StartupModule()
+void UE::PixelStreaming::FPixelStreamingModule::StartupModule()
 {
 	// Pixel Streaming does not make sense without an RHI so we don't run in commandlets without one.
 	if (IsRunningCommandlet() && !IsAllowCommandletRendering())
@@ -157,22 +165,22 @@ void FPixelStreamingModule::StartupModule()
 	}
 
 	// Initialise all settings from command line args etc
-	PixelStreamingSettings::InitialiseSettings();
+	UE::PixelStreaming::Settings::InitialiseSettings();
 
 	// only D3D11/D3D12 is supported
 	if (GDynamicRHI == nullptr || !(GDynamicRHI->GetName() == FString(TEXT("D3D11")) || GDynamicRHI->GetName() == FString(TEXT("D3D12")) || GDynamicRHI->GetName() == FString(TEXT("Vulkan"))))
 	{
-		UE_LOG(PixelStreaming, Warning, TEXT("Only D3D11/D3D12/Vulkan Dynamic RHI is supported. Detected %s"), GDynamicRHI != nullptr ? GDynamicRHI->GetName() : TEXT("[null]"));
+		UE_LOG(LogPixelStreaming, Warning, TEXT("Only D3D11/D3D12/Vulkan Dynamic RHI is supported. Detected %s"), GDynamicRHI != nullptr ? GDynamicRHI->GetName() : TEXT("[null]"));
 		return;
 	}
 	else if (GDynamicRHI->GetName() == FString(TEXT("D3D11")) || GDynamicRHI->GetName() == FString(TEXT("D3D12")) || GDynamicRHI->GetName() == FString(TEXT("Vulkan")))
 	{
 		// By calling InitStreamer post engine init we can use pixel streaming in standalone editor mode
-		FCoreDelegates::OnFEngineLoopInitComplete.AddRaw(this, &FPixelStreamingModule::InitStreamer);
+		FCoreDelegates::OnFEngineLoopInitComplete.AddRaw(this, &UE::PixelStreaming::FPixelStreamingModule::InitStreamer);
 	}
 }
 
-void FPixelStreamingModule::ShutdownModule()
+void UE::PixelStreaming::FPixelStreamingModule::ShutdownModule()
 {
 	if (FSlateApplication::IsInitialized())
 	{
@@ -183,7 +191,7 @@ void FPixelStreamingModule::ShutdownModule()
 	IModularFeatures::Get().UnregisterModularFeature(GetModularFeatureName(), this);
 }
 
-IPixelStreamingModule* FPixelStreamingModule::GetModule()
+IPixelStreamingModule* UE::PixelStreaming::FPixelStreamingModule::GetModule()
 {
 	if (PixelStreamingModule)
 	{
@@ -197,42 +205,42 @@ IPixelStreamingModule* FPixelStreamingModule::GetModule()
 	return PixelStreamingModule;
 }
 
-bool FPixelStreamingModule::CheckPlatformCompatibility() const
+bool UE::PixelStreaming::FPixelStreamingModule::IsPlatformCompatible() const
 {
 	bool bCompatible = true;
 
 #if PLATFORM_WINDOWS || PLATFORM_XBOXONE
-	bool bWin8OrHigher = FPlatformMisc::VerifyWindowsVersion(6, 2);
+	bool bWin8OrHigher = IsWindows8OrGreater();
 	if (!bWin8OrHigher)
 	{
 		FString ErrorString(TEXT("Failed to initialize Pixel Streaming plugin because minimum requirement is Windows 8"));
 		FText ErrorText = FText::FromString(ErrorString);
 		FText TitleText = FText::FromString(TEXT("Pixel Streaming Plugin"));
 		FMessageDialog::Open(EAppMsgType::Ok, ErrorText, &TitleText);
-		UE_LOG(PixelStreamer, Error, TEXT("%s"), *ErrorString);
+		UE_LOG(LogPixelStreaming, Error, TEXT("%s"), *ErrorString);
 		bCompatible = false;
 	}
 #endif
 
-	if (!FStreamer::CheckPlatformCompatibility())
+	if (!UE::PixelStreaming::FStreamer::IsPlatformCompatible())
 	{
 		FText TitleText = FText::FromString(TEXT("Pixel Streaming Plugin"));
 		FString ErrorString = TEXT("No compatible GPU found, or failed to load their respective encoder libraries");
 		FText ErrorText = FText::FromString(ErrorString);
 		FMessageDialog::Open(EAppMsgType::Ok, ErrorText, &TitleText);
-		UE_LOG(PixelStreamer, Error, TEXT("%s"), *ErrorString);
+		UE_LOG(LogPixelStreaming, Error, TEXT("%s"), *ErrorString);
 		bCompatible = false;
 	}
 
 	return bCompatible;
 }
 
-void FPixelStreamingModule::UpdateViewport(FSceneViewport* Viewport)
+void UE::PixelStreaming::FPixelStreamingModule::UpdateViewport(FSceneViewport* Viewport)
 {
 	FRHIViewport* const ViewportRHI = Viewport->GetViewportRHI().GetReference();
 }
 
-void FPixelStreamingModule::OnBackBufferReady_RenderThread(SWindow& SlateWindow, const FTexture2DRHIRef& BackBuffer)
+void UE::PixelStreaming::FPixelStreamingModule::OnBackBufferReady_RenderThread(SWindow& SlateWindow, const FTexture2DRHIRef& BackBuffer)
 {
 	// enable streaming explicitly by providing `PixelStreamingIP` and `PixelStreamingPort` cmd-args
 	if (!Streamer)
@@ -263,48 +271,48 @@ void FPixelStreamingModule::OnBackBufferReady_RenderThread(SWindow& SlateWindow,
 	}
 }
 
-TSharedPtr<class IInputDevice> FPixelStreamingModule::CreateInputDevice(const TSharedRef<FGenericApplicationMessageHandler>& InMessageHandler)
+TSharedPtr<class IInputDevice> UE::PixelStreaming::FPixelStreamingModule::CreateInputDevice(const TSharedRef<FGenericApplicationMessageHandler>& InMessageHandler)
 {
-	InputDevice = MakeShareable(new FInputDevice(InMessageHandler));
+	InputDevice = MakeShareable(new UE::PixelStreaming::FInputDevice(InMessageHandler));
 	return InputDevice;
 }
 
-IPixelStreamingModule::FReadyEvent& FPixelStreamingModule::OnReady()
+IPixelStreamingModule::FReadyEvent& UE::PixelStreaming::FPixelStreamingModule::OnReady()
 {
 	return ReadyEvent;
 }
 
-bool FPixelStreamingModule::IsReady()
+bool UE::PixelStreaming::FPixelStreamingModule::IsReady()
 {
 	return Streamer.IsValid();
 }
 
-FInputDevice& FPixelStreamingModule::GetInputDevice()
+IInputDevice& UE::PixelStreaming::FPixelStreamingModule::GetInputDevice()
 {
 	return *InputDevice;
 }
 
-TSharedPtr<FInputDevice> FPixelStreamingModule::GetInputDevicePtr()
+TSharedPtr<UE::PixelStreaming::FInputDevice> UE::PixelStreaming::FPixelStreamingModule::GetInputDevicePtr()
 {
 	return InputDevice;
 }
 
-void FPixelStreamingModule::AddInputComponent(UPixelStreamerInputComponent* InInputComponent)
+void UE::PixelStreaming::FPixelStreamingModule::AddInputComponent(UPixelStreamingInput* InInputComponent)
 {
-	this->InputComponents.Add(InInputComponent);
+	InputComponents.Add(InInputComponent);
 }
 
-void FPixelStreamingModule::RemoveInputComponent(UPixelStreamerInputComponent* InInputComponent)
+void UE::PixelStreaming::FPixelStreamingModule::RemoveInputComponent(UPixelStreamingInput* InInputComponent)
 {
-	this->InputComponents.Remove(InInputComponent);
+	InputComponents.Remove(InInputComponent);
 }
 
-const TArray<UPixelStreamerInputComponent*> FPixelStreamingModule::GetInputComponents()
+const TArray<UPixelStreamingInput*> UE::PixelStreaming::FPixelStreamingModule::GetInputComponents()
 {
-	return this->InputComponents;
+	return InputComponents;
 }
 
-void FPixelStreamingModule::FreezeFrame(UTexture2D* Texture)
+void UE::PixelStreaming::FPixelStreamingModule::FreezeFrame(UTexture2D* Texture)
 {
 	if (Texture)
 	{
@@ -314,23 +322,23 @@ void FPixelStreamingModule::FreezeFrame(UTexture2D* Texture)
 			FTexture2DRHIRef Texture2DRHI = (Texture->GetResource() && Texture->GetResource()->TextureRHI) ? Texture->GetResource()->TextureRHI->GetTexture2D() : nullptr;
 			if (!Texture2DRHI)
 			{
-				UE_LOG(PixelStreamer, Error, TEXT("Attempting freeze frame with texture %s with no texture 2D RHI"), *Texture->GetName());
+				UE_LOG(LogPixelStreaming, Error, TEXT("Attempting freeze frame with texture %s with no texture 2D RHI"), *Texture->GetName());
 				return;
 			}
 			uint32 Width = Texture2DRHI->GetSizeX();
 			uint32 Height = Texture2DRHI->GetSizeY();
 
-			FTexture2DRHIRef DestTexture = CreateTexture(Width, Height);
+			FTexture2DRHIRef DestTexture = UE::PixelStreaming::CreateTexture(Width, Height);
 
 			FGPUFenceRHIRef CopyFence = GDynamicRHI->RHICreateGPUFence(*FString::Printf(TEXT("FreezeFrameFence")));
 
 			// Copy freeze frame texture to empty texture
-			CopyTexture(Texture2DRHI, DestTexture, CopyFence);
+			UE::PixelStreaming::CopyTexture(Texture2DRHI, DestTexture, CopyFence);
 
 			TArray<FColor> Data;
 			FIntRect Rect(0, 0, Width, Height);
 			RHICmdList.ReadSurfaceData(DestTexture, Rect, Data, FReadSurfaceDataFlags());
-			this->SendJpeg(MoveTemp(Data), Rect);
+			SendJpeg(MoveTemp(Data), Rect);
 		});
 	}
 	else
@@ -344,7 +352,7 @@ void FPixelStreamingModule::FreezeFrame(UTexture2D* Texture)
 	bFrozen = true;
 }
 
-void FPixelStreamingModule::UnfreezeFrame()
+void UE::PixelStreaming::FPixelStreamingModule::UnfreezeFrame()
 {
 	if (!Streamer.IsValid())
 	{
@@ -356,54 +364,54 @@ void FPixelStreamingModule::UnfreezeFrame()
 	// Resume streaming.
 	bFrozen = false;
 }
-void FPixelStreamingModule::AddPlayerConfig(TSharedRef<FJsonObject>& JsonObject)
+void UE::PixelStreaming::FPixelStreamingModule::AddPlayerConfig(TSharedRef<FJsonObject>& JsonObject)
 {
 	checkf(InputDevice.IsValid(), TEXT("No Input Device available when populating Player Config"));
 
 	JsonObject->SetBoolField(TEXT("FakingTouchEvents"), InputDevice->IsFakingTouchEvents());
 
 	FString PixelStreamingControlScheme;
-	if (PixelStreamingSettings::GetControlScheme(PixelStreamingControlScheme))
+	if (UE::PixelStreaming::Settings::GetControlScheme(PixelStreamingControlScheme))
 	{
 		JsonObject->SetStringField(TEXT("ControlScheme"), PixelStreamingControlScheme);
 	}
 
 	float PixelStreamingFastPan;
-	if (PixelStreamingSettings::GetFastPan(PixelStreamingFastPan))
+	if (UE::PixelStreaming::Settings::GetFastPan(PixelStreamingFastPan))
 	{
 		JsonObject->SetNumberField(TEXT("FastPan"), PixelStreamingFastPan);
 	}
 }
 
-void FPixelStreamingModule::SendResponse(const FString& Descriptor)
+void UE::PixelStreaming::FPixelStreamingModule::SendResponse(const FString& Descriptor)
 {
 	if (!Streamer.IsValid())
 	{
 		return;
 	}
 
-	Streamer->SendPlayerMessage(PixelStreamingProtocol::EToPlayerMsg::Response, Descriptor);
+	Streamer->SendPlayerMessage(UE::PixelStreaming::Protocol::EToPlayerMsg::Response, Descriptor);
 }
 
-void FPixelStreamingModule::SendCommand(const FString& Descriptor)
+void UE::PixelStreaming::FPixelStreamingModule::SendCommand(const FString& Descriptor)
 {
 	if (!Streamer.IsValid())
 	{
 		return;
 	}
 
-	Streamer->SendPlayerMessage(PixelStreamingProtocol::EToPlayerMsg::Command, Descriptor);
+	Streamer->SendPlayerMessage(UE::PixelStreaming::Protocol::EToPlayerMsg::Command, Descriptor);
 }
 
-void FPixelStreamingModule::OnGameModePostLogin(AGameModeBase* GameMode, APlayerController* NewPlayer)
+void UE::PixelStreaming::FPixelStreamingModule::OnGameModePostLogin(AGameModeBase* GameMode, APlayerController* NewPlayer)
 {
 }
 
-void FPixelStreamingModule::OnGameModeLogout(AGameModeBase* GameMode, AController* Exiting)
+void UE::PixelStreaming::FPixelStreamingModule::OnGameModeLogout(AGameModeBase* GameMode, AController* Exiting)
 {
 }
 
-void FPixelStreamingModule::SendJpeg(TArray<FColor> RawData, const FIntRect& Rect)
+void UE::PixelStreaming::FPixelStreamingModule::SendJpeg(TArray<FColor> RawData, const FIntRect& Rect)
 {
 	if (!Streamer.IsValid())
 	{
@@ -416,81 +424,60 @@ void FPixelStreamingModule::SendJpeg(TArray<FColor> RawData, const FIntRect& Rec
 	if (bSuccess)
 	{
 		// Compress to a JPEG of the maximum possible quality.
-		int32 Quality = PixelStreamingSettings::CVarPixelStreamingFreezeFrameQuality.GetValueOnAnyThread();
+		int32 Quality = UE::PixelStreaming::Settings::CVarPixelStreamingFreezeFrameQuality.GetValueOnAnyThread();
 		const TArray64<uint8>& JpegBytes = ImageWrapper->GetCompressed(Quality);
 		Streamer->SendFreezeFrame(JpegBytes);
 	}
 	else
 	{
-		UE_LOG(PixelStreamer, Error, TEXT("JPEG image wrapper failed to accept frame data"));
+		UE_LOG(LogPixelStreaming, Error, TEXT("JPEG image wrapper failed to accept frame data"));
 	}
 }
 
-void FPixelStreamingModule::SendFileData(TArray<uint8>& ByteData, FString& MimeType, FString& FileExtension)
+void UE::PixelStreaming::FPixelStreamingModule::SendFileData(TArray<uint8>& ByteData, FString& MimeType, FString& FileExtension)
 {
 	Streamer->SendFileData(ByteData, MimeType, FileExtension);
 }
 
-bool FPixelStreamingModule::IsTickableWhenPaused() const
+bool UE::PixelStreaming::FPixelStreamingModule::IsTickableWhenPaused() const
 {
 	return true;
 }
 
-bool FPixelStreamingModule::IsTickableInEditor() const
+bool UE::PixelStreaming::FPixelStreamingModule::IsTickableInEditor() const
 {
 	return true;
 }
 
-void FPixelStreamingModule::Tick(float DeltaTime)
+void UE::PixelStreaming::FPixelStreamingModule::Tick(float DeltaTime)
 {
 }
 
-TStatId FPixelStreamingModule::GetStatId() const
+TStatId UE::PixelStreaming::FPixelStreamingModule::GetStatId() const
 {
 	RETURN_QUICK_DECLARE_CYCLE_STAT(FPixelStreamingModule, STATGROUP_Tickables);
 }
 
-IPixelStreamingAudioSink* FPixelStreamingModule::GetPeerAudioSink(FPlayerId PlayerId)
+IPixelStreamingAudioSink* UE::PixelStreaming::FPixelStreamingModule::GetPeerAudioSink(FPixelStreamingPlayerId PlayerId)
 {
-	if (!this->Streamer.IsValid())
+	if (!Streamer.IsValid())
 	{
-		UE_LOG(PixelStreamer, Error, TEXT("Cannot get audio sink when streamer does not yet exist."));
+		UE_LOG(LogPixelStreaming, Error, TEXT("Cannot get audio sink when streamer does not yet exist."));
 		return nullptr;
 	}
 
-	return this->Streamer->GetAudioSink(PlayerId);
+	return Streamer->GetAudioSink(PlayerId);
 }
 
-IPixelStreamingAudioSink* FPixelStreamingModule::GetUnlistenedAudioSink()
+IPixelStreamingAudioSink* UE::PixelStreaming::FPixelStreamingModule::GetUnlistenedAudioSink()
 {
-	if (!this->Streamer.IsValid())
+	if (!Streamer.IsValid())
 	{
-		UE_LOG(PixelStreamer, Error, TEXT("Cannot get audio sink when streamer does not yet exist."));
+		UE_LOG(LogPixelStreaming, Error, TEXT("Cannot get audio sink when streamer does not yet exist."));
 		return nullptr;
 	}
 
-	return this->Streamer->GetUnlistenedAudioSink();
-}
-
-void FPixelStreamingModule::AddAnyStatChangedCallback(TWeakPtr<IPixelStreamingStatsConsumer> Callback)
-{
-	if (!this->Streamer.IsValid())
-	{
-		UE_LOG(PixelStreamer, Error, TEXT("Cannot add stat callback when streamer does not yet exist."));
-		return;
-	}
-
-	this->Streamer->AddAnyStatChangedCallback(Callback);
-}
-
-void FPixelStreamingModule::RemoveAnyStatChangedCallback(TWeakPtr<IPixelStreamingStatsConsumer> Callback)
-{
-	if (!this->Streamer.IsValid())
-	{
-		UE_LOG(PixelStreamer, Error, TEXT("Cannot remove stat callback when streamer does not yet exist."));
-		return;
-	}
-	this->Streamer->RemoveAnyStatChangedCallback(Callback);
+	return Streamer->GetUnlistenedAudioSink();
 }
 
-IMPLEMENT_MODULE(FPixelStreamingModule, PixelStreaming)
+IMPLEMENT_MODULE(UE::PixelStreaming::FPixelStreamingModule, PixelStreaming)
diff --git a/Engine/Plugins/Media/PixelStreaming/Source/PixelStreaming/Private/PixelStreamingModule.h b/Engine/Plugins/Media/PixelStreaming/Source/PixelStreaming/Private/PixelStreamingModule.h
--- a/Engine/Plugins/Media/PixelStreaming/Source/PixelStreaming/Private/PixelStreamingModule.h
+++ b/Engine/Plugins/Media/PixelStreaming/Source/PixelStreaming/Private/PixelStreamingModule.h
@@ -5,80 +5,86 @@
 #include "IPixelStreamingModule.h"
 #include "RHI.h"
 #include "Tickable.h"
+#include "InputDevice.h"
 
 class AController;
 class AGameModeBase;
 class APlayerController;
 class FSceneViewport;
-class FStreamer;
-class UPixelStreamerInputComponent;
+class UPixelStreamingInput;
 class SWindow;
 
-/**
- * This plugin allows the back buffer to be sent as a compressed video across
- * a network.
- */
-class FPixelStreamingModule : public IPixelStreamingModule, public FTickableGameObject
+namespace UE
 {
-public:
-	static IPixelStreamingModule* GetModule();
+	namespace PixelStreaming
+	{
+		class FStreamer;
 
-private:
-	/** IModuleInterface implementation */
-	void StartupModule() override;
-	void ShutdownModule() override;
+		/**
+		 * This plugin allows the back buffer to be sent as a compressed video across
+		 * a network.
+		 */
+		class FPixelStreamingModule : public IPixelStreamingModule, public FTickableGameObject
+		{
+		public:
+			static IPixelStreamingModule* GetModule();
 
-	TSharedPtr<class IInputDevice> CreateInputDevice(const TSharedRef<FGenericApplicationMessageHandler>& InMessageHandler) override;
+		private:
+			/** IModuleInterface implementation */
+			void StartupModule() override;
+			void ShutdownModule() override;
 
-	/** IPixelStreamingModule implementation */
-	IPixelStreamingModule::FReadyEvent& OnReady() override;
-	bool IsReady() override;
-	FInputDevice& GetInputDevice() override;
-	void AddPlayerConfig(TSharedRef<FJsonObject>& JsonObject) override;
-	void SendResponse(const FString& Descriptor) override;
-	void SendCommand(const FString& Descriptor) override;
+			TSharedPtr<class IInputDevice> CreateInputDevice(const TSharedRef<FGenericApplicationMessageHandler>& InMessageHandler) override;
 
-	/**
-	 * Returns a shared pointer to the device which handles pixel streaming
-	 * input.
-	 * @return The shared pointer to the input device.
-	 */
-	TSharedPtr<FInputDevice> GetInputDevicePtr();
-	void AddInputComponent(UPixelStreamerInputComponent* InInputComponent) override;
-	void RemoveInputComponent(UPixelStreamerInputComponent* InInputComponent) override;
-	const TArray<UPixelStreamerInputComponent*> GetInputComponents() override;
+			/** IPixelStreamingModule implementation */
+			IPixelStreamingModule::FReadyEvent& OnReady() override;
+			bool IsReady() override;
+			IInputDevice& GetInputDevice() override;
+			void AddPlayerConfig(TSharedRef<FJsonObject>& JsonObject) override;
+			void SendResponse(const FString& Descriptor) override;
+			void SendCommand(const FString& Descriptor) override;
 
-	void FreezeFrame(UTexture2D* Texture) override;
-	void UnfreezeFrame() override;
-	IPixelStreamingAudioSink* GetPeerAudioSink(FPlayerId PlayerId) override;
-	IPixelStreamingAudioSink* GetUnlistenedAudioSink() override;
-	void AddAnyStatChangedCallback(TWeakPtr<IPixelStreamingStatsConsumer> Callback) override;
-	void RemoveAnyStatChangedCallback(TWeakPtr<IPixelStreamingStatsConsumer> Callback) override;
-	/** End IPixelStreamingModule implementation */
+			/**
+			 * Returns a shared pointer to the device which handles pixel streaming
+			 * input.
+			 * @return The shared pointer to the input device.
+			 */
+			TSharedPtr<FInputDevice> GetInputDevicePtr();
+			void AddInputComponent(UPixelStreamingInput* InInputComponent) override;
+			void RemoveInputComponent(UPixelStreamingInput* InInputComponent) override;
+			const TArray<UPixelStreamingInput*> GetInputComponents() override;
 
-	// FTickableGameObject
-	bool IsTickableWhenPaused() const override;
-	bool IsTickableInEditor() const override;
-	void Tick(float DeltaTime) override;
-	TStatId GetStatId() const override;
+			void FreezeFrame(UTexture2D* Texture) override;
+			void UnfreezeFrame() override;
+			IPixelStreamingAudioSink* GetPeerAudioSink(FPixelStreamingPlayerId PlayerId) override;
+			IPixelStreamingAudioSink* GetUnlistenedAudioSink() override;
+			/** End IPixelStreamingModule implementation */
 
-	bool CheckPlatformCompatibility() const;
-	void UpdateViewport(FSceneViewport* Viewport);
-	void OnBackBufferReady_RenderThread(SWindow& SlateWindow, const FTexture2DRHIRef& BackBuffer);
-	void OnGameModePostLogin(AGameModeBase* GameMode, APlayerController* NewPlayer);
-	void OnGameModeLogout(AGameModeBase* GameMode, AController* Exiting);
-	void SendJpeg(TArray<FColor> RawData, const FIntRect& Rect);
-	void SendFileData(TArray<uint8>& ByteData, FString& MimeType, FString& FileExtension);
+			// FTickableGameObject
+			bool IsTickableWhenPaused() const override;
+			bool IsTickableInEditor() const override;
+			void Tick(float DeltaTime) override;
+			TStatId GetStatId() const override;
 
-	void InitStreamer();
+			bool IsPlatformCompatible() const;
+			void UpdateViewport(FSceneViewport* Viewport);
+			void OnBackBufferReady_RenderThread(SWindow& SlateWindow, const FTexture2DRHIRef& BackBuffer);
+			void OnGameModePostLogin(AGameModeBase* GameMode, APlayerController* NewPlayer);
+			void OnGameModeLogout(AGameModeBase* GameMode, AController* Exiting);
+			void SendJpeg(TArray<FColor> RawData, const FIntRect& Rect);
+			void SendFileData(TArray<uint8>& ByteData, FString& MimeType, FString& FileExtension);
 
-private:
-	IPixelStreamingModule::FReadyEvent ReadyEvent;
-	TUniquePtr<FStreamer> Streamer;
-	TSharedPtr<FInputDevice> InputDevice;
-	TArray<UPixelStreamerInputComponent*> InputComponents;
-	bool bFrozen = false;
-	bool bCaptureNextBackBufferAndStream = false;
-	double LastVideoEncoderQPReportTime = 0;
-	static IPixelStreamingModule* PixelStreamingModule;
-};
+			void InitStreamer();
+
+		private:
+			IPixelStreamingModule::FReadyEvent ReadyEvent;
+			TUniquePtr<FStreamer> Streamer;
+			TSharedPtr<FInputDevice> InputDevice;
+			TArray<UPixelStreamingInput*> InputComponents;
+			bool bFrozen = false;
+			bool bCaptureNextBackBufferAndStream = false;
+			double LastVideoEncoderQPReportTime = 0;
+			static IPixelStreamingModule* PixelStreamingModule;
+		};
+	} // namespace PixelStreaming
+} // namespace UE
diff --git a/Engine/Plugins/Media/PixelStreaming/Source/PixelStreaming/Private/PixelStreamingPrivate.h b/Engine/Plugins/Media/PixelStreaming/Source/PixelStreaming/Private/PixelStreamingPrivate.h
--- a/Engine/Plugins/Media/PixelStreaming/Source/PixelStreaming/Private/PixelStreamingPrivate.h
+++ b/Engine/Plugins/Media/PixelStreaming/Source/PixelStreaming/Private/PixelStreamingPrivate.h
@@ -5,6 +5,4 @@
 #include "HAL/IConsoleManager.h"
 #include "Logging/LogMacros.h"
 
-DECLARE_LOG_CATEGORY_EXTERN(PixelStreaming, Log, All);
-DECLARE_LOG_CATEGORY_EXTERN(PixelStreamer, Log, All);
-DECLARE_LOG_CATEGORY_EXTERN(PixelPlayer, Log, All);
\ No newline at end of file
+DECLARE_LOG_CATEGORY_EXTERN(LogPixelStreaming, Log, All);
\ No newline at end of file
diff --git a/Engine/Plugins/Media/PixelStreaming/Source/PixelStreaming/Private/PlayerSession.cpp b/Engine/Plugins/Media/PixelStreaming/Source/PixelStreaming/Private/PlayerSession.cpp
--- a/Engine/Plugins/Media/PixelStreaming/Source/PixelStreaming/Private/PlayerSession.cpp
+++ b/Engine/Plugins/Media/PixelStreaming/Source/PixelStreaming/Private/PlayerSession.cpp
@@ -1,48 +1,49 @@
 // Copyright Epic Games, Inc. All Rights Reserved.
 
 #include "PlayerSession.h"
-#include "VideoEncoder.h"
+#include "ToStringExtensions.h"
+#include "VideoEncoderRTC.h"
 #include "VideoEncoderFactory.h"
-#include "PixelStreamingEncoderFactory.h"
-#include "PixelStreamingVideoEncoder.h"
-#include "PixelStreamingSettings.h"
+#include "EncoderFactory.h"
+#include "VideoEncoder.h"
+#include "Settings.h"
 #include "WebRTCIncludes.h"
 #include "Containers/UnrealString.h"
-#include "PixelStreamingAudioSink.h"
+#include "AudioSink.h"
 #include "GenericPlatform/GenericPlatformMath.h"
 #include "SignallingServerConnection.h"
 #include "Async/Async.h"
-#include "WebRtcObservers.h"
-#include "PixelStreamingStats.h"
+#include "SessionDescriptionObservers.h"
+#include "Stats.h"
 #include "PixelStreamingStatNames.h"
 
 class QPConsumer : public IPixelStreamingStatsConsumer
 {
 public:
-	QPConsumer(IPixelStreamingSessions* InSessions)
+	QPConsumer(UE::PixelStreaming::IPixelStreamingSessions* InSessions)
 		: Sessions(InSessions) {}
 
 	QPConsumer(QPConsumer&& Other)
 		: Sessions(Other.Sessions) {}
 
 	// Begin IPixelStreamingStatsConsumer
-	void ConsumeStat(FPlayerId PlayerId, FName StatName, float StatValue) override { Sessions->SendLatestQP(PlayerId, (int)StatValue); }
+	void ConsumeStat(FPixelStreamingPlayerId PlayerId, FName StatName, float StatValue) override { Sessions->SendLatestQP(PlayerId, (int)StatValue); }
 	// End IPixelStreamingStatsConsumer
 private:
-	IPixelStreamingSessions* Sessions;
+	UE::PixelStreaming::IPixelStreamingSessions* Sessions;
 };
 
-FPlayerSession::FPlayerSession(IPixelStreamingSessions* InSessions, FSignallingServerConnection* InSignallingServerConnection, FPlayerId InPlayerId)
+UE::PixelStreaming::FPlayerSession::FPlayerSession(UE::PixelStreaming::IPixelStreamingSessions* InSessions, UE::PixelStreaming::FSignallingServerConnection* InSignallingServerConnection, FPixelStreamingPlayerId InPlayerId)
 	: SignallingServerConnection(InSignallingServerConnection)
 	, PlayerId(InPlayerId)
 	, DataChannelObserver(InSessions, InPlayerId)
-	, WebRTCStatsCallback(new rtc::RefCountedObject<FPixelStreamingRTCStatsCollector>(InPlayerId))
+	, WebRTCStatsCallback(new rtc::RefCountedObject<UE::PixelStreaming::FRTCStatsCollector>(InPlayerId))
 	, QPReporter(MakeShared<QPConsumer>(InSessions))
 {
-	UE_LOG(PixelStreamer, Log, TEXT("%s: PlayerId=%s"), TEXT("Created FPlayerSession::FPlayerSession"), *PlayerId);
+	UE_LOG(LogPixelStreaming, Log, TEXT("%s: PlayerId=%s"), TEXT("Created UE::PixelStreaming::FPlayerSession::FPlayerSession"), *PlayerId);
 
 	// Listen for changes on QP for this peer and when there is a change send that QP over the data channel to the browser
-	FPixelStreamingStats* Stats = FPixelStreamingStats::Get();
+	UE::PixelStreaming::FStats* Stats = UE::PixelStreaming::FStats::Get();
 	if (Stats)
 	{
 		Stats->AddOnPeerStatChangedCallback(PlayerId,
@@ -51,27 +52,27 @@ FPlayerSession::FPlayerSession(IPixelStreamingSessions* InSessions, FSignallingS
 	}
 }
 
-FPlayerSession::~FPlayerSession()
+UE::PixelStreaming::FPlayerSession::~FPlayerSession()
 {
-	UE_LOG(PixelStreamer, Log, TEXT("%s: PlayerId=%s"), TEXT("FPlayerSession::~FPlayerSession"), *PlayerId);
+	UE_LOG(LogPixelStreaming, Log, TEXT("%s: PlayerId=%s"), TEXT("UE::PixelStreaming::FPlayerSession::~FPlayerSession"), *PlayerId);
 
-	this->DataChannelObserver.Unregister();
-	this->DataChannel = nullptr;
+	DataChannelObserver.Unregister();
+	DataChannel = nullptr;
 
-	if (this->PeerConnection)
+	if (PeerConnection)
 	{
-		this->PeerConnection->Close();
-		this->PeerConnection = nullptr;
+		PeerConnection->Close();
+		PeerConnection = nullptr;
 	}
 
-	FPixelStreamingStats* Stats = FPixelStreamingStats::Get();
+	UE::PixelStreaming::FStats* Stats = UE::PixelStreaming::FStats::Get();
 	if (Stats)
 	{
 		Stats->RemovePeersStats(PlayerId);
 	}
 }
 
-void FPlayerSession::PollWebRTCStats() const
+void UE::PixelStreaming::FPlayerSession::PollWebRTCStats() const
 {
 	check(PeerConnection);
 	//PeerConnection->GetStats(WebRTCStatsCallback);
@@ -85,75 +86,80 @@ void FPlayerSession::PollWebRTCStats() const
 	}
 }
 
-FPixelStreamingDataChannelObserver& FPlayerSession::GetDataChannelObserver()
+UE::PixelStreaming::FDataChannelObserver& UE::PixelStreaming::FPlayerSession::GetDataChannelObserver()
 {
-	return this->DataChannelObserver;
+	return DataChannelObserver;
 }
 
-webrtc::PeerConnectionInterface& FPlayerSession::GetPeerConnection()
+webrtc::PeerConnectionInterface& UE::PixelStreaming::FPlayerSession::GetPeerConnection()
 {
 	check(PeerConnection);
 	return *PeerConnection;
 }
 
-void FPlayerSession::SetPeerConnection(const rtc::scoped_refptr<webrtc::PeerConnectionInterface>& InPeerConnection)
+void UE::PixelStreaming::FPlayerSession::SetPeerConnection(const rtc::scoped_refptr<webrtc::PeerConnectionInterface>& InPeerConnection)
 {
 	PeerConnection = InPeerConnection;
 }
 
-void FPlayerSession::SetDataChannel(const rtc::scoped_refptr<webrtc::DataChannelInterface>& InDataChannel)
+void UE::PixelStreaming::FPlayerSession::SetDataChannel(const rtc::scoped_refptr<webrtc::DataChannelInterface>& InDataChannel)
 {
 	DataChannel = InDataChannel;
 }
 
-FPlayerId FPlayerSession::GetPlayerId() const
+void UE::PixelStreaming::FPlayerSession::SetVideoSource(const rtc::scoped_refptr<webrtc::VideoTrackSourceInterface>& InVideoSource)
+{
+	VideoSource = InVideoSource;
+}
+
+FPixelStreamingPlayerId UE::PixelStreaming::FPlayerSession::GetPlayerId() const
 {
 	return PlayerId;
 }
 
-void FPlayerSession::OnAnswer(FString Sdp)
+void UE::PixelStreaming::FPlayerSession::OnAnswer(FString Sdp)
 {
 	webrtc::SdpParseError Error;
 	std::unique_ptr<webrtc::SessionDescriptionInterface> SessionDesc = webrtc::CreateSessionDescription(webrtc::SdpType::kAnswer, to_string(Sdp), &Error);
 	if (!SessionDesc)
 	{
-		UE_LOG(PixelStreamer, Error, TEXT("Failed to parse answer's SDP\n%s"), *Sdp);
+		UE_LOG(LogPixelStreaming, Error, TEXT("Failed to parse answer's SDP\n%s"), *Sdp);
 		return;
 	}
 
-	FSetSessionDescriptionObserver* SetRemoteDescriptionObserver = FSetSessionDescriptionObserver::Create(
+	UE::PixelStreaming::FSetSessionDescriptionObserver* SetRemoteDescriptionObserver = UE::PixelStreaming::FSetSessionDescriptionObserver::Create(
 		[]() // on success
 		{
 		},
 		[](const FString& Error) // on failure
 		{
-			UE_LOG(PixelStreamer, Error, TEXT("Failed to set remote description: %s"), *Error);
+			UE_LOG(LogPixelStreaming, Error, TEXT("Failed to set remote description: %s"), *Error);
 		});
 
 	PeerConnection->SetRemoteDescription(SetRemoteDescriptionObserver, SessionDesc.release());
 }
 
-void FPlayerSession::OnRemoteIceCandidate(const std::string& SdpMid, int SdpMLineIndex, const std::string& Sdp)
+void UE::PixelStreaming::FPlayerSession::OnRemoteIceCandidate(const std::string& SdpMid, int SdpMLineIndex, const std::string& Sdp)
 {
 	webrtc::SdpParseError Error;
 	std::unique_ptr<webrtc::IceCandidateInterface> Candidate(webrtc::CreateIceCandidate(SdpMid, SdpMLineIndex, Sdp, &Error));
 	if (Candidate)
 	{
-		PeerConnection->AddIceCandidate(std::move(Candidate), [&](webrtc::RTCError error) {
+		PeerConnection->AddIceCandidate(std::move(Candidate), [this](webrtc::RTCError error) {
 			if (!error.ok())
 			{
-				UE_LOG(PixelStreamer, Error, TEXT("AddIceCandidate failed (%s): %S"), *PlayerId, error.message());
+				UE_LOG(LogPixelStreaming, Error, TEXT("AddIceCandidate failed (%s): %S"), *PlayerId, error.message());
 			}
 		});
 	}
 	else
 	{
-		UE_LOG(PixelStreamer, Error, TEXT("Could not create ice candidate for player %s"), *this->PlayerId);
-		UE_LOG(PixelStreamer, Error, TEXT("Bad sdp at line %s | Description: %s"), *FString(Error.line.c_str()), *FString(Error.description.c_str()));
+		UE_LOG(LogPixelStreaming, Error, TEXT("Could not create ice candidate for player %s"), *PlayerId);
+		UE_LOG(LogPixelStreaming, Error, TEXT("Bad sdp at line %s | Description: %s"), *FString(Error.line.c_str()), *FString(Error.description.c_str()));
 	}
 }
 
-void FPlayerSession::DisconnectPlayer(const FString& Reason)
+void UE::PixelStreaming::FPlayerSession::DisconnectPlayer(const FString& Reason)
 {
 	if (bDisconnecting)
 	{
@@ -161,89 +167,142 @@ void FPlayerSession::DisconnectPlayer(const FString& Reason)
 	}
 
 	bDisconnecting = true;
-	this->SignallingServerConnection->SendDisconnectPlayer(PlayerId, Reason);
+	SignallingServerConnection->SendDisconnectPlayer(PlayerId, Reason);
 }
 
-FPixelStreamingAudioSink& FPlayerSession::GetAudioSink()
+UE::PixelStreaming::FAudioSink& UE::PixelStreaming::FPlayerSession::GetAudioSink()
 {
-	return this->AudioSink;
+	return AudioSink;
 }
 
-bool FPlayerSession::SendMessage(PixelStreamingProtocol::EToPlayerMsg Type, const FString& Descriptor) const
+bool UE::PixelStreaming::FPlayerSession::SendMessage(UE::PixelStreaming::Protocol::EToPlayerMsg Type, const FString& Descriptor) const
 {
 	if (!DataChannel)
 	{
 		return false;
 	}
 
-	const uint8 MessageType = static_cast<uint8>(Type);
+	const uint8 MessageType		= static_cast<uint8>(Type);
 	const size_t DescriptorSize = Descriptor.Len() * sizeof(TCHAR);
 
 	rtc::CopyOnWriteBuffer Buffer(sizeof(MessageType) + DescriptorSize);
 
 	size_t Pos = 0;
-	Pos = SerializeToBuffer(Buffer, Pos, &MessageType, sizeof(MessageType));
-	Pos = SerializeToBuffer(Buffer, Pos, *Descriptor, DescriptorSize);
+	Pos		   = SerializeToBuffer(Buffer, Pos, &MessageType, sizeof(MessageType));
+	Pos		   = SerializeToBuffer(Buffer, Pos, *Descriptor, DescriptorSize);
 
 	return DataChannel->Send(webrtc::DataBuffer(Buffer, true));
 }
 
-void FPlayerSession::SendQualityControlStatus(bool bIsQualityController) const
+void UE::PixelStreaming::FPlayerSession::SendQualityControlStatus(bool bIsQualityController) const
 {
+	UE::PixelStreaming::FStats::Get()->StorePeerStat(PlayerId, UE::PixelStreaming::FStatData(PixelStreamingStatNames::QualityController, bIsQualityController ? 1.0 : 0.0, 0));
+
 	if (!DataChannel)
 	{
 		return;
 	}
 
-	const uint8 MessageType = static_cast<uint8>(PixelStreamingProtocol::EToPlayerMsg::QualityControlOwnership);
+	const uint8 MessageType		= static_cast<uint8>(UE::PixelStreaming::Protocol::EToPlayerMsg::QualityControlOwnership);
 	const uint8 ControlsQuality = bIsQualityController ? 1 : 0;
 
 	rtc::CopyOnWriteBuffer Buffer(sizeof(MessageType) + sizeof(ControlsQuality));
 
 	size_t Pos = 0;
-	Pos = SerializeToBuffer(Buffer, Pos, &MessageType, sizeof(MessageType));
-	Pos = SerializeToBuffer(Buffer, Pos, &ControlsQuality, sizeof(ControlsQuality));
+	Pos		   = SerializeToBuffer(Buffer, Pos, &MessageType, sizeof(MessageType));
+	Pos		   = SerializeToBuffer(Buffer, Pos, &ControlsQuality, sizeof(ControlsQuality));
 
 	if (!DataChannel->Send(webrtc::DataBuffer(Buffer, true)))
 	{
-		UE_LOG(PixelStreamer, Error, TEXT("failed to send quality control status"));
+		UE_LOG(LogPixelStreaming, Error, TEXT("failed to send quality control status"));
 	}
 }
 
-void FPlayerSession::SendFreezeFrame(const TArray64<uint8>& JpegBytes) const
+void UE::PixelStreaming::FPlayerSession::SendFreezeFrame(const TArray64<uint8>& JpegBytes) const
 {
-	SendArbitraryData(this->DataChannel, static_cast<TArray<uint8>>(JpegBytes), static_cast<uint8>(PixelStreamingProtocol::EToPlayerMsg::FreezeFrame));
+	UE::PixelStreaming::FPlayerSession::SendArbitraryData(static_cast<TArray<uint8>>(JpegBytes), static_cast<uint8>(UE::PixelStreaming::Protocol::EToPlayerMsg::FreezeFrame));
 }
 
-void FPlayerSession::SendFileData(const TArray<uint8>& ByteData, const FString& MimeType, const FString& FileExtension) const
+void UE::PixelStreaming::FPlayerSession::SendFileData(const TArray<uint8>& ByteData, const FString& MimeType, const FString& FileExtension) const
 {
 	// Send the mime type first
-	FPlayerSession::SendMessage(PixelStreamingProtocol::EToPlayerMsg::FileMimeType, MimeType);
+	UE::PixelStreaming::FPlayerSession::SendMessage(UE::PixelStreaming::Protocol::EToPlayerMsg::FileMimeType, MimeType);
 	// Send the extension next
-	FPlayerSession::SendMessage(PixelStreamingProtocol::EToPlayerMsg::FileExtension, FileExtension);
+	UE::PixelStreaming::FPlayerSession::SendMessage(UE::PixelStreaming::Protocol::EToPlayerMsg::FileExtension, FileExtension);
 	// Send the contents of the file
-	AsyncTask(ENamedThreads::AnyHiPriThreadHiPriTask, [&, ByteData]() {
-		SendArbitraryData(this->DataChannel, ByteData, static_cast<uint8>(PixelStreamingProtocol::EToPlayerMsg::FileContents));
+	AsyncTask(ENamedThreads::AnyHiPriThreadHiPriTask, [this, ByteData]() {
+		UE::PixelStreaming::FPlayerSession::SendArbitraryData(ByteData, static_cast<uint8>(UE::PixelStreaming::Protocol::EToPlayerMsg::FileContents));
 	});
 }
 
-void FPlayerSession::SendUnfreezeFrame() const
+void UE::PixelStreaming::FPlayerSession::SendUnfreezeFrame() const
 {
 	if (!DataChannel)
 	{
 		return;
 	}
 
-	const uint8 MessageType = static_cast<uint8>(PixelStreamingProtocol::EToPlayerMsg::UnfreezeFrame);
+	const uint8 MessageType = static_cast<uint8>(UE::PixelStreaming::Protocol::EToPlayerMsg::UnfreezeFrame);
 
 	rtc::CopyOnWriteBuffer Buffer(sizeof(MessageType));
 
 	size_t Pos = 0;
-	Pos = SerializeToBuffer(Buffer, Pos, &MessageType, sizeof(MessageType));
+	Pos		   = SerializeToBuffer(Buffer, Pos, &MessageType, sizeof(MessageType));
 
 	if (!DataChannel->Send(webrtc::DataBuffer(Buffer, true)))
 	{
-		UE_LOG(PixelStreamer, Error, TEXT("failed to send unfreeze frame"));
+		UE_LOG(LogPixelStreaming, Error, TEXT("failed to send unfreeze frame"));
+	}
+}
+
+void UE::PixelStreaming::FPlayerSession::SendArbitraryData(const TArray<uint8>& DataBytes, const uint8 MessageType) const
+{
+	if (!DataChannel)
+	{
+		return;
+	}
+
+	// int32 results in a maximum 4GB file (4,294,967,296 bytes)
+	const int32 DataSize = DataBytes.Num();
+
+	// Maximum size of a single buffer should be 16KB as this is spec compliant message length for a single data channel transmission
+	const int32 MaxBufferBytes = 16 * 1024;
+	const int32 MessageHeader = sizeof(MessageType) + sizeof(DataSize);
+	const int32 MaxDataBytesPerMsg = MaxBufferBytes - MessageHeader;
+
+	int32 BytesTransmitted = 0;
+
+	while (BytesTransmitted < DataSize)
+	{
+		int32 RemainingBytes = DataSize - BytesTransmitted;
+		int32 BytesToTransmit = FGenericPlatformMath::Min(MaxDataBytesPerMsg, RemainingBytes);
+
+		rtc::CopyOnWriteBuffer Buffer(MessageHeader + BytesToTransmit);
+
+		size_t Pos = 0;
+
+		// Write message type
+		Pos = SerializeToBuffer(Buffer, Pos, &MessageType, sizeof(MessageType));
+		// Write size of payload
+		Pos = SerializeToBuffer(Buffer, Pos, &DataSize, sizeof(DataSize));
+		// Write the data bytes payload
+		Pos = SerializeToBuffer(Buffer, Pos, DataBytes.GetData() + BytesTransmitted, BytesToTransmit);
+
+		uint64_t BufferBefore = DataChannel->buffered_amount();
+		while (BufferBefore + BytesToTransmit >= 16 * 1024 * 1024) // 16MB (WebRTC Data Channel buffer size)
+		{
+			FPlatformProcess::Sleep(0.000001f); // sleep 1 microsecond
+			BufferBefore = DataChannel->buffered_amount();
+		}
+
+		if (!DataChannel->Send(webrtc::DataBuffer(Buffer, true)))
+		{
+			UE_LOG(LogPixelStreaming, Error, TEXT("Failed to send data channel packet"));
+			return;
+		}
+
+		// Increment the number of bytes transmitted
+		BytesTransmitted += BytesToTransmit;
 	}
 }
 
@@ -251,101 +310,101 @@ void FPlayerSession::SendUnfreezeFrame() const
 // webrtc::PeerConnectionObserver implementation.
 //
 
-void FPlayerSession::OnSignalingChange(webrtc::PeerConnectionInterface::SignalingState NewState)
+void UE::PixelStreaming::FPlayerSession::OnSignalingChange(webrtc::PeerConnectionInterface::SignalingState NewState)
 {
-	UE_LOG(PixelStreamer, Log, TEXT("%s : PlayerId=%s, NewState=%s"), TEXT("FPlayerSession::OnSignalingChange"), *PlayerId, ToString(NewState));
+	UE_LOG(LogPixelStreaming, Log, TEXT("%s : PlayerId=%s, NewState=%s"), TEXT("UE::PixelStreaming::FPlayerSession::OnSignalingChange"), *PlayerId, ToString(NewState));
 }
 
 // Called when a remote stream is added
-void FPlayerSession::OnAddStream(rtc::scoped_refptr<webrtc::MediaStreamInterface> Stream)
+void UE::PixelStreaming::FPlayerSession::OnAddStream(rtc::scoped_refptr<webrtc::MediaStreamInterface> Stream)
 {
-	UE_LOG(PixelStreamer, Log, TEXT("%s : PlayerId=%s, Stream=%s"), TEXT("FPlayerSession::OnAddStream"), *PlayerId, *ToString(Stream->id()));
+	UE_LOG(LogPixelStreaming, Log, TEXT("%s : PlayerId=%s, Stream=%s"), TEXT("UE::PixelStreaming::FPlayerSession::OnAddStream"), *PlayerId, *ToString(Stream->id()));
 }
 
-void FPlayerSession::OnRemoveStream(rtc::scoped_refptr<webrtc::MediaStreamInterface> Stream)
+void UE::PixelStreaming::FPlayerSession::OnRemoveStream(rtc::scoped_refptr<webrtc::MediaStreamInterface> Stream)
 {
-	UE_LOG(PixelStreamer, Log, TEXT("%s : PlayerId=%s, Stream=%s"), TEXT("FPlayerSession::OnRemoveStream"), *PlayerId, *ToString(Stream->id()));
+	UE_LOG(LogPixelStreaming, Log, TEXT("%s : PlayerId=%s, Stream=%s"), TEXT("UE::PixelStreaming::FPlayerSession::OnRemoveStream"), *PlayerId, *ToString(Stream->id()));
 }
 
-void FPlayerSession::OnDataChannel(rtc::scoped_refptr<webrtc::DataChannelInterface> InDataChannel)
+void UE::PixelStreaming::FPlayerSession::OnDataChannel(rtc::scoped_refptr<webrtc::DataChannelInterface> InDataChannel)
 {
-	UE_LOG(PixelStreamer, Log, TEXT("%s : PlayerId=%s"), TEXT("FPlayerSession::OnDataChannel"), *PlayerId);
+	UE_LOG(LogPixelStreaming, Log, TEXT("%s : PlayerId=%s"), TEXT("UE::PixelStreaming::FPlayerSession::OnDataChannel"), *PlayerId);
 	DataChannel = InDataChannel;
 	DataChannelObserver.Register(InDataChannel);
 }
 
-void FPlayerSession::OnRenegotiationNeeded()
+void UE::PixelStreaming::FPlayerSession::OnRenegotiationNeeded()
 {
-	UE_LOG(PixelStreamer, Log, TEXT("%s : PlayerId=%s"), TEXT("FPlayerSession::OnRenegotiationNeeded"), *PlayerId);
+	UE_LOG(LogPixelStreaming, Log, TEXT("%s : PlayerId=%s"), TEXT("UE::PixelStreaming::FPlayerSession::OnRenegotiationNeeded"), *PlayerId);
 }
 
-void FPlayerSession::OnIceConnectionChange(webrtc::PeerConnectionInterface::IceConnectionState NewState)
+void UE::PixelStreaming::FPlayerSession::OnIceConnectionChange(webrtc::PeerConnectionInterface::IceConnectionState NewState)
 {
-	UE_LOG(PixelStreamer, Log, TEXT("%s : PlayerId=%s, NewState=%s"), TEXT("FPlayerSession::OnIceConnectionChange"), *PlayerId, ToString(NewState));
+	UE_LOG(LogPixelStreaming, Log, TEXT("%s : PlayerId=%s, NewState=%s"), TEXT("UE::PixelStreaming::FPlayerSession::OnIceConnectionChange"), *PlayerId, ToString(NewState));
 }
 
-void FPlayerSession::OnIceGatheringChange(webrtc::PeerConnectionInterface::IceGatheringState NewState)
+void UE::PixelStreaming::FPlayerSession::OnIceGatheringChange(webrtc::PeerConnectionInterface::IceGatheringState NewState)
 {
-	UE_LOG(PixelStreamer, Log, TEXT("%s : PlayerId=%s, NewState=%s"), TEXT("FPlayerSession::OnIceGatheringChange"), *PlayerId, ToString(NewState));
+	UE_LOG(LogPixelStreaming, Log, TEXT("%s : PlayerId=%s, NewState=%s"), TEXT("UE::PixelStreaming::FPlayerSession::OnIceGatheringChange"), *PlayerId, ToString(NewState));
 }
 
-void FPlayerSession::OnIceCandidate(const webrtc::IceCandidateInterface* Candidate)
+void UE::PixelStreaming::FPlayerSession::OnIceCandidate(const webrtc::IceCandidateInterface* Candidate)
 {
-	UE_LOG(PixelStreamer, Log, TEXT("%s : PlayerId=%s"), TEXT("FPlayerSession::OnIceCandidate"), *PlayerId);
+	UE_LOG(LogPixelStreaming, Log, TEXT("%s : PlayerId=%s"), TEXT("UE::PixelStreaming::FPlayerSession::OnIceCandidate"), *PlayerId);
 
-	this->SignallingServerConnection->SendIceCandidate(PlayerId, *Candidate);
+	SignallingServerConnection->SendIceCandidate(PlayerId, *Candidate);
 }
 
-void FPlayerSession::OnIceCandidatesRemoved(const std::vector<cricket::Candidate>& candidates)
+void UE::PixelStreaming::FPlayerSession::OnIceCandidatesRemoved(const std::vector<cricket::Candidate>& candidates)
 {
-	UE_LOG(PixelStreamer, Log, TEXT("%s : PlayerId=%s"), TEXT("FPlayerSession::OnIceCandidatesRemoved"), *PlayerId);
+	UE_LOG(LogPixelStreaming, Log, TEXT("%s : PlayerId=%s"), TEXT("UE::PixelStreaming::FPlayerSession::OnIceCandidatesRemoved"), *PlayerId);
 }
 
-void FPlayerSession::OnIceConnectionReceivingChange(bool Receiving)
+void UE::PixelStreaming::FPlayerSession::OnIceConnectionReceivingChange(bool Receiving)
 {
-	UE_LOG(PixelStreamer, Log, TEXT("%s : PlayerId=%s, Receiving=%d"), TEXT("FPlayerSession::OnIceConnectionReceivingChange"), *PlayerId, *reinterpret_cast<int8*>(&Receiving));
+	UE_LOG(LogPixelStreaming, Log, TEXT("%s : PlayerId=%s, Receiving=%d"), TEXT("UE::PixelStreaming::FPlayerSession::OnIceConnectionReceivingChange"), *PlayerId, *reinterpret_cast<int8*>(&Receiving));
 }
 
-void FPlayerSession::OnTrack(rtc::scoped_refptr<webrtc::RtpTransceiverInterface> transceiver)
+void UE::PixelStreaming::FPlayerSession::OnTrack(rtc::scoped_refptr<webrtc::RtpTransceiverInterface> transceiver)
 {
-	UE_LOG(PixelStreamer, Log, TEXT("%s : PlayerId=%s"), TEXT("FPlayerSession::OnTrack"), *PlayerId);
+	UE_LOG(LogPixelStreaming, Log, TEXT("%s : PlayerId=%s"), TEXT("UE::PixelStreaming::FPlayerSession::OnTrack"), *PlayerId);
 
 	// print out track type
-	auto mediaType = transceiver->media_type();
+	cricket::MediaType mediaType = transceiver->media_type();
 	switch (mediaType)
 	{
 		case cricket::MediaType::MEDIA_TYPE_AUDIO:
-			UE_LOG(PixelStreamer, Log, TEXT("Track was type: audio"));
+			UE_LOG(LogPixelStreaming, Log, TEXT("Track was type: audio"));
 			break;
 		case cricket::MediaType::MEDIA_TYPE_VIDEO:
-			UE_LOG(PixelStreamer, Log, TEXT("Track was type: video"));
+			UE_LOG(LogPixelStreaming, Log, TEXT("Track was type: video"));
 			break;
 		case cricket::MediaType::MEDIA_TYPE_DATA:
-			UE_LOG(PixelStreamer, Log, TEXT("Track was type: data"));
+			UE_LOG(LogPixelStreaming, Log, TEXT("Track was type: data"));
 			break;
 		default:
-			UE_LOG(PixelStreamer, Log, TEXT("Track was an unsupported type"));
+			UE_LOG(LogPixelStreaming, Log, TEXT("Track was an unsupported type"));
 			break;
 	}
 
 	// print out track direction
-	webrtc::RtpTransceiverDirection direction = transceiver->direction();
-	switch (direction)
+	webrtc::RtpTransceiverDirection Direction = transceiver->direction();
+	switch (Direction)
 	{
 		case webrtc::RtpTransceiverDirection::kSendRecv:
-			UE_LOG(PixelStreamer, Log, TEXT("Track direction: send+recv"));
+			UE_LOG(LogPixelStreaming, Log, TEXT("Track direction: send+recv"));
 			break;
 		case webrtc::RtpTransceiverDirection::kSendOnly:
-			UE_LOG(PixelStreamer, Log, TEXT("Track direction: send only"));
+			UE_LOG(LogPixelStreaming, Log, TEXT("Track direction: send only"));
 			break;
 		case webrtc::RtpTransceiverDirection::kRecvOnly:
-			UE_LOG(PixelStreamer, Log, TEXT("Track direction: recv only"));
+			UE_LOG(LogPixelStreaming, Log, TEXT("Track direction: recv only"));
 			break;
 		case webrtc::RtpTransceiverDirection::kInactive:
-			UE_LOG(PixelStreamer, Log, TEXT("Track direction: inactive"));
+			UE_LOG(LogPixelStreaming, Log, TEXT("Track direction: inactive"));
 			break;
 		case webrtc::RtpTransceiverDirection::kStopped:
-			UE_LOG(PixelStreamer, Log, TEXT("Track direction: stopped"));
+			UE_LOG(LogPixelStreaming, Log, TEXT("Track direction: stopped"));
 			break;
 	}
 
@@ -356,33 +415,33 @@ void FPlayerSession::OnTrack(rtc::scoped_refptr<webrtc::RtpTransceiverInterface>
 	}
 
 	rtc::scoped_refptr<webrtc::MediaStreamTrackInterface> MediaStreamTrack = Receiver->track();
-	FString TrackEnabledStr = MediaStreamTrack->enabled() ? FString("Enabled") : FString("Disabled");
-	FString TrackStateStr = MediaStreamTrack->state() == webrtc::MediaStreamTrackInterface::TrackState::kLive ? FString("Live") : FString("Ended");
-	UE_LOG(PixelStreamer, Log, TEXT("MediaStreamTrack id: %s | Is enabled: %s | State: %s"), *FString(MediaStreamTrack->id().c_str()), *TrackEnabledStr, *TrackStateStr);
+	FString TrackEnabledStr												   = MediaStreamTrack->enabled() ? FString("Enabled") : FString("Disabled");
+	FString TrackStateStr												   = MediaStreamTrack->state() == webrtc::MediaStreamTrackInterface::TrackState::kLive ? FString("Live") : FString("Ended");
+	UE_LOG(LogPixelStreaming, Log, TEXT("MediaStreamTrack id: %s | Is enabled: %s | State: %s"), *FString(MediaStreamTrack->id().c_str()), *TrackEnabledStr, *TrackStateStr);
 
-	webrtc::AudioTrackInterface* AudioTrack = static_cast<webrtc::AudioTrackInterface*>(MediaStreamTrack.get());
+	webrtc::AudioTrackInterface* AudioTrack	  = static_cast<webrtc::AudioTrackInterface*>(MediaStreamTrack.get());
 	webrtc::AudioSourceInterface* AudioSource = AudioTrack->GetSource();
-	FString AudioSourceStateStr = AudioSource->state() == webrtc::MediaSourceInterface::SourceState::kLive ? FString("Live") : FString("Not live");
-	FString AudioSourceRemoteStr = AudioSource->remote() ? FString("Remote") : FString("Local");
-	UE_LOG(PixelStreamer, Log, TEXT("AudioSource | State: %s | Locality: %s"), *AudioSourceStateStr, *AudioSourceRemoteStr);
-	AudioSource->AddSink(&this->AudioSink);
+	FString AudioSourceStateStr				  = AudioSource->state() == webrtc::MediaSourceInterface::SourceState::kLive ? FString("Live") : FString("Not live");
+	FString AudioSourceRemoteStr			  = AudioSource->remote() ? FString("Remote") : FString("Local");
+	UE_LOG(LogPixelStreaming, Log, TEXT("AudioSource | State: %s | Locality: %s"), *AudioSourceStateStr, *AudioSourceRemoteStr);
+	AudioSource->AddSink(&AudioSink);
 }
 
-void FPlayerSession::OnRemoveTrack(rtc::scoped_refptr<webrtc::RtpReceiverInterface> receiver)
+void UE::PixelStreaming::FPlayerSession::OnRemoveTrack(rtc::scoped_refptr<webrtc::RtpReceiverInterface> receiver)
 {
-	UE_LOG(PixelStreamer, Log, TEXT("%s : PlayerId=%s"), TEXT("FPlayerSession::OnRemoveTrack"), *PlayerId);
+	UE_LOG(LogPixelStreaming, Log, TEXT("%s : PlayerId=%s"), TEXT("UE::PixelStreaming::FPlayerSession::OnRemoveTrack"), *PlayerId);
 }
 
-size_t FPlayerSession::SerializeToBuffer(rtc::CopyOnWriteBuffer& Buffer, size_t Pos, const void* Data, size_t DataSize) const
+size_t UE::PixelStreaming::FPlayerSession::SerializeToBuffer(rtc::CopyOnWriteBuffer& Buffer, size_t Pos, const void* Data, size_t DataSize) const
 {
 	FMemory::Memcpy(&Buffer[Pos], reinterpret_cast<const uint8_t*>(Data), DataSize);
 	return Pos + DataSize;
 }
 
-void FPlayerSession::SendVideoEncoderQP(double QP) const
+void UE::PixelStreaming::FPlayerSession::SendVideoEncoderQP(double QP) const
 {
-	if (!SendMessage(PixelStreamingProtocol::EToPlayerMsg::VideoEncoderAvgQP, FString::FromInt(QP)))
+	if (!SendMessage(UE::PixelStreaming::Protocol::EToPlayerMsg::VideoEncoderAvgQP, FString::FromInt(QP)))
 	{
-		UE_LOG(PixelStreamer, Verbose, TEXT("Failed to send video encoder QP to peer."));
+		UE_LOG(LogPixelStreaming, Verbose, TEXT("Failed to send video encoder QP to peer."));
 	}
 }
diff --git a/Engine/Plugins/Media/PixelStreaming/Source/PixelStreaming/Private/PlayerSession.h b/Engine/Plugins/Media/PixelStreaming/Source/PixelStreaming/Private/PlayerSession.h
--- a/Engine/Plugins/Media/PixelStreaming/Source/PixelStreaming/Private/PlayerSession.h
+++ b/Engine/Plugins/Media/PixelStreaming/Source/PixelStreaming/Private/PlayerSession.h
@@ -5,72 +5,78 @@
 #include "PixelStreamingPrivate.h"
 #include "ProtocolDefs.h"
 #include "HAL/ThreadSafeBool.h"
-#include "PixelStreamingAudioSink.h"
-#include "PlayerId.h"
-#include "PixelStreamingDataChannelObserver.h"
-#include "PixelStreamingStats.h"
-#include "PixelStreamingRTCStatsCollector.h"
+#include "AudioSink.h"
+#include "PixelStreamingPlayerId.h"
+#include "DataChannelObserver.h"
+#include "Stats.h"
+#include "RTCStatsCollector.h"
 #include "IPixelStreamingStatsConsumer.h"
 
-class FSignallingServerConnection;
-class IPixelStreamingSessions;
+namespace UE {
+	namespace PixelStreaming {
+		class FSignallingServerConnection;
+		class IPixelStreamingSessions;
 
-class FPlayerSession : public webrtc::PeerConnectionObserver
-{
-public:
-	FPlayerSession(IPixelStreamingSessions* InSessions, FSignallingServerConnection* InSignallingServerConnection, FPlayerId PlayerId);
-	virtual ~FPlayerSession() override;
+		class FPlayerSession : public webrtc::PeerConnectionObserver
+		{
+		public:
+			FPlayerSession(IPixelStreamingSessions* InSessions, FSignallingServerConnection* InSignallingServerConnection, FPixelStreamingPlayerId PlayerId);
+			virtual ~FPlayerSession() override;
 
-	webrtc::PeerConnectionInterface& GetPeerConnection();
-	void SetPeerConnection(const rtc::scoped_refptr<webrtc::PeerConnectionInterface>& InPeerConnection);
+			webrtc::PeerConnectionInterface& GetPeerConnection();
+			void SetPeerConnection(const rtc::scoped_refptr<webrtc::PeerConnectionInterface>& InPeerConnection);
+			void SetDataChannel(const rtc::scoped_refptr<webrtc::DataChannelInterface>& InDataChannel);
+			void SetVideoSource(const rtc::scoped_refptr<webrtc::VideoTrackSourceInterface>& InVideoSource);
 
-	void SetDataChannel(const rtc::scoped_refptr<webrtc::DataChannelInterface>& InDataChannel);
+			void OnAnswer(FString Sdp);
+			void OnRemoteIceCandidate(const std::string& SdpMid, int SdpMLineIndex, const std::string& Sdp);
+			void DisconnectPlayer(const FString& Reason);
 
-	void OnAnswer(FString Sdp);
-	void OnRemoteIceCandidate(const std::string& SdpMid, int SdpMLineIndex, const std::string& Sdp);
-	void DisconnectPlayer(const FString& Reason);
+			FPixelStreamingPlayerId GetPlayerId() const;
+			FAudioSink& GetAudioSink();
+			FDataChannelObserver& GetDataChannelObserver();
 
-	FPlayerId GetPlayerId() const;
-	FPixelStreamingAudioSink& GetAudioSink();
-	FPixelStreamingDataChannelObserver& GetDataChannelObserver();
+			bool SendMessage(Protocol::EToPlayerMsg Type, const FString& Descriptor) const;
+			void SendQualityControlStatus(bool bIsQualityController) const;
+			void SendFreezeFrame(const TArray64<uint8>& JpegBytes) const;
+			void SendFileData(const TArray<uint8>& ByteData, const FString& MimeType, const FString& FileExtension) const;
+			void SendUnfreezeFrame() const;
+			void SendArbitraryData(const TArray<uint8>& DataBytes, const uint8 MessageType) const;
+			void SendVideoEncoderQP(double QP) const;
+			void PollWebRTCStats() const;
 
-	bool SendMessage(PixelStreamingProtocol::EToPlayerMsg Type, const FString& Descriptor) const;
-	void SendQualityControlStatus(bool bIsQualityController) const;
-	void SendFreezeFrame(const TArray64<uint8>& JpegBytes) const;
-	void SendFileData(const TArray<uint8>& ByteData, const FString& MimeType, const FString& FileExtension) const;
-	void SendUnfreezeFrame() const;
-	void SendVideoEncoderQP(double QP) const;
-	void PollWebRTCStats() const;
+		private:
+			void ModifyAudioTransceiverDirection();
 
-private:
-	void ModifyAudioTransceiverDirection();
+			//
+			// webrtc::PeerConnectionObserver implementation.
+			//
+			virtual void OnSignalingChange(webrtc::PeerConnectionInterface::SignalingState NewState) override;
+			virtual void OnAddStream(rtc::scoped_refptr<webrtc::MediaStreamInterface> Stream) override;
+			virtual void OnRemoveStream(rtc::scoped_refptr<webrtc::MediaStreamInterface> Stream) override;
+			virtual void OnDataChannel(rtc::scoped_refptr<webrtc::DataChannelInterface> Channel) override;
+			virtual void OnRenegotiationNeeded() override;
+			virtual void OnIceConnectionChange(webrtc::PeerConnectionInterface::IceConnectionState NewState) override;
+			virtual void OnIceGatheringChange(webrtc::PeerConnectionInterface::IceGatheringState NewState) override;
+			virtual void OnIceCandidate(const webrtc::IceCandidateInterface* Candidate) override;
+			virtual void OnIceCandidatesRemoved(const std::vector<cricket::Candidate>& candidates) override;
+			virtual void OnIceConnectionReceivingChange(bool Receiving) override;
+			virtual void OnTrack(rtc::scoped_refptr<webrtc::RtpTransceiverInterface> transceiver) override;
+			virtual void OnRemoveTrack(rtc::scoped_refptr<webrtc::RtpReceiverInterface> receiver) override;
 
-	//
-	// webrtc::PeerConnectionObserver implementation.
-	//
-	virtual void OnSignalingChange(webrtc::PeerConnectionInterface::SignalingState NewState) override;
-	virtual void OnAddStream(rtc::scoped_refptr<webrtc::MediaStreamInterface> Stream) override;
-	virtual void OnRemoveStream(rtc::scoped_refptr<webrtc::MediaStreamInterface> Stream) override;
-	virtual void OnDataChannel(rtc::scoped_refptr<webrtc::DataChannelInterface> Channel) override;
-	virtual void OnRenegotiationNeeded() override;
-	virtual void OnIceConnectionChange(webrtc::PeerConnectionInterface::IceConnectionState NewState) override;
-	virtual void OnIceGatheringChange(webrtc::PeerConnectionInterface::IceGatheringState NewState) override;
-	virtual void OnIceCandidate(const webrtc::IceCandidateInterface* Candidate) override;
-	virtual void OnIceCandidatesRemoved(const std::vector<cricket::Candidate>& candidates) override;
-	virtual void OnIceConnectionReceivingChange(bool Receiving) override;
-	virtual void OnTrack(rtc::scoped_refptr<webrtc::RtpTransceiverInterface> transceiver) override;
-	virtual void OnRemoveTrack(rtc::scoped_refptr<webrtc::RtpReceiverInterface> receiver) override;
+		private:
+			size_t SerializeToBuffer(rtc::CopyOnWriteBuffer& Buffer, size_t Pos, const void* Data, size_t DataSize) const;
 
-private:
-	size_t SerializeToBuffer(rtc::CopyOnWriteBuffer& Buffer, size_t Pos, const void* Data, size_t DataSize) const;
-
-	FSignallingServerConnection* SignallingServerConnection;
-	FPlayerId PlayerId;
-	rtc::scoped_refptr<webrtc::PeerConnectionInterface> PeerConnection;
-	rtc::scoped_refptr<webrtc::DataChannelInterface> DataChannel;
-	FPixelStreamingDataChannelObserver DataChannelObserver;
-	FThreadSafeBool bDisconnecting = false;
-	FPixelStreamingAudioSink AudioSink;
-	rtc::scoped_refptr<webrtc::RTCStatsCollectorCallback> WebRTCStatsCallback;
-	TSharedPtr<IPixelStreamingStatsConsumer> QPReporter;
-};
+			FSignallingServerConnection* SignallingServerConnection;
+			FPixelStreamingPlayerId PlayerId;
+			rtc::scoped_refptr<webrtc::PeerConnectionInterface> PeerConnection;
+			rtc::scoped_refptr<webrtc::DataChannelInterface> DataChannel;
+			FDataChannelObserver DataChannelObserver;
+			FThreadSafeBool bDisconnecting = false;
+			FAudioSink AudioSink;
+			rtc::scoped_refptr<webrtc::RTCStatsCollectorCallback> WebRTCStatsCallback;
+			TSharedPtr<IPixelStreamingStatsConsumer> QPReporter;
+			rtc::scoped_refptr<webrtc::VideoTrackSourceInterface> VideoSource;
+		};
+	}
+}
diff --git a/Engine/Plugins/Media/PixelStreaming/Source/PixelStreaming/Private/ProtocolDefs.h b/Engine/Plugins/Media/PixelStreaming/Source/PixelStreaming/Private/ProtocolDefs.h
--- a/Engine/Plugins/Media/PixelStreaming/Source/PixelStreaming/Private/ProtocolDefs.h
+++ b/Engine/Plugins/Media/PixelStreaming/Source/PixelStreaming/Private/ProtocolDefs.h
@@ -5,112 +5,116 @@
 #include "HAL/Platform.h"
 #include "api/data_channel_interface.h"
 
-namespace PixelStreamingProtocol
-{
-	enum EPixelStreamingPlayerFlags
-	{
-		PSPFlag_None = 0x00,
-		PSPFlag_SupportsDataChannel = 0x01,
-		PSPFlag_IsSFU = 0x02,
-	};
-
-	enum class EToStreamerMsg : uint8
-	{
-		/**********************************************************************/
-
-		/*
-		 * Control Messages. Range = 0..49.
-		 */
-
-		IFrameRequest = 0,
-		RequestQualityControl = 1,
-		FpsRequest = 2,
-		AverageBitrateRequest = 3,
-		StartStreaming = 4,
-		StopStreaming = 5,
-		LatencyTest = 6,
-		RequestInitialSettings = 7,
-
-		/**********************************************************************/
-
-		/*
-		 * Input Messages. Range = 50..89.
-		 */
-
-		// Generic Input Messages. Range = 50..59.
-		UIInteraction = 50,
-		Command = 51,
-
-		// Keyboard Input Message. Range = 60..69.
-		KeyDown = 60,
-		KeyUp = 61,
-		KeyPress = 62,
-
-		// Mouse Input Messages. Range = 70..79.
-		MouseEnter = 70,
-		MouseLeave = 71,
-		MouseDown = 72,
-		MouseUp = 73,
-		MouseMove = 74,
-		MouseWheel = 75,
-
-		// Touch Input Messages. Range = 80..89.
-		TouchStart = 80,
-		TouchEnd = 81,
-		TouchMove = 82,
-
-		// Gamepad Input Messages. Range = 90..99
-		GamepadButtonPressed = 90,
-		GamepadButtonReleased = 91,
-		GamepadAnalog = 92,
-
-		/**********************************************************************/
-
-		/*
-		 * Ensure Count is the final entry.
-		 */
-		Count
-
-		/**********************************************************************/
-	};
-
-	//! Messages that can be sent to the webrtc players
-	enum class EToPlayerMsg : uint8
-	{
-		QualityControlOwnership = 0,
-		Response = 1,
-		Command = 2,
-		FreezeFrame = 3,
-		UnfreezeFrame = 4,
-		VideoEncoderAvgQP = 5, // average Quantisation Parameter value of Video Encoder, roughly depicts video encoding quality
-		LatencyTest = 6,
-		InitialSettings = 7,
-		FileExtension = 8,
-		FileMimeType = 9,
-		FileContents = 10
-	};
-
-	template <typename T>
-	static const T& ParseBuffer(const uint8*& Data, uint32& Size)
-	{
-		checkf(sizeof(T) <= Size, TEXT("%d - %d"), sizeof(T), Size);
-		const T& Value = *reinterpret_cast<const T*>(Data);
-		Data += sizeof(T);
-		Size -= sizeof(T);
-		return Value;
-	}
-
-	static FString ParseString(const webrtc::DataBuffer& Buffer, const size_t Offset = 0)
-	{
-		FString Res;
-		if (Buffer.data.size() > Offset)
+namespace UE {
+	namespace PixelStreaming {
+		namespace Protocol
 		{
-			size_t StringLength = (Buffer.data.size() - Offset) / sizeof(TCHAR);
-			Res.GetCharArray().SetNumUninitialized(StringLength + 1);
-			FMemory::Memcpy(Res.GetCharArray().GetData(), Buffer.data.data() + Offset, StringLength * sizeof(TCHAR));
-			Res.GetCharArray()[StringLength] = 0;
-		}
-		return Res;
+			enum EPlayerFlags
+			{
+				PSPFlag_None = 0x00,
+				PSPFlag_SupportsDataChannel = 0x01,
+				PSPFlag_IsSFU = 0x02,
+			};
+
+			enum class EToStreamerMsg : uint8
+			{
+				/**********************************************************************/
+
+				/*
+				* Control Messages. Range = 0..49.
+				*/
+
+				IFrameRequest = 0,
+				RequestQualityControl = 1,
+				FpsRequest = 2,
+				AverageBitrateRequest = 3,
+				StartStreaming = 4,
+				StopStreaming = 5,
+				LatencyTest = 6,
+				RequestInitialSettings = 7,
+
+				/**********************************************************************/
+
+				/*
+				* Input Messages. Range = 50..89.
+				*/
+
+				// Generic Input Messages. Range = 50..59.
+				UIInteraction = 50,
+				Command = 51,
+
+				// Keyboard Input Message. Range = 60..69.
+				KeyDown = 60,
+				KeyUp = 61,
+				KeyPress = 62,
+
+				// Mouse Input Messages. Range = 70..79.
+				MouseEnter = 70,
+				MouseLeave = 71,
+				MouseDown = 72,
+				MouseUp = 73,
+				MouseMove = 74,
+				MouseWheel = 75,
+
+				// Touch Input Messages. Range = 80..89.
+				TouchStart = 80,
+				TouchEnd = 81,
+				TouchMove = 82,
+
+				// Gamepad Input Messages. Range = 90..99
+				GamepadButtonPressed = 90,
+				GamepadButtonReleased = 91,
+				GamepadAnalog = 92,
+
+				/**********************************************************************/
+
+				/*
+				* Ensure Count is the final entry.
+				*/
+				Count
+
+				/**********************************************************************/
+			};
+
+			//! Messages that can be sent to the webrtc players
+			enum class EToPlayerMsg : uint8
+			{
+				QualityControlOwnership = 0,
+				Response = 1,
+				Command = 2,
+				FreezeFrame = 3,
+				UnfreezeFrame = 4,
+				VideoEncoderAvgQP = 5, // average Quantisation Parameter value of Video Encoder, roughly depicts video encoding quality
+				LatencyTest = 6,
+				InitialSettings = 7,
+				FileExtension = 8,
+				FileMimeType = 9,
+				FileContents = 10
+			};
+
+			template <typename T>
+			static const T& ParseBuffer(const uint8*& Data, uint32& Size)
+			{
+				checkf(sizeof(T) <= Size, TEXT("%d - %d"), sizeof(T), Size);
+				const T& Value = *reinterpret_cast<const T*>(Data);
+				Data += sizeof(T);
+				Size -= sizeof(T);
+				return Value;
+			}
+
+			static FString ParseString(const webrtc::DataBuffer& Buffer, const size_t Offset = 0)
+			{
+				FString Res;
+				if (Buffer.data.size() > Offset)
+				{
+					size_t StringLength = (Buffer.data.size() - Offset) / sizeof(TCHAR);
+					Res.GetCharArray().SetNumUninitialized(StringLength + 1);
+					FMemory::Memcpy(Res.GetCharArray().GetData(), Buffer.data.data() + Offset, StringLength * sizeof(TCHAR));
+					Res.GetCharArray()[StringLength] = 0;
+				}
+				return Res;
+			}
+
+		}; // namespace UE::PixelStreaming::Protocol
 	}
-
-}; // namespace PixelStreamingProtocol
+}
diff --git a/Engine/Plugins/Media/PixelStreaming/Source/PixelStreaming/Private/RTCStatsCollector.cpp b/Engine/Plugins/Media/PixelStreaming/Source/PixelStreaming/Private/RTCStatsCollector.cpp
new file mode 100644
--- /dev/null
+++ b/Engine/Plugins/Media/PixelStreaming/Source/PixelStreaming/Private/RTCStatsCollector.cpp
@@ -0,0 +1,281 @@
+// Copyright Epic Games, Inc. All Rights Reserved.
+
+#include "RTCStatsCollector.h"
+#include "PixelStreamingPrivate.h"
+#include "PixelStreamingStatNames.h"
+
+/*
+------------------------ FStreamStatsSink------------------------
+*/
+
+UE::PixelStreaming::FRTCStatsCollector::FStreamStatsSink::FStreamStatsSink(FString InRid)
+	: Rid(InRid)
+{
+	Add(PixelStreamingStatNames::FirCount, 0);
+	Add(PixelStreamingStatNames::PliCount, 0);
+	Add(PixelStreamingStatNames::NackCount, 0);
+	Add(PixelStreamingStatNames::SliCount, 0);
+	Add(PixelStreamingStatNames::RetransmittedBytesSent, 0);
+	Add(PixelStreamingStatNames::TargetBitrate, 0);
+	Add(PixelStreamingStatNames::TotalEncodeBytesTarget, 0);
+	Add(PixelStreamingStatNames::KeyFramesEncoded, 0);
+	Add(PixelStreamingStatNames::FrameWidth, 0);
+	Add(PixelStreamingStatNames::FrameHeight, 0);
+	Add(PixelStreamingStatNames::HugeFramesSent, 0);
+
+	AddForCalculation(PixelStreamingStatNames::FramesSent);
+	AddForCalculation(PixelStreamingStatNames::BytesSent);
+	AddForCalculation(PixelStreamingStatNames::QPSum);
+	AddForCalculation(PixelStreamingStatNames::TotalEncodeTime);
+	AddForCalculation(PixelStreamingStatNames::TotalPacketSendDelay);
+	AddForCalculation(PixelStreamingStatNames::FramesEncoded);
+}
+
+/*
+------------- UE::PixelStreaming::FRTCStatsCollector-------------------
+*/
+
+UE::PixelStreaming::FRTCStatsCollector::FRTCStatsCollector()
+	: UE::PixelStreaming::FRTCStatsCollector(INVALID_PLAYER_ID)
+{
+}
+
+UE::PixelStreaming::FRTCStatsCollector::FRTCStatsCollector(FPixelStreamingPlayerId PlayerId)
+	: AssociatedPlayerId(PlayerId)
+	, LastCalculationCycles(FPlatformTime::Cycles64())
+	, bIsEnabled(!UE::PixelStreaming::Settings::CVarPixelStreamingWebRTCDisableStats.GetValueOnAnyThread())
+{
+	TrackStatsSink.Add(PixelStreamingStatNames::JitterBufferDelay, 2);
+	TrackStatsSink.Add(PixelStreamingStatNames::FramesReceived, 0);
+	TrackStatsSink.Add(PixelStreamingStatNames::FramesDecoded, 0);
+	TrackStatsSink.Add(PixelStreamingStatNames::FramesDropped, 0);
+	TrackStatsSink.Add(PixelStreamingStatNames::FramesCorrupted, 0);
+	TrackStatsSink.Add(PixelStreamingStatNames::PartialFramesLost, 0);
+	TrackStatsSink.Add(PixelStreamingStatNames::FullFramesLost, 0);
+	TrackStatsSink.Add(PixelStreamingStatNames::JitterBufferTargetDelay, 2);
+	TrackStatsSink.Add(PixelStreamingStatNames::InterruptionCount, 0);
+	TrackStatsSink.Add(PixelStreamingStatNames::TotalInterruptionDuration, 2);
+	TrackStatsSink.Add(PixelStreamingStatNames::FreezeCount, 0);
+	TrackStatsSink.Add(PixelStreamingStatNames::PauseCount, 0);
+	TrackStatsSink.Add(PixelStreamingStatNames::TotalFreezesDuration, 2);
+	TrackStatsSink.Add(PixelStreamingStatNames::TotalPausesDuration, 2);
+
+	SourceStatsSink.Add(PixelStreamingStatNames::SourceFps, 0);
+}
+
+void UE::PixelStreaming::FRTCStatsCollector::AddRef() const
+{
+	FPlatformAtomics::InterlockedIncrement(&RefCount);
+}
+
+rtc::RefCountReleaseStatus UE::PixelStreaming::FRTCStatsCollector::Release() const
+{
+	if (FPlatformAtomics::InterlockedDecrement(&RefCount) == 0)
+	{
+		return rtc::RefCountReleaseStatus::kDroppedLastRef;
+	}
+
+	return rtc::RefCountReleaseStatus::kOtherRefsRemained;
+}
+
+void UE::PixelStreaming::FRTCStatsCollector::ExtractValueAndSet(const webrtc::RTCStatsMemberInterface* ExtractFrom, UE::PixelStreaming::FStatData* SetValueHere)
+{
+	FString StatValueStr = FString(ExtractFrom->ValueToString().c_str());
+	double StatValueDouble = FCString::Atod(*StatValueStr);
+	SetValueHere->StatValue = StatValueDouble;
+}
+
+void UE::PixelStreaming::FRTCStatsCollector::OnStatsDelivered(const rtc::scoped_refptr<const webrtc::RTCStatsReport>& Report)
+{	
+	if(!bIsEnabled) 
+	{
+		return;
+	}	
+
+	UE::PixelStreaming::FStats* PSStats = UE::PixelStreaming::FStats::Get();
+
+	if (!PSStats)
+	{
+		return;
+	}
+
+	if (!Report)
+	{
+		return;
+	}
+
+	for (webrtc::RTCStatsReport::ConstIterator Iter = Report->begin(); Iter != Report->end(); ++Iter)
+	{
+		const webrtc::RTCStats& Stats = *Iter;
+		FString StatsType = FString(Stats.type());
+		std::vector<const webrtc::RTCStatsMemberInterface*> StatMembers = Stats.Members();
+
+		// For debugging
+		// UE_LOG(LogPixelStreaming, Log, TEXT("------------%s---%s------------"), *FString(Stats.id().c_str()), *FString(Stats.type()));
+		// for (const webrtc::RTCStatsMemberInterface* StatMember : StatMembers)
+		// {
+		// 	UE_LOG(LogPixelStreaming, Log, TEXT("%s"), *FString(StatMember->name()));
+		// }
+
+		// Find relevant FStatsSink
+		FStatsSink* StatsSink = nullptr;
+		FString Rid;
+
+		if (StatsType == TrackStatsStr)
+		{
+			StatsSink = &TrackStatsSink;
+		}
+		else if (StatsType == SourceStatsStr)
+		{
+			StatsSink = &SourceStatsSink;
+		}
+		else if (StatsType == StreamStatsStr)
+		{
+			// Extract the `rid`
+			for (const webrtc::RTCStatsMemberInterface* StatMember : StatMembers)
+			{
+				FString StatName = FString(StatMember->name());
+				if (StatName == RidStr)
+				{
+					Rid = FString(StatMember->ValueToString().c_str());
+					if (!StreamStatsSinks.Contains(Rid))
+					{
+						StreamStatsSinks.Add(Rid, FStreamStatsSink(Rid));
+					}
+					StatsSink = StreamStatsSinks.Find(Rid);
+					break;
+				}
+			}
+		}
+		else
+		{
+			continue;
+		}
+
+		if (!StatsSink)
+		{
+			return;
+		}
+
+		bool bSimulcastStat = Rid.Contains(TEXT("Simulcast"), ESearchCase::IgnoreCase, ESearchDir::FromStart);
+		FString Id = bSimulcastStat ? Rid : AssociatedPlayerId;
+
+		for (const webrtc::RTCStatsMemberInterface* StatMember : StatMembers)
+		{
+			FName StatName = FName(StatMember->name());
+
+			// Check if the stat is a stat we use for calculation, if so, extract the value and store it
+			FCalculatedStat* CalcStat = StatsSink->GetCalcStat(StatName);
+			if (CalcStat)
+			{
+				ExtractValueAndSet(StatMember, &CalcStat->LatestStat);
+				continue;
+			}
+
+			// Check if stat is a stat we emit immediately, if so, extract the value and emit it
+			UE::PixelStreaming::FStatData* StatToEmit = StatsSink->Get(StatName);
+			if (StatToEmit)
+			{
+				bool bStatAlreadyNonZero = StatToEmit->StatValue > 0.0 || StatToEmit->StatValue < 0.0;
+
+				ExtractValueAndSet(StatMember, StatToEmit);
+
+				// Don't bother emitting a zero value
+				if (bStatAlreadyNonZero || StatToEmit->StatValue > 0.0 || StatToEmit->StatValue < 0.0)
+				{
+					PSStats->StorePeerStat(Id, *StatToEmit);
+				}
+
+				continue;
+			}
+		}
+
+		CalculateStats(PSStats);
+	}
+}
+
+void UE::PixelStreaming::FRTCStatsCollector::CalculateStats(UE::PixelStreaming::FStats* PSStats)
+{
+	uint64 CyclesNow = FPlatformTime::Cycles64();
+	double SecondsDelta = FGenericPlatformTime::ToSeconds64(CyclesNow - LastCalculationCycles);
+
+	// Don't calculate stats if 1 second window has not yet elapsed
+	if (SecondsDelta < 1.0)
+	{
+		return;
+	}
+
+	// This will bring the calculations into a 1 second window
+	double PerSecondRatio = 1.0 / SecondsDelta;
+	double FramesSentPerSecond = 0.0;
+	double EncodedFramesPerSecond = 0.0;
+
+	// Calculate stats based on each of the stream stats sinks
+
+	for (auto& Entry : StreamStatsSinks)
+	{
+		FString& Rid = Entry.Key;
+		FStatsSink& StreamSink = Entry.Value;
+		bool bSimulcastStat = Rid.Contains(TEXT("Simulcast"), ESearchCase::IgnoreCase, ESearchDir::FromStart);
+		FString Id = bSimulcastStat ? Rid : AssociatedPlayerId;
+
+		// --------- FrameSent Per Second -----------
+		FCalculatedStat* FramesSentStat = StreamSink.GetCalcStat(PixelStreamingStatNames::FramesSent);
+		if (FramesSentStat && FramesSentStat->LatestStat.StatValue > 0)
+		{
+			FramesSentPerSecond = FramesSentStat->CalculateDelta(PerSecondRatio);
+			PSStats->StorePeerStat(Id, UE::PixelStreaming::FStatData(PixelStreamingStatNames::FramesSentPerSecond, FramesSentPerSecond, 0));
+		}
+
+		// --------- BytesSent Per Second -----------
+		FCalculatedStat* BytesSentStat = StreamSink.GetCalcStat(PixelStreamingStatNames::BytesSent);
+		if (BytesSentStat && BytesSentStat->LatestStat.StatValue > 0)
+		{
+			double BytesSentPerSecond = BytesSentStat->CalculateDelta(PerSecondRatio);
+			double MegabitsPerSecond = BytesSentPerSecond / 1000.0 * 8.0;
+			FName StatName = PixelStreamingStatNames::Bitrate;
+			PSStats->StorePeerStat(Id, UE::PixelStreaming::FStatData(StatName, MegabitsPerSecond, 0));
+		}
+
+		// ------------- Encoded fps -------------
+		FCalculatedStat* EncodedFramesStat = StreamSink.GetCalcStat(PixelStreamingStatNames::FramesEncoded);
+		if (EncodedFramesStat && EncodedFramesStat->LatestStat.StatValue > 0)
+		{
+			EncodedFramesPerSecond = EncodedFramesStat->CalculateDelta(PerSecondRatio);
+			FName StatName = PixelStreamingStatNames::EncodedFramesPerSecond;
+			PSStats->StorePeerStat(Id, UE::PixelStreaming::FStatData(StatName, EncodedFramesPerSecond, 0));
+		}
+
+		// ------------- Avg QP Per Second -------------
+		FCalculatedStat* QPSumStat = StreamSink.GetCalcStat(PixelStreamingStatNames::QPSum);
+		if (QPSumStat && QPSumStat->LatestStat.StatValue > 0 && EncodedFramesPerSecond > 0.0)
+		{
+			double QPSumDeltaPerSecond = QPSumStat->CalculateDelta(PerSecondRatio);
+			double MeanQPPerFrame = QPSumDeltaPerSecond / EncodedFramesPerSecond;
+			FName StatName = PixelStreamingStatNames::MeanQPPerSecond;
+			PSStats->StorePeerStat(Id, UE::PixelStreaming::FStatData(StatName, MeanQPPerFrame, 0));
+		}
+
+		// ------------- Mean EncodeTime (ms) Per Frame -------------
+		FCalculatedStat* TotalEncodeTimeStat = StreamSink.GetCalcStat(PixelStreamingStatNames::TotalEncodeTime);
+		if (TotalEncodeTimeStat && TotalEncodeTimeStat->LatestStat.StatValue > 0 && EncodedFramesPerSecond > 0.0)
+		{
+			double TotalEncodeTimePerSecond = TotalEncodeTimeStat->CalculateDelta(PerSecondRatio);
+			double MeanEncodeTimePerFrameMs = TotalEncodeTimePerSecond / EncodedFramesPerSecond * 1000.0;
+			FName StatName = PixelStreamingStatNames::MeanEncodeTime;
+			PSStats->StorePeerStat(Id, UE::PixelStreaming::FStatData(StatName, MeanEncodeTimePerFrameMs, 2));
+		}
+
+		// ------------- Mean SendDelay (ms) Per Frame -------------
+		FCalculatedStat* TotalSendDelayStat = StreamSink.GetCalcStat(PixelStreamingStatNames::TotalPacketSendDelay);
+		if (TotalSendDelayStat && TotalSendDelayStat->LatestStat.StatValue > 0 && FramesSentPerSecond > 0)
+		{
+			double TotalSendDelayPerSecond = TotalSendDelayStat->CalculateDelta(PerSecondRatio);
+			double MeanSendDelayPerFrameMs = TotalSendDelayPerSecond / FramesSentPerSecond * 1000.0;
+			FName StatName = PixelStreamingStatNames::MeanSendDelay;
+			PSStats->StorePeerStat(Id, UE::PixelStreaming::FStatData(StatName, MeanSendDelayPerFrameMs, 2));
+		}
+	}
+
+	LastCalculationCycles = FPlatformTime::Cycles64();
+}
\ No newline at end of file
