Merging Pixel Streaming refactor using //UE5/Dev-TensorWorks-to-//UE5/Release-5.0

diff --git a/Engine/Plugins/Media/PixelStreaming/Source/PixelStreaming/Private/RTCStatsCollector.h b/Engine/Plugins/Media/PixelStreaming/Source/PixelStreaming/Private/RTCStatsCollector.h
new file mode 100644
--- /dev/null
+++ b/Engine/Plugins/Media/PixelStreaming/Source/PixelStreaming/Private/RTCStatsCollector.h
@@ -0,0 +1,95 @@
+// Copyright Epic Games, Inc. All Rights Reserved.
+
+#pragma once
+
+#include "WebRTCIncludes.h"
+#include "Containers/Map.h"
+#include "PixelStreamingPlayerId.h"
+#include "Stats.h"
+
+namespace UE {
+	namespace PixelStreaming {
+		class FRTCStatsCollector : public webrtc::RTCStatsCollectorCallback
+		{
+		public:
+			FRTCStatsCollector();
+			FRTCStatsCollector(FPixelStreamingPlayerId PlayerId);
+
+			// Begin RTCStatsCollectorCallback interface
+			void OnStatsDelivered(const rtc::scoped_refptr<const webrtc::RTCStatsReport>& report) override;
+			void AddRef() const override;
+			rtc::RefCountReleaseStatus Release() const override;
+			// End RTCStatsCollectorCallback interface
+
+		private:
+			void ExtractValueAndSet(const webrtc::RTCStatsMemberInterface* ExtractFrom, FStatData* SetValueHere);
+			void CalculateStats(FStats* PSStats);
+
+		private:
+			// Some stats have a latest value that is update each time that stats comes and
+			// then also a calculated value that might get updates every 1 second.
+			// An example is FPS, where total frames sent is the stat that is tracked but
+			// the calculated stat is frames sent per second.
+			struct FCalculatedStat
+			{
+				FStatData LatestStat;
+				double PrevValue;
+
+				double CalculateDelta(double Period)
+				{
+					double Delta = (LatestStat.StatValue - PrevValue) * Period;
+					PrevValue = LatestStat.StatValue;
+					return Delta;
+				}
+			};
+
+			class FStatsSink
+			{
+			public:
+				virtual ~FStatsSink() = default;
+				// Sink only interested in a single type of stat from WebRTC (e.g. track, outbound-rtp).
+				virtual bool Contains(const FName& StatName) const { return Stats.Contains(StatName); };
+				virtual bool ContainsCalcStat(const FName& StatName) const { return CalcStats.Contains(StatName); }
+				virtual void Add(FName StatName, int NDecimalPlaces) { Stats.Add(StatName, FStatData(StatName, 0.0, NDecimalPlaces)); };
+				virtual void AddForCalculation(FName StatName) { CalcStats.Add(StatName, { FStatData(StatName, 0.0, 2), 0.0 }); };
+				virtual FStatData* Get(const FName& StatName) { return Stats.Find(StatName); };
+				virtual FCalculatedStat* GetCalcStat(const FName& StatName) { return CalcStats.Find(StatName); }
+
+			protected:
+				TMap<FName, FStatData> Stats;
+				TMap<FName, FCalculatedStat> CalcStats;
+			};
+
+			class FStreamStatsSink : public FStatsSink
+			{
+			public:
+				FStreamStatsSink(FString InRid);
+				virtual ~FStreamStatsSink() = default;
+
+			private:
+				// Rid is the unique stream identifier.
+				FString Rid;
+			};
+
+		private:
+			FPixelStreamingPlayerId AssociatedPlayerId;
+
+			mutable int32 RefCount;
+
+			const FString RidStr = FString(TEXT("rid"));
+			const FString TrackStatsStr = FString(TEXT("track"));
+			const FString StreamStatsStr = FString(TEXT("outbound-rtp"));
+			const FString SourceStatsStr = FString(TEXT("media-source"));
+
+			uint64 LastCalculationCycles;
+
+			FStatsSink TrackStatsSink;
+			FStatsSink SourceStatsSink;
+
+			bool bIsEnabled;
+
+			// Stream stats are further broken down and stored per RID, as some peers (e.g. an SFU) can have multiple streams.
+			TMap<FString, FStatsSink> StreamStatsSinks;
+		};
+	}
+}
diff --git a/Engine/Plugins/Media/PixelStreaming/Source/PixelStreaming/Private/SessionDescriptionObservers.h b/Engine/Plugins/Media/PixelStreaming/Source/PixelStreaming/Private/SessionDescriptionObservers.h
new file mode 100644
--- /dev/null
+++ b/Engine/Plugins/Media/PixelStreaming/Source/PixelStreaming/Private/SessionDescriptionObservers.h
@@ -0,0 +1,88 @@
+// Copyright Epic Games, Inc. All Rights Reserved.
+
+#pragma once
+
+#include "WebRTCIncludes.h"
+
+#include "Templates/UnrealTemplate.h"
+#include "Templates/Function.h"
+#include "Containers/UnrealString.h"
+
+#include <string>
+#include <memory>
+
+namespace UE {
+	namespace PixelStreaming {
+		//////////////////////////////////////////////////////////////////////////
+		// FSetSessionDescriptionObserver
+		// WebRTC requires an implementation of `webrtc::SetSessionDescriptionObserver` interface as a callback
+		// for setting session description, either on receiving remote `offer` (`PeerConnection::SetRemoteDescription`)
+		// of on sending `answer` (`PeerConnection::SetLocalDescription`)
+		class FSetSessionDescriptionObserver : public webrtc::SetSessionDescriptionObserver
+		{
+		public:
+			using FSuccessCallback = TUniqueFunction<void()>;
+			using FFailureCallback = TUniqueFunction<void(const FString&)>;
+
+			static FSetSessionDescriptionObserver* Create(FSuccessCallback&& successCallback, FFailureCallback&& failureCallback)
+			{
+				return new rtc::RefCountedObject<FSetSessionDescriptionObserver>(MoveTemp(successCallback), MoveTemp(failureCallback));
+			}
+
+			FSetSessionDescriptionObserver(FSuccessCallback&& successCallback, FFailureCallback&& failureCallback)
+				: SuccessCallback(MoveTemp(successCallback))
+				, FailureCallback(MoveTemp(failureCallback))
+			{
+			}
+
+			// we don't need to do anything on success
+			void OnSuccess() override
+			{
+				SuccessCallback();
+			}
+
+			// errors usually mean incompatibility between our session configuration (often H.264, its profile and level) and
+			// player, malformed SDP or if player doesn't support UnifiedPlan (whatever was used by proxy)
+			void OnFailure(webrtc::RTCError Error) override
+			{
+				FailureCallback(FString(UTF8_TO_TCHAR(Error.message())));
+			}
+
+		private:
+			FSuccessCallback SuccessCallback;
+			FFailureCallback FailureCallback;
+		};
+
+		class FCreateSessionDescriptionObserver : public webrtc::CreateSessionDescriptionObserver
+		{
+		public:
+			using FSuccessCallback = TUniqueFunction<void(webrtc::SessionDescriptionInterface*)>;
+			using FFailureCallback = TUniqueFunction<void(const FString&)>;
+
+			static FCreateSessionDescriptionObserver* Create(FSuccessCallback&& successCallback, FFailureCallback&& failureCallback)
+			{
+				return new rtc::RefCountedObject<FCreateSessionDescriptionObserver>(MoveTemp(successCallback), MoveTemp(failureCallback));
+			}
+
+			FCreateSessionDescriptionObserver(FSuccessCallback&& successCallback, FFailureCallback&& failureCallback)
+				: SuccessCallback(MoveTemp(successCallback))
+				, FailureCallback(MoveTemp(failureCallback))
+			{
+			}
+
+			void OnSuccess(webrtc::SessionDescriptionInterface* SDP) override
+			{
+				SuccessCallback(SDP);
+			}
+
+			void OnFailure(webrtc::RTCError Error) override
+			{
+				FailureCallback(ANSI_TO_TCHAR(Error.message()));
+			}
+
+		private:
+			FSuccessCallback SuccessCallback;
+			FFailureCallback FailureCallback;
+		};
+	}
+}
diff --git a/Engine/Plugins/Media/PixelStreaming/Source/PixelStreaming/Private/Settings.cpp b/Engine/Plugins/Media/PixelStreaming/Source/PixelStreaming/Private/Settings.cpp
new file mode 100644
--- /dev/null
+++ b/Engine/Plugins/Media/PixelStreaming/Source/PixelStreaming/Private/Settings.cpp
@@ -0,0 +1,477 @@
+// Copyright Epic Games, Inc. All Rights Reserved.
+
+#include "Settings.h"
+#include "PixelStreamingPrivate.h"
+#include "Misc/DefaultValueHelper.h"
+#include "Async/Async.h"
+#include "IPixelStreamingStatsConsumer.h"
+#include "IPixelStreamingModule.h"
+#include "PixelStreamingDelegates.h"
+
+template <typename T>
+void CommandLineParseValue(const TCHAR* Match, TAutoConsoleVariable<T>& CVar)
+{
+	T Value;
+	if (FParse::Value(FCommandLine::Get(), Match, Value))
+		CVar->Set(Value, ECVF_SetByCommandline);
+};
+
+void CommandLineParseValue(const TCHAR* Match, TAutoConsoleVariable<FString>& CVar, bool bStopOnSeparator = false)
+{
+	FString Value;
+	if (FParse::Value(FCommandLine::Get(), Match, Value, bStopOnSeparator))
+		CVar->Set(*Value, ECVF_SetByCommandline);
+};
+
+void CommandLineParseOption(const TCHAR* Match, TAutoConsoleVariable<bool>& CVar)
+{
+	FString ValueMatch(Match);
+	ValueMatch.Append(TEXT("="));
+	FString Value;
+	if (FParse::Value(FCommandLine::Get(), *ValueMatch, Value))
+	{
+		if (Value.Equals(FString(TEXT("true")), ESearchCase::IgnoreCase))
+		{
+			CVar->Set(true, ECVF_SetByCommandline);
+		}
+		else if (Value.Equals(FString(TEXT("false")), ESearchCase::IgnoreCase))
+		{
+			CVar->Set(false, ECVF_SetByCommandline);
+		}
+	}
+	else if (FParse::Param(FCommandLine::Get(), Match))
+	{
+		CVar->Set(true, ECVF_SetByCommandline);
+	}
+}
+
+namespace UE::PixelStreaming::Settings
+{
+	// Begin Encoder CVars
+	TAutoConsoleVariable<int32> CVarPixelStreamingEncoderTargetBitrate(
+		TEXT("PixelStreaming.Encoder.TargetBitrate"),
+		-1,
+		TEXT("Target bitrate (bps). Ignore the bitrate WebRTC wants (not recommended). Set to -1 to disable. Default -1."),
+		ECVF_RenderThreadSafe);
+
+	TAutoConsoleVariable<int32> CVarPixelStreamingEncoderMaxBitrate(
+		TEXT("PixelStreaming.Encoder.MaxBitrateVBR"),
+		20000000,
+		TEXT("Max bitrate (bps). Does not work in CBR rate control mode with NVENC."),
+		ECVF_RenderThreadSafe);
+
+	TAutoConsoleVariable<bool> CVarPixelStreamingDebugDumpFrame(
+		TEXT("PixelStreaming.Encoder.DumpDebugFrames"),
+		false,
+		TEXT("Dumps frames from the encoder to a file on disk for debugging purposes."),
+		ECVF_Default);
+
+	TAutoConsoleVariable<int32> CVarPixelStreamingEncoderMinQP(
+		TEXT("PixelStreaming.Encoder.MinQP"),
+		0,
+		TEXT("0-51, lower values result in better quality but higher bitrate. Default 0 - i.e. no limit on a minimum QP."),
+		ECVF_Default);
+
+	TAutoConsoleVariable<int32> CVarPixelStreamingEncoderMaxQP(
+		TEXT("PixelStreaming.Encoder.MaxQP"),
+		51,
+		TEXT("0-51, lower values result in better quality but higher bitrate. Default 51 - i.e. no limit on a maximum QP."),
+		ECVF_Default);
+
+	TAutoConsoleVariable<FString> CVarPixelStreamingEncoderRateControl(
+		TEXT("PixelStreaming.Encoder.RateControl"),
+		TEXT("CBR"),
+		TEXT("PixelStreaming video encoder RateControl mode. Supported modes are `ConstQP`, `VBR`, `CBR`. Default: CBR, which we recommend."),
+		ECVF_Default);
+
+	TAutoConsoleVariable<bool> CVarPixelStreamingEnableFillerData(
+		TEXT("PixelStreaming.Encoder.EnableFillerData"),
+		false,
+		TEXT("Maintains constant bitrate by filling with junk data. Note: Should not be required with CBR and MinQP = -1. Default: false."),
+		ECVF_Default);
+
+	TAutoConsoleVariable<FString> CVarPixelStreamingEncoderMultipass(
+		TEXT("PixelStreaming.Encoder.Multipass"),
+		TEXT("FULL"),
+		TEXT("PixelStreaming encoder multipass. Supported modes are `DISABLED`, `QUARTER`, `FULL`"),
+		ECVF_Default);
+
+	TAutoConsoleVariable<FString> CVarPixelStreamingH264Profile(
+		TEXT("PixelStreaming.Encoder.H264Profile"),
+		TEXT("BASELINE"),
+		TEXT("PixelStreaming encoder profile. Supported modes are `AUTO`, `BASELINE`, `MAIN`, `HIGH`, `HIGH444`, `STEREO`, `SVC_TEMPORAL_SCALABILITY`, `PROGRESSIVE_HIGH`, `CONSTRAINED_HIGH`"),
+		ECVF_Default);
+
+	TAutoConsoleVariable<int32> CVarPixelStreamingEncoderKeyframeInterval(
+		TEXT("PixelStreaming.Encoder.KeyframeInterval"),
+		300,
+		TEXT("How many frames before a key frame is sent. Default: 300. Values <=0 will disable sending of periodic key frames. Note: NVENC does not support changing this after encoding has started."),
+		ECVF_Default);
+	// End Encoder CVars
+
+	// Begin Capturer CVars
+	TAutoConsoleVariable<bool> CVarPixelStreamingUseBackBufferCaptureSize(
+		TEXT("PixelStreaming.Capturer.UseBackBufferSize"),
+		true,
+		TEXT("Whether to use back buffer size or custom size"),
+		ECVF_Default);
+
+	TAutoConsoleVariable<FString> CVarPixelStreamingCaptureSize(
+		TEXT("PixelStreaming.Capturer.CaptureSize"),
+		TEXT("1920x1080"),
+		TEXT("Capture size in format widthxheight. Recommended to UseBackBufferSize instead."),
+		ECVF_Default);
+
+	TAutoConsoleVariable<int32> CVarPixelStreamingMaxNumBackBuffers(
+		TEXT("PixelStreaming.Capturer.MaxNumBackBuffers"),
+		3,
+		TEXT("Maximum number of back buffers to use. A value of 0 will not limit the number of back buffers."),
+		ECVF_Default);
+	// End Capturer CVars
+
+	// Begin WebRTC CVars
+	TAutoConsoleVariable<FString> CVarPixelStreamingDegradationPreference(
+		TEXT("PixelStreaming.WebRTC.DegradationPreference"),
+		TEXT("MAINTAIN_FRAMERATE"),
+		TEXT("PixelStreaming degradation preference. Supported modes are `BALANCED`, `MAINTAIN_FRAMERATE`, `MAINTAIN_RESOLUTION`"),
+		ECVF_Default);
+
+	TAutoConsoleVariable<int32> CVarPixelStreamingWebRTCFps(
+		TEXT("PixelStreaming.WebRTC.Fps"),
+		60,
+		TEXT("Framerate for WebRTC encoding. Default: 60"),
+		ECVF_Default);
+
+	//Note: 1 megabit is the maximum allowed in WebRTC for a start bitrate.
+	TAutoConsoleVariable<int32> CVarPixelStreamingWebRTCStartBitrate(
+		TEXT("PixelStreaming.WebRTC.StartBitrate"),
+		1000000,
+		TEXT("Start bitrate (bps) that WebRTC will try begin the stream with. Must be between Min/Max bitrates. Default: 1000000"),
+		ECVF_RenderThreadSafe);
+
+	TAutoConsoleVariable<int32> CVarPixelStreamingWebRTCMinBitrate(
+		TEXT("PixelStreaming.WebRTC.MinBitrate"),
+		100000,
+		TEXT("Min bitrate (bps) that WebRTC will not request below. Careful not to set too high otherwise WebRTC will just drop frames. Default: 100000"),
+		ECVF_RenderThreadSafe);
+
+	TAutoConsoleVariable<int32> CVarPixelStreamingWebRTCMaxBitrate(
+		TEXT("PixelStreaming.WebRTC.MaxBitrate"),
+		20000000,
+		TEXT("Max bitrate (bps) that WebRTC will not request above. Careful not to set too high otherwise because a local (ideal network) will actually reach this. Default: 20000000"),
+		ECVF_RenderThreadSafe);
+
+	TAutoConsoleVariable<int> CVarPixelStreamingWebRTCLowQpThreshold(
+		TEXT("PixelStreaming.WebRTC.LowQpThreshold"),
+		25,
+		TEXT("Only useful when MinQP=-1. Value between 1-51 (default: 25). If WebRTC is getting frames below this QP it will try to increase resolution when not in MAINTAIN_RESOLUTION mode."),
+		ECVF_Default);
+
+	TAutoConsoleVariable<int> CVarPixelStreamingWebRTCHighQpThreshold(
+		TEXT("PixelStreaming.WebRTC.HighQpThreshold"),
+		37,
+		TEXT("Only useful when MinQP=-1. Value between 1-51 (default: 37). If WebRTC is getting frames above this QP it will decrease resolution when not in MAINTAIN_RESOLUTION mode."),
+		ECVF_Default);
+
+	TAutoConsoleVariable<bool> CVarPixelStreamingWebRTCDisableReceiveAudio(
+		TEXT("PixelStreaming.WebRTC.DisableReceiveAudio"),
+		false,
+		TEXT("Disables receiving audio from the browser into UE."),
+		ECVF_Default);
+
+	TAutoConsoleVariable<bool> CVarPixelStreamingWebRTCDisableTransmitAudio(
+		TEXT("PixelStreaming.WebRTC.DisableTransmitAudio"),
+		false,
+		TEXT("Disables transmission of UE audio to the browser."),
+		ECVF_Default);
+
+	TAutoConsoleVariable<bool> CVarPixelStreamingWebRTCDisableAudioSync(
+		TEXT("PixelStreaming.WebRTC.DisableAudioSync"),
+		true,
+		TEXT("Disables the synchronization of audio and video tracks in WebRTC. This can be useful in low latency usecases where synchronization is not required."),
+		ECVF_Default);
+
+	TAutoConsoleVariable<bool> CVarPixelStreamingWebRTCUseLegacyAudioDevice(
+		TEXT("PixelStreaming.WebRTC.UseLegacyAudioDevice"),
+		false,
+		TEXT("Whether put audio and video in the same stream (which will make WebRTC try to sync them)."),
+		ECVF_Default);
+
+	TAutoConsoleVariable<bool> CVarPixelStreamingWebRTCDisableResolutionChange(
+		TEXT("PixelStreaming.WebRTC.DisableResolutionChange"),
+		true,
+		TEXT("Disables whether WebRTC can trigger a resolution a change of the stream (regardless of it wants to)."),
+		ECVF_Default);
+
+	TAutoConsoleVariable<bool> CVarPixelStreamingWebRTCDisableStats(
+		TEXT("PixelStreaming.WebRTC.DisableStats"),
+		false,
+		TEXT("Disables the collection of WebRTC stats."),
+		ECVF_Default);
+
+	// End WebRTC CVars
+
+	// Begin Pixel Streaming Plugin CVars
+
+	TAutoConsoleVariable<bool> CVarPixelStreamingOnScreenStats(
+		TEXT("PixelStreaming.HUDStats"),
+		false,
+		TEXT("Whether to show PixelStreaming stats on the in-game HUD (default: true)."),
+		ECVF_Default);
+
+	TAutoConsoleVariable<bool> CVarPixelStreamingLogStats(
+		TEXT("PixelStreaming.LogStats"),
+		false,
+		TEXT("Whether to show PixelStreaming stats in the log (default: false)."),
+		ECVF_Default);
+
+	TAutoConsoleVariable<int32> CVarPixelStreamingFreezeFrameQuality(
+		TEXT("PixelStreaming.FreezeFrameQuality"),
+		100,
+		TEXT("Compression quality of the freeze frame"),
+		ECVF_Default);
+
+	TAutoConsoleVariable<bool> CVarSendPlayerIdAsInteger(
+		TEXT("PixelStreaming.SendPlayerIdAsInteger"),
+		true,
+		TEXT("If true transmit the player id as an integer (for backward compatibility) or as a string."),
+		ECVF_Default);
+
+	TAutoConsoleVariable<bool> CVarPixelStreamingDisableLatencyTester(
+		TEXT("PixelStreaming.DisableLatencyTester"),
+		false,
+		TEXT("If true disables latency tester being triggerable."),
+		ECVF_Default);
+
+	TAutoConsoleVariable<FString> CVarPixelStreamingKeyFilter(
+		TEXT("PixelStreaming.KeyFilter"),
+		"",
+		TEXT("Comma separated list of keys to ignore from streaming clients."),
+		ECVF_Default);
+
+	TAutoConsoleVariable<bool> CVarPixelStreamingAllowConsoleCommands(
+		TEXT("PixelStreaming.AllowPixelStreamingCommands"),
+		false,
+		TEXT("If true browser can send consoleCommand payloads that execute in UE's console."),
+		ECVF_Default);
+
+	TArray<FKey> FilteredKeys;
+
+	void OnFilteredKeysChanged(IConsoleVariable* Var)
+	{
+		FString CommaList = Var->GetString();
+		TArray<FString> KeyStringArray;
+		CommaList.ParseIntoArray(KeyStringArray, TEXT(","), true);
+		FilteredKeys.Empty();
+		for (auto&& KeyString : KeyStringArray)
+		{
+			FilteredKeys.Add(FKey(*KeyString));
+		}
+	}
+
+	void OnKeyframeIntervalChanged(IConsoleVariable* Var)
+	{
+		AsyncTask(ENamedThreads::GameThread, [Var]() {
+			IConsoleVariable* CVarNVENCKeyframeInterval = IConsoleManager::Get().FindConsoleVariable(TEXT("NVENC.KeyframeInterval"));
+			if (CVarNVENCKeyframeInterval)
+			{
+				CVarNVENCKeyframeInterval->Set(Var->GetInt(), ECVF_SetByCommandline);
+			}
+
+			IConsoleVariable* CVarAMFKeyframeInterval = IConsoleManager::Get().FindConsoleVariable(TEXT("AMF.KeyframeInterval"));
+			if (CVarNVENCKeyframeInterval)
+			{
+				CVarAMFKeyframeInterval->Set(Var->GetInt(), ECVF_SetByCommandline);
+			}
+		});
+	}
+	// Ends Pixel Streaming Plugin CVars
+
+	// Begin utility functions etc.
+	std::map<FString, AVEncoder::FVideoEncoder::RateControlMode> const RateControlCVarMap{
+		{ "ConstQP", AVEncoder::FVideoEncoder::RateControlMode::CONSTQP },
+		{ "VBR", AVEncoder::FVideoEncoder::RateControlMode::VBR },
+		{ "CBR", AVEncoder::FVideoEncoder::RateControlMode::CBR },
+	};
+
+	std::map<FString, AVEncoder::FVideoEncoder::MultipassMode> const MultipassCVarMap{
+		{ "DISABLED", AVEncoder::FVideoEncoder::MultipassMode::DISABLED },
+		{ "QUARTER", AVEncoder::FVideoEncoder::MultipassMode::QUARTER },
+		{ "FULL", AVEncoder::FVideoEncoder::MultipassMode::FULL },
+	};
+
+	std::map<FString, AVEncoder::FVideoEncoder::H264Profile> const H264ProfileMap{
+		{ "AUTO", AVEncoder::FVideoEncoder::H264Profile::AUTO },
+		{ "BASELINE", AVEncoder::FVideoEncoder::H264Profile::BASELINE },
+		{ "MAIN", AVEncoder::FVideoEncoder::H264Profile::MAIN },
+		{ "HIGH", AVEncoder::FVideoEncoder::H264Profile::HIGH },
+		{ "HIGH444", AVEncoder::FVideoEncoder::H264Profile::HIGH444 },
+		{ "STEREO", AVEncoder::FVideoEncoder::H264Profile::STEREO },
+		{ "SVC_TEMPORAL_SCALABILITY", AVEncoder::FVideoEncoder::H264Profile::SVC_TEMPORAL_SCALABILITY },
+		{ "PROGRESSIVE_HIGH", AVEncoder::FVideoEncoder::H264Profile::PROGRESSIVE_HIGH },
+		{ "CONSTRAINED_HIGH", AVEncoder::FVideoEncoder::H264Profile::CONSTRAINED_HIGH },
+	};
+
+	AVEncoder::FVideoEncoder::RateControlMode GetRateControlCVar()
+	{
+		const FString EncoderRateControl = CVarPixelStreamingEncoderRateControl.GetValueOnAnyThread();
+		auto const Iter = RateControlCVarMap.find(EncoderRateControl);
+		if (Iter == std::end(RateControlCVarMap))
+			return AVEncoder::FVideoEncoder::RateControlMode::CBR;
+		return Iter->second;
+	}
+
+	AVEncoder::FVideoEncoder::MultipassMode GetMultipassCVar()
+	{
+		const FString EncoderMultipass = CVarPixelStreamingEncoderMultipass.GetValueOnAnyThread();
+		auto const Iter = MultipassCVarMap.find(EncoderMultipass);
+		if (Iter == std::end(MultipassCVarMap))
+			return AVEncoder::FVideoEncoder::MultipassMode::FULL;
+		return Iter->second;
+	}
+
+	webrtc::DegradationPreference GetDegradationPreference()
+	{
+		FString DegradationPreference = CVarPixelStreamingDegradationPreference.GetValueOnAnyThread();
+		if (DegradationPreference == "MAINTAIN_FRAMERATE")
+		{
+			return webrtc::DegradationPreference::MAINTAIN_FRAMERATE;
+		}
+		else if (DegradationPreference == "MAINTAIN_RESOLUTION")
+		{
+			return webrtc::DegradationPreference::MAINTAIN_RESOLUTION;
+		}
+		// Everything else, return balanced.
+		return webrtc::DegradationPreference::BALANCED;
+	}
+
+	AVEncoder::FVideoEncoder::H264Profile GetH264Profile()
+	{
+		const FString H264Profile = CVarPixelStreamingH264Profile.GetValueOnAnyThread();
+		auto const Iter = H264ProfileMap.find(H264Profile);
+		if (Iter == std::end(H264ProfileMap))
+			return AVEncoder::FVideoEncoder::H264Profile::BASELINE;
+		return Iter->second;
+	}
+	// End utility functions etc.
+
+	FSimulcastParameters SimulcastParameters;
+
+	void ReadSimulcastParameters()
+	{
+		SimulcastParameters.Layers.Empty();
+
+		FString StringOptions;
+		bool bPassedSimulcastParams = FParse::Value(FCommandLine::Get(), TEXT("SimulcastParameters="), StringOptions, false);
+
+		// If no simulcast parameters are passed use some default values
+		if (!bPassedSimulcastParams)
+		{
+			//StringOptions = FString(TEXT("1.0,5000000,20000000,2.0,1000000,5000000,4.0,50000,1000000"));
+			StringOptions = FString(TEXT("1.0,5000000,20000000,2.0,1000000,5000000"));
+		}
+
+		TArray<FString> ParameterArray;
+		StringOptions.ParseIntoArray(ParameterArray, TEXT(","), true);
+		const int OptionCount = ParameterArray.Num();
+		bool bSuccess = OptionCount % 3 == 0;
+		int NextOption = 0;
+		while (bSuccess && ((OptionCount - NextOption) >= 3))
+		{
+			FSimulcastParameters::FLayer Layer;
+			bSuccess = FDefaultValueHelper::ParseFloat(ParameterArray[NextOption++], Layer.Scaling);
+			bSuccess = FDefaultValueHelper::ParseInt(ParameterArray[NextOption++], Layer.MinBitrate);
+			bSuccess = FDefaultValueHelper::ParseInt(ParameterArray[NextOption++], Layer.MaxBitrate);
+			SimulcastParameters.Layers.Add(Layer);
+		}
+
+		if (!bSuccess)
+		{
+			// failed parsing the parameters. just ignore the parameters.
+			UE_LOG(LogPixelStreaming, Error, TEXT("Simulcast parameters malformed. Expected [Scaling_0, MinBitrate_0, MaxBitrate_0, ... , Scaling_N, MinBitrate_N, MaxBitrate_N] as [float, int, int, ... , float, int, int] etc.]"));
+			SimulcastParameters.Layers.Empty();
+		}
+	}
+
+	/*
+	* Stats logger - as turned on/off by CVarPixelStreamingLogStats
+	*/
+	void ConsumeStat(FPixelStreamingPlayerId PlayerId, FName StatName, float StatValue)
+	{
+		UE_LOG(LogPixelStreaming, Log, TEXT("[%s](%s) = %f"), *PlayerId, *StatName.ToString(), StatValue);
+	}
+
+	void OnLogStatsChanged(IConsoleVariable* Var)
+	{
+		bool bLogStats = Var->GetBool();
+		UPixelStreamingDelegates* Delegates = UPixelStreamingDelegates::GetPixelStreamingDelegates();
+		if (bLogStats && Delegates)
+		{
+			Delegates->OnPixelStreamerStatChangedNative.AddStatic(&UE::PixelStreaming::Settings::ConsumeStat);
+		}
+	}
+
+	/*
+	* Settings parsing and initialization.
+	*/
+
+	// Some settings need to be set after streamer is initialized
+	void OnStreamerReady(IPixelStreamingModule& Module)
+	{
+		UE::PixelStreaming::Settings::CVarPixelStreamingLogStats.AsVariable()->SetOnChangedCallback(FConsoleVariableDelegate::CreateStatic(&UE::PixelStreaming::Settings::OnLogStatsChanged));
+		CommandLineParseOption(TEXT("PixelStreamingLogStats"), UE::PixelStreaming::Settings::CVarPixelStreamingLogStats);
+	}
+
+	void InitialiseSettings()
+	{
+		UE_LOG(LogPixelStreaming, Log, TEXT("Initialising Pixel Streaming settings."));
+
+		UE::PixelStreaming::Settings::CVarPixelStreamingKeyFilter.AsVariable()->SetOnChangedCallback(FConsoleVariableDelegate::CreateStatic(&UE::PixelStreaming::Settings::OnFilteredKeysChanged));
+		UE::PixelStreaming::Settings::CVarPixelStreamingEncoderKeyframeInterval.AsVariable()->SetOnChangedCallback(FConsoleVariableDelegate::CreateStatic(&UE::PixelStreaming::Settings::OnKeyframeIntervalChanged));
+
+		// Values parse from commands line
+		CommandLineParseValue(TEXT("PixelStreamingEncoderKeyframeInterval="), UE::PixelStreaming::Settings::CVarPixelStreamingEncoderKeyframeInterval);
+		CommandLineParseValue(TEXT("PixelStreamingEncoderTargetBitrate="), UE::PixelStreaming::Settings::CVarPixelStreamingEncoderTargetBitrate);
+		CommandLineParseValue(TEXT("PixelStreamingEncoderMaxBitrate="), UE::PixelStreaming::Settings::CVarPixelStreamingEncoderMaxBitrate);
+		CommandLineParseValue(TEXT("PixelStreamingEncoderMinQP="), UE::PixelStreaming::Settings::CVarPixelStreamingEncoderMinQP);
+		CommandLineParseValue(TEXT("PixelStreamingEncoderMaxQP="), UE::PixelStreaming::Settings::CVarPixelStreamingEncoderMaxQP);
+		CommandLineParseValue(TEXT("PixelStreamingEncoderRateControl="), UE::PixelStreaming::Settings::CVarPixelStreamingEncoderRateControl);
+		CommandLineParseValue(TEXT("PixelStreamingEncoderMultipass="), UE::PixelStreaming::Settings::CVarPixelStreamingEncoderMultipass);
+		CommandLineParseValue(TEXT("PixelStreamingH264Profile="), UE::PixelStreaming::Settings::CVarPixelStreamingH264Profile);
+		CommandLineParseValue(TEXT("PixelStreamingCaptureSize="), UE::PixelStreaming::Settings::CVarPixelStreamingCaptureSize);
+		CommandLineParseValue(TEXT("PixelStreamingMaxNumBackBuffers="), UE::PixelStreaming::Settings::CVarPixelStreamingMaxNumBackBuffers);
+		CommandLineParseValue(TEXT("PixelStreamingDegradationPreference="), UE::PixelStreaming::Settings::CVarPixelStreamingDegradationPreference);
+		CommandLineParseValue(TEXT("PixelStreamingWebRTCDegradationPreference="), UE::PixelStreaming::Settings::CVarPixelStreamingDegradationPreference);
+		CommandLineParseValue(TEXT("PixelStreamingWebRTCFps="), UE::PixelStreaming::Settings::CVarPixelStreamingWebRTCFps);
+		CommandLineParseValue(TEXT("PixelStreamingWebRTCStartBitrate="), UE::PixelStreaming::Settings::CVarPixelStreamingWebRTCStartBitrate);
+		CommandLineParseValue(TEXT("PixelStreamingWebRTCMinBitrate="), UE::PixelStreaming::Settings::CVarPixelStreamingWebRTCMinBitrate);
+		CommandLineParseValue(TEXT("PixelStreamingWebRTCMaxBitrate="), UE::PixelStreaming::Settings::CVarPixelStreamingWebRTCMaxBitrate);
+		CommandLineParseValue(TEXT("PixelStreamingWebRTCLowQpThreshold="), UE::PixelStreaming::Settings::CVarPixelStreamingWebRTCLowQpThreshold);
+		CommandLineParseValue(TEXT("PixelStreamingWebRTCHighQpThreshold="), UE::PixelStreaming::Settings::CVarPixelStreamingWebRTCHighQpThreshold);
+		CommandLineParseValue(TEXT("PixelStreamingFreezeFrameQuality"), UE::PixelStreaming::Settings::CVarPixelStreamingFreezeFrameQuality);
+		CommandLineParseValue(TEXT("PixelStreamingKeyFilter="), UE::PixelStreaming::Settings::CVarPixelStreamingKeyFilter);
+
+		// Options parse (if these exist they are set to true)
+		CommandLineParseOption(TEXT("AllowPixelStreamingCommands"), UE::PixelStreaming::Settings::CVarPixelStreamingAllowConsoleCommands);
+		CommandLineParseOption(TEXT("PixelStreamingOnScreenStats"), UE::PixelStreaming::Settings::CVarPixelStreamingOnScreenStats);
+		CommandLineParseOption(TEXT("PixelStreamingHudStats"), UE::PixelStreaming::Settings::CVarPixelStreamingOnScreenStats);
+
+		CommandLineParseOption(TEXT("PixelStreamingDebugDumpFrame"), UE::PixelStreaming::Settings::CVarPixelStreamingDebugDumpFrame);
+		CommandLineParseOption(TEXT("PixelStreamingEnableFillerData"), UE::PixelStreaming::Settings::CVarPixelStreamingEnableFillerData);
+		CommandLineParseOption(TEXT("PixelStreamingUseBackBufferSize"), UE::PixelStreaming::Settings::CVarPixelStreamingUseBackBufferCaptureSize);
+		CommandLineParseOption(TEXT("PixelStreamingWebRTCDisableStats"), UE::PixelStreaming::Settings::CVarPixelStreamingWebRTCDisableStats);
+		CommandLineParseOption(TEXT("PixelStreamingWebRTCDisableReceiveAudio"), UE::PixelStreaming::Settings::CVarPixelStreamingWebRTCDisableReceiveAudio);
+		CommandLineParseOption(TEXT("PixelStreamingWebRTCDisableTransmitAudio"), UE::PixelStreaming::Settings::CVarPixelStreamingWebRTCDisableTransmitAudio);
+		CommandLineParseOption(TEXT("PixelStreamingWebRTCDisableAudioSync"), UE::PixelStreaming::Settings::CVarPixelStreamingWebRTCDisableAudioSync);
+		CommandLineParseOption(TEXT("PixelStreamingWebRTCDisableResolutionChange"), UE::PixelStreaming::Settings::CVarPixelStreamingWebRTCDisableResolutionChange);
+		CommandLineParseOption(TEXT("PixelStreamingSendPlayerIdAsInteger"), UE::PixelStreaming::Settings::CVarSendPlayerIdAsInteger);
+		CommandLineParseOption(TEXT("PixelStreamingWebRTCUseLegacyAudioDevice"), UE::PixelStreaming::Settings::CVarPixelStreamingWebRTCUseLegacyAudioDevice);
+		CommandLineParseOption(TEXT("PixelStreamingDisableLatencyTester"), UE::PixelStreaming::Settings::CVarPixelStreamingDisableLatencyTester);
+
+		ReadSimulcastParameters();
+
+		IPixelStreamingModule& Module = IPixelStreamingModule::Get();
+		Module.OnReady().AddStatic(&UE::PixelStreaming::Settings::OnStreamerReady);
+	}
+
+} // namespace UE::PixelStreaming::Settings
diff --git a/Engine/Plugins/Media/PixelStreaming/Source/PixelStreaming/Private/Settings.h b/Engine/Plugins/Media/PixelStreaming/Source/PixelStreaming/Private/Settings.h
new file mode 100644
--- /dev/null
+++ b/Engine/Plugins/Media/PixelStreaming/Source/PixelStreaming/Private/Settings.h
@@ -0,0 +1,123 @@
+// Copyright Epic Games, Inc. All Rights Reserved.
+
+#pragma once
+
+#include "CoreMinimal.h"
+#include "HAL/IConsoleManager.h"
+#include "Misc/CommandLine.h"
+#include "InputCoreTypes.h"
+#include "Misc/CommandLine.h"
+#include "VideoEncoder.h"
+#include "WebRTCIncludes.h"
+
+namespace UE {
+	namespace PixelStreaming {
+		// Console variables (CVars)
+		namespace Settings
+		{
+			extern void InitialiseSettings();
+
+			// Begin Encoder CVars
+			extern TAutoConsoleVariable<int32> CVarPixelStreamingEncoderTargetBitrate;
+			extern TAutoConsoleVariable<int32> CVarPixelStreamingEncoderMaxBitrate;
+			extern TAutoConsoleVariable<bool> CVarPixelStreamingDebugDumpFrame;
+			extern TAutoConsoleVariable<int32> CVarPixelStreamingEncoderMinQP;
+			extern TAutoConsoleVariable<int32> CVarPixelStreamingEncoderMaxQP;
+			extern TAutoConsoleVariable<FString> CVarPixelStreamingEncoderRateControl;
+			extern TAutoConsoleVariable<bool> CVarPixelStreamingEnableFillerData;
+			extern TAutoConsoleVariable<FString> CVarPixelStreamingEncoderMultipass;
+			extern TAutoConsoleVariable<FString> CVarPixelStreamingH264Profile;
+			extern TAutoConsoleVariable<int32> CVarPixelStreamingEncoderKeyframeInterval;
+			// End Encoder CVars
+
+			// Begin Capturer CVars
+			extern TAutoConsoleVariable<bool> CVarPixelStreamingUseBackBufferCaptureSize;
+			extern TAutoConsoleVariable<FString> CVarPixelStreamingCaptureSize;
+			extern TAutoConsoleVariable<int32> CVarPixelStreamingMaxNumBackBuffers;
+			// End Capturer CVars
+
+			// Begin WebRTC CVars
+			extern TAutoConsoleVariable<FString> CVarPixelStreamingDegradationPreference;
+			extern TAutoConsoleVariable<int32> CVarPixelStreamingWebRTCFps;
+			extern TAutoConsoleVariable<int32> CVarPixelStreamingWebRTCStartBitrate;
+			extern TAutoConsoleVariable<int32> CVarPixelStreamingWebRTCMinBitrate;
+			extern TAutoConsoleVariable<int32> CVarPixelStreamingWebRTCMaxBitrate;
+			extern TAutoConsoleVariable<int> CVarPixelStreamingWebRTCLowQpThreshold;
+			extern TAutoConsoleVariable<int> CVarPixelStreamingWebRTCHighQpThreshold;
+			extern TAutoConsoleVariable<bool> CVarPixelStreamingWebRTCDisableReceiveAudio;
+			extern TAutoConsoleVariable<bool> CVarPixelStreamingWebRTCDisableTransmitAudio;
+			extern TAutoConsoleVariable<bool> CVarPixelStreamingWebRTCDisableAudioSync;
+			extern TAutoConsoleVariable<bool> CVarPixelStreamingWebRTCUseLegacyAudioDevice;
+			extern TAutoConsoleVariable<bool> CVarPixelStreamingWebRTCDisableResolutionChange;
+			extern TAutoConsoleVariable<bool> CVarPixelStreamingWebRTCDisableStats;
+			// End WebRTC CVars
+
+			// Begin Pixel Streaming Plugin CVars
+			extern TAutoConsoleVariable<bool> CVarPixelStreamingAllowConsoleCommands;
+			extern TAutoConsoleVariable<bool> CVarPixelStreamingOnScreenStats;
+			extern TAutoConsoleVariable<bool> CVarPixelStreamingLogStats;
+
+			extern TAutoConsoleVariable<int32> CVarPixelStreamingFreezeFrameQuality;
+			extern TAutoConsoleVariable<bool> CVarSendPlayerIdAsInteger;
+			extern TAutoConsoleVariable<bool> CVarPixelStreamingDisableLatencyTester;
+			extern TArray<FKey> FilteredKeys;
+			// Ends Pixel Streaming Plugin CVars
+
+			// Begin utility functions etc.
+			AVEncoder::FVideoEncoder::RateControlMode GetRateControlCVar();
+			AVEncoder::FVideoEncoder::MultipassMode GetMultipassCVar();
+			webrtc::DegradationPreference GetDegradationPreference();
+			AVEncoder::FVideoEncoder::H264Profile GetH264Profile();
+			// End utility functions etc.
+
+			struct FSimulcastParameters
+			{
+				struct FLayer
+				{
+					float Scaling;
+					int MinBitrate;
+					int MaxBitrate;
+				};
+
+				TArray<FLayer> Layers;
+			};
+
+			extern FSimulcastParameters SimulcastParameters;
+
+			// Begin Command line args
+
+			inline bool IsPixelStreamingHideCursor()
+			{
+				return FParse::Param(FCommandLine::Get(), TEXT("PixelStreamingHideCursor"));
+			}
+
+			inline bool IsForceVP8()
+			{
+				return FParse::Param(FCommandLine::Get(), TEXT("PSForceVP8"));
+			}
+
+			inline bool GetSignallingServerIP(FString& OutSignallingServerIP)
+			{
+				return FParse::Value(FCommandLine::Get(), TEXT("PixelStreamingIP="), OutSignallingServerIP);
+			}
+
+			inline bool GetSignallingServerPort(uint16& OutSignallingServerPort)
+			{
+				return FParse::Value(FCommandLine::Get(), TEXT("PixelStreamingPort="), OutSignallingServerPort);
+			}
+
+			inline bool GetControlScheme(FString& OutControlScheme)
+			{
+				return FParse::Value(FCommandLine::Get(), TEXT("PixelStreamingControlScheme="), OutControlScheme);
+			}
+
+			inline bool GetFastPan(float& OutFastPan)
+			{
+				return FParse::Value(FCommandLine::Get(), TEXT("PixelStreamingFastPan="), OutFastPan);
+			}
+
+			// End Command line args
+
+		} // namespace UE::PixelStreaming::Settings
+	}
+}
diff --git a/Engine/Plugins/Media/PixelStreaming/Source/PixelStreaming/Private/SignallingServerConnection.cpp b/Engine/Plugins/Media/PixelStreaming/Source/PixelStreaming/Private/SignallingServerConnection.cpp
--- a/Engine/Plugins/Media/PixelStreaming/Source/PixelStreaming/Private/SignallingServerConnection.cpp
+++ b/Engine/Plugins/Media/PixelStreaming/Source/PixelStreaming/Private/SignallingServerConnection.cpp
@@ -1,7 +1,7 @@
 // Copyright Epic Games, Inc. All Rights Reserved.
 
 #include "SignallingServerConnection.h"
-#include "Utils.h"
+#include "ToStringExtensions.h"
 
 #include "WebSocketsModule.h"
 #include "IWebSocket.h"
@@ -11,9 +11,9 @@
 #include "Policies/CondensedJsonPrintPolicy.h"
 #include "Misc/AssertionMacros.h"
 #include "Logging/LogMacros.h"
-#include "PixelStreamingSettings.h"
+#include "Settings.h"
 #include "TimerManager.h"
-#include "PixelStreamerDelegates.h"
+#include "PixelStreamingDelegates.h"
 #include "PlayerSession.h"
 
 DECLARE_LOG_CATEGORY_EXTERN(LogPixelStreamingSS, Log, VeryVerbose);
@@ -38,12 +38,12 @@ DEFINE_LOG_CATEGORY(LogPixelStreamingSS);
 	}                                                                                               \
 	while (false);
 
-FSignallingServerConnection::FSignallingServerConnection(FSignallingServerConnectionObserver& InObserver, const FString& InStreamerId)
+UE::PixelStreaming::FSignallingServerConnection::FSignallingServerConnection(UE::PixelStreaming::FSignallingServerConnectionObserver& InObserver, const FString& InStreamerId)
 	: Observer(InObserver), StreamerId(InStreamerId)
 {
 }
 
-void FSignallingServerConnection::Connect(const FString& Url)
+void UE::PixelStreaming::FSignallingServerConnection::Connect(const FString& Url)
 {
 	// Already have a websocket connection, no need to make another one
 	if (WS)
@@ -62,7 +62,7 @@ void FSignallingServerConnection::Connect(const FString& Url)
 	WS->Connect();
 }
 
-void FSignallingServerConnection::Disconnect()
+void UE::PixelStreaming::FSignallingServerConnection::Disconnect()
 {
 	if (!WS)
 	{
@@ -83,12 +83,12 @@ void FSignallingServerConnection::Disconnect()
 	WS = nullptr;
 }
 
-FSignallingServerConnection::~FSignallingServerConnection()
+UE::PixelStreaming::FSignallingServerConnection::~FSignallingServerConnection()
 {
-	this->Disconnect();
+	Disconnect();
 }
 
-void FSignallingServerConnection::SendOffer(FPlayerId PlayerId, const webrtc::SessionDescriptionInterface& SDP)
+void UE::PixelStreaming::FSignallingServerConnection::SendOffer(FPixelStreamingPlayerId PlayerId, const webrtc::SessionDescriptionInterface& SDP)
 {
 	FJsonObjectPtr OfferJson = MakeShared<FJsonObject>();
 	OfferJson->SetStringField(TEXT("type"), TEXT("offer"));
@@ -104,7 +104,7 @@ void FSignallingServerConnection::SendOffer(FPlayerId PlayerId, const webrtc::Se
 	WS->Send(ToString(OfferJson, false));
 }
 
-void FSignallingServerConnection::SendAnswer(FPlayerId PlayerId, const webrtc::SessionDescriptionInterface& SDP)
+void UE::PixelStreaming::FSignallingServerConnection::SendAnswer(FPixelStreamingPlayerId PlayerId, const webrtc::SessionDescriptionInterface& SDP)
 {
 	FJsonObjectPtr AnswerJson = MakeShared<FJsonObject>();
 	AnswerJson->SetStringField(TEXT("type"), TEXT("answer"));
@@ -120,9 +120,9 @@ void FSignallingServerConnection::SendAnswer(FPlayerId PlayerId, const webrtc::S
 	WS->Send(ToString(AnswerJson, false));
 }
 
-void FSignallingServerConnection::SetPlayerIdJson(FJsonObjectPtr& JsonObject, FPlayerId PlayerId)
+void UE::PixelStreaming::FSignallingServerConnection::SetPlayerIdJson(FJsonObjectPtr& JsonObject, FPixelStreamingPlayerId PlayerId)
 {
-	bool bSendAsInteger = PixelStreamingSettings::CVarSendPlayerIdAsInteger.GetValueOnAnyThread();
+	bool bSendAsInteger = UE::PixelStreaming::Settings::CVarSendPlayerIdAsInteger.GetValueOnAnyThread();
 	if (bSendAsInteger)
 	{
 		int32 PlayerIdAsInt = PlayerIdToInt(PlayerId);
@@ -134,9 +134,9 @@ void FSignallingServerConnection::SetPlayerIdJson(FJsonObjectPtr& JsonObject, FP
 	}
 }
 
-bool FSignallingServerConnection::GetPlayerIdJson(const FJsonObjectPtr& Json, FPlayerId& OutPlayerId)
+bool UE::PixelStreaming::FSignallingServerConnection::GetPlayerIdJson(const FJsonObjectPtr& Json, FPixelStreamingPlayerId& OutPlayerId)
 {
-	bool bSendAsInteger = PixelStreamingSettings::CVarSendPlayerIdAsInteger.GetValueOnAnyThread();
+	bool bSendAsInteger = UE::PixelStreaming::Settings::CVarSendPlayerIdAsInteger.GetValueOnAnyThread();
 	if (bSendAsInteger)
 	{
 		uint32 PlayerIdInt;
@@ -155,7 +155,7 @@ bool FSignallingServerConnection::GetPlayerIdJson(const FJsonObjectPtr& Json, FP
 	return false;
 }
 
-void FSignallingServerConnection::SendIceCandidate(const webrtc::IceCandidateInterface& IceCandidate)
+void UE::PixelStreaming::FSignallingServerConnection::SendIceCandidate(const webrtc::IceCandidateInterface& IceCandidate)
 {
 	FJsonObjectPtr IceCandidateJson = MakeShared<FJsonObject>();
 
@@ -177,12 +177,12 @@ void FSignallingServerConnection::SendIceCandidate(const webrtc::IceCandidateInt
 	WS->Send(ToString(IceCandidateJson, false));
 }
 
-void FSignallingServerConnection::SendIceCandidate(FPlayerId PlayerId, const webrtc::IceCandidateInterface& IceCandidate)
+void UE::PixelStreaming::FSignallingServerConnection::SendIceCandidate(FPixelStreamingPlayerId PlayerId, const webrtc::IceCandidateInterface& IceCandidate)
 {
 	FJsonObjectPtr IceCandidateJson = MakeShared<FJsonObject>();
 
 	IceCandidateJson->SetStringField(TEXT("type"), TEXT("iceCandidate"));
-	this->SetPlayerIdJson(IceCandidateJson, PlayerId);
+	SetPlayerIdJson(IceCandidateJson, PlayerId);
 
 	FJsonObjectPtr CandidateJson = MakeShared<FJsonObject>();
 	CandidateJson->SetStringField(TEXT("sdpMid"), IceCandidate.sdp_mid().c_str());
@@ -200,7 +200,7 @@ void FSignallingServerConnection::SendIceCandidate(FPlayerId PlayerId, const web
 	WS->Send(ToString(IceCandidateJson, false));
 }
 
-void FSignallingServerConnection::KeepAlive()
+void UE::PixelStreaming::FSignallingServerConnection::KeepAlive()
 {
 	FJsonObjectPtr Json = MakeShared<FJsonObject>();
 	double unixTime = FDateTime::UtcNow().ToUnixTimestamp();
@@ -213,12 +213,12 @@ void FSignallingServerConnection::KeepAlive()
 	}
 }
 
-void FSignallingServerConnection::SendDisconnectPlayer(FPlayerId PlayerId, const FString& Reason)
+void UE::PixelStreaming::FSignallingServerConnection::SendDisconnectPlayer(FPixelStreamingPlayerId PlayerId, const FString& Reason)
 {
 	FJsonObjectPtr Json = MakeShared<FJsonObject>();
 
 	Json->SetStringField(TEXT("type"), TEXT("disconnectPlayer"));
-	this->SetPlayerIdJson(Json, PlayerId);
+	SetPlayerIdJson(Json, PlayerId);
 	Json->SetStringField(TEXT("reason"), Reason);
 
 	FString Msg = ToString(Json, false);
@@ -227,49 +227,49 @@ void FSignallingServerConnection::SendDisconnectPlayer(FPlayerId PlayerId, const
 	WS->Send(Msg);
 }
 
-void FSignallingServerConnection::OnConnected()
+void UE::PixelStreaming::FSignallingServerConnection::OnConnected()
 {
 	UE_LOG(LogPixelStreamingSS, Log, TEXT("Connected to SS"));
 
 	//Send message to keep connection alive every 60 seconds
-	GWorld->GetTimerManager().SetTimer(TimerHandle_KeepAlive, std::bind(&FSignallingServerConnection::KeepAlive, this), KEEP_ALIVE_INTERVAL, true);
+	GWorld->GetTimerManager().SetTimer(TimerHandle_KeepAlive, std::bind(&UE::PixelStreaming::FSignallingServerConnection::KeepAlive, this), KEEP_ALIVE_INTERVAL, true);
 
-	if (UPixelStreamerDelegates* Delegates = UPixelStreamerDelegates::GetPixelStreamerDelegates())
+	if (UPixelStreamingDelegates* Delegates = UPixelStreamingDelegates::GetPixelStreamingDelegates())
 	{
 		Delegates->OnConnecedToSignallingServer.Broadcast();
 	}
 }
 
-void FSignallingServerConnection::OnConnectionError(const FString& Error)
+void UE::PixelStreaming::FSignallingServerConnection::OnConnectionError(const FString& Error)
 {
 	UE_LOG(LogPixelStreamingSS, Error, TEXT("Failed to connect to SS: %s"), *Error);
 	Observer.OnSignallingServerDisconnected();
 	GWorld->GetTimerManager().ClearTimer(TimerHandle_KeepAlive);
 
-	if (UPixelStreamerDelegates* Delegates = UPixelStreamerDelegates::GetPixelStreamerDelegates())
+	if (UPixelStreamingDelegates* Delegates = UPixelStreamingDelegates::GetPixelStreamingDelegates())
 	{
 		Delegates->OnDisconnectedFromSignallingServer.Broadcast();
 	}
 }
 
-void FSignallingServerConnection::OnClosed(int32 StatusCode, const FString& Reason, bool bWasClean)
+void UE::PixelStreaming::FSignallingServerConnection::OnClosed(int32 StatusCode, const FString& Reason, bool bWasClean)
 {
 	UE_LOG(LogPixelStreamingSS, Log, TEXT("Connection to SS closed: \n\tstatus %d\n\treason: %s\n\twas clean: %s"), StatusCode, *Reason, bWasClean ? TEXT("true") : TEXT("false"));
 	Observer.OnSignallingServerDisconnected();
 	GWorld->GetTimerManager().ClearTimer(TimerHandle_KeepAlive);
 
-	if (UPixelStreamerDelegates* Delegates = UPixelStreamerDelegates::GetPixelStreamerDelegates())
+	if (UPixelStreamingDelegates* Delegates = UPixelStreamingDelegates::GetPixelStreamingDelegates())
 	{
 		Delegates->OnDisconnectedFromSignallingServer.Broadcast();
 	}
 }
 
-void FSignallingServerConnection::OnMessage(const FString& Msg)
+void UE::PixelStreaming::FSignallingServerConnection::OnMessage(const FString& Msg)
 {
 	UE_LOG(LogPixelStreamingSS, Log, TEXT("<- SS: %s"), *Msg);
 
 	TSharedPtr<FJsonObject> JsonMsg;
-	auto JsonReader = TJsonReaderFactory<TCHAR>::Create(Msg);
+	TSharedRef<TJsonReader<TCHAR>> JsonReader = TJsonReaderFactory<TCHAR>::Create(Msg);
 	if (!FJsonSerializer::Deserialize(JsonReader, JsonMsg))
 	{
 		HANDLE_SS_ERROR(TEXT("Failed to parse SS message:\n%s"), *Msg);
@@ -327,7 +327,7 @@ void FSignallingServerConnection::OnMessage(const FString& Msg)
 	}
 }
 
-void FSignallingServerConnection::OnIdRequested()
+void UE::PixelStreaming::FSignallingServerConnection::OnIdRequested()
 {
 	FJsonObjectPtr Json = MakeShared<FJsonObject>();
 
@@ -340,7 +340,7 @@ void FSignallingServerConnection::OnIdRequested()
 	WS->Send(Msg);
 }
 
-void FSignallingServerConnection::OnConfig(const FJsonObjectPtr& Json)
+void UE::PixelStreaming::FSignallingServerConnection::OnConfig(const FJsonObjectPtr& Json)
 {
 	// SS sends `config` that looks like:
 	// `{peerConnectionOptions: { 'iceServers': [{'urls': ['stun:34.250.222.95:19302', 'turn:34.250.222.95:19303']}] }}`
@@ -409,7 +409,7 @@ void FSignallingServerConnection::OnConfig(const FJsonObjectPtr& Json)
 	Observer.OnConfig(RTCConfig);
 }
 
-void FSignallingServerConnection::OnSessionDescription(const FJsonObjectPtr& Json)
+void UE::PixelStreaming::FSignallingServerConnection::OnSessionDescription(const FJsonObjectPtr& Json)
 {
 	webrtc::SdpType Type = Json->GetStringField(TEXT("type")) == TEXT("offer") ? webrtc::SdpType::kOffer : webrtc::SdpType::kAnswer;
 
@@ -419,8 +419,8 @@ void FSignallingServerConnection::OnSessionDescription(const FJsonObjectPtr& Jso
 		HANDLE_SS_ERROR(TEXT("Cannot find `sdp` in Streamer's answer\n%s"), *ToString(Json));
 	}
 
-	FPlayerId PlayerId;
-	bool bGotPlayerId = this->GetPlayerIdJson(Json, PlayerId);
+	FPixelStreamingPlayerId PlayerId;
+	bool bGotPlayerId = GetPlayerIdJson(Json, PlayerId);
 	if (!bGotPlayerId)
 	{
 		HANDLE_SS_ERROR(TEXT("Failed to get `playerId` from `offer/answer` message\n%s"), *ToString(Json));
@@ -429,7 +429,7 @@ void FSignallingServerConnection::OnSessionDescription(const FJsonObjectPtr& Jso
 	Observer.OnSessionDescription(PlayerId, Type, Sdp);
 }
 
-void FSignallingServerConnection::OnStreamerIceCandidate(const FJsonObjectPtr& Json)
+void UE::PixelStreaming::FSignallingServerConnection::OnStreamerIceCandidate(const FJsonObjectPtr& Json)
 {
 	const FJsonObjectPtr* CandidateJson;
 	if (!Json->TryGetObjectField(TEXT("candidate"), CandidateJson))
@@ -465,10 +465,10 @@ void FSignallingServerConnection::OnStreamerIceCandidate(const FJsonObjectPtr& J
 	Observer.OnRemoteIceCandidate(TUniquePtr<webrtc::IceCandidateInterface>{ Candidate.release() });
 }
 
-void FSignallingServerConnection::OnPlayerIceCandidate(const FJsonObjectPtr& Json)
+void UE::PixelStreaming::FSignallingServerConnection::OnPlayerIceCandidate(const FJsonObjectPtr& Json)
 {
-	FPlayerId PlayerId;
-	bool bGotPlayerId = this->GetPlayerIdJson(Json, PlayerId);
+	FPixelStreamingPlayerId PlayerId;
+	bool bGotPlayerId = GetPlayerIdJson(Json, PlayerId);
 	if (!bGotPlayerId)
 	{
 		HANDLE_PLAYER_SS_ERROR(PlayerId, TEXT("Failed to get `playerId` from remote `iceCandidate` message\n%s"), *ToString(Json));
@@ -501,7 +501,7 @@ void FSignallingServerConnection::OnPlayerIceCandidate(const FJsonObjectPtr& Jso
 	Observer.OnRemoteIceCandidate(PlayerId, to_string(SdpMid), SdpMLineIndex, to_string(CandidateStr));
 }
 
-void FSignallingServerConnection::OnPlayerCount(const FJsonObjectPtr& Json)
+void UE::PixelStreaming::FSignallingServerConnection::OnPlayerCount(const FJsonObjectPtr& Json)
 {
 	uint32 Count;
 	if (!Json->TryGetNumberField(TEXT("count"), Count))
@@ -512,10 +512,10 @@ void FSignallingServerConnection::OnPlayerCount(const FJsonObjectPtr& Json)
 	Observer.OnPlayerCount(Count);
 }
 
-void FSignallingServerConnection::OnPlayerConnected(const FJsonObjectPtr& Json)
+void UE::PixelStreaming::FSignallingServerConnection::OnPlayerConnected(const FJsonObjectPtr& Json)
 {
-	FPlayerId PlayerId;
-	bool bGotPlayerId = this->GetPlayerIdJson(Json, PlayerId);
+	FPixelStreamingPlayerId PlayerId;
+	bool bGotPlayerId = GetPlayerIdJson(Json, PlayerId);
 	if (!bGotPlayerId)
 	{
 		HANDLE_SS_ERROR(TEXT("Failed to get `playerId` from `join` message\n%s"), *ToString(Json));
@@ -530,15 +530,15 @@ void FSignallingServerConnection::OnPlayerConnected(const FJsonObjectPtr& Json)
 	bool bIsSFU = false;
 	Json->TryGetBoolField(TEXT("sfu"), bIsSFU);
 
-	Flags |= bMakeDataChannel ? PixelStreamingProtocol::EPixelStreamingPlayerFlags::PSPFlag_SupportsDataChannel : 0;
-	Flags |= bIsSFU ? PixelStreamingProtocol::EPixelStreamingPlayerFlags::PSPFlag_IsSFU : 0;
+	Flags |= bMakeDataChannel ? UE::PixelStreaming::Protocol::EPlayerFlags::PSPFlag_SupportsDataChannel : 0;
+	Flags |= bIsSFU ? UE::PixelStreaming::Protocol::EPlayerFlags::PSPFlag_IsSFU : 0;
 	Observer.OnPlayerConnected(PlayerId, Flags);
 }
 
-void FSignallingServerConnection::OnPlayerDisconnected(const FJsonObjectPtr& Json)
+void UE::PixelStreaming::FSignallingServerConnection::OnPlayerDisconnected(const FJsonObjectPtr& Json)
 {
-	FPlayerId PlayerId;
-	bool bGotPlayerId = this->GetPlayerIdJson(Json, PlayerId);
+	FPixelStreamingPlayerId PlayerId;
+	bool bGotPlayerId = GetPlayerIdJson(Json, PlayerId);
 	if (!bGotPlayerId)
 	{
 		HANDLE_SS_ERROR(TEXT("Failed to get `playerId` from `playerDisconnected` message\n%s"), *ToString(Json));
diff --git a/Engine/Plugins/Media/PixelStreaming/Source/PixelStreaming/Private/SignallingServerConnection.h b/Engine/Plugins/Media/PixelStreaming/Source/PixelStreaming/Private/SignallingServerConnection.h
--- a/Engine/Plugins/Media/PixelStreaming/Source/PixelStreaming/Private/SignallingServerConnection.h
+++ b/Engine/Plugins/Media/PixelStreaming/Source/PixelStreaming/Private/SignallingServerConnection.h
@@ -11,97 +11,102 @@
 #include "Dom/JsonObject.h"
 #include "Delegates/IDelegateInstance.h"
 #include "Engine/EngineTypes.h"
-#include "PlayerId.h"
+#include "PixelStreamingPlayerId.h"
 
 class IWebSocket;
 
-// callback interface for `FSignallingServerConnection`
-class FSignallingServerConnectionObserver
-{
-public:
-	virtual ~FSignallingServerConnectionObserver()
-	{
+namespace UE {
+	namespace PixelStreaming {
+
+		// callback interface for `FSignallingServerConnection`
+		class FSignallingServerConnectionObserver
+		{
+		public:
+			virtual ~FSignallingServerConnectionObserver()
+			{
+			}
+
+			virtual void OnSignallingServerDisconnected() = 0;
+			virtual void OnConfig(const webrtc::PeerConnectionInterface::RTCConfiguration& Config) = 0;
+			virtual void OnSessionDescription(FPixelStreamingPlayerId PlayerId, webrtc::SdpType Type, const FString& Sdp)
+			{
+				unimplemented();
+			}
+
+			// Streamer-only
+			virtual void OnRemoteIceCandidate(FPixelStreamingPlayerId PlayerId, const std::string& SdpMid, int SdpMLineIndex, const std::string& Sdp)
+			{
+				unimplemented();
+			}
+			virtual void OnPlayerConnected(FPixelStreamingPlayerId PlayerId, int Flags)
+			{
+				unimplemented();
+			}
+			virtual void OnPlayerDisconnected(FPixelStreamingPlayerId PlayerId)
+			{
+				unimplemented();
+			}
+
+			// Player-only
+			virtual void OnRemoteIceCandidate(TUniquePtr<webrtc::IceCandidateInterface> Candidate)
+			{
+				unimplemented();
+			}
+			virtual void OnPlayerCount(uint32 Count)
+			{
+				unimplemented();
+			}
+		};
+
+		class FSignallingServerConnection final
+		{
+		public:
+			explicit FSignallingServerConnection(FSignallingServerConnectionObserver& Observer, const FString& StreamerId);
+			~FSignallingServerConnection();
+			void Connect(const FString& Url);
+			void Disconnect();
+
+			void SendOffer(FPixelStreamingPlayerId PlayerId, const webrtc::SessionDescriptionInterface& SDP);
+			void SendAnswer(FPixelStreamingPlayerId PlayerId, const webrtc::SessionDescriptionInterface& SDP);
+			void SendIceCandidate(const webrtc::IceCandidateInterface& IceCandidate);
+			void SendIceCandidate(FPixelStreamingPlayerId PlayerId, const webrtc::IceCandidateInterface& IceCandidate);
+			void SendDisconnectPlayer(FPixelStreamingPlayerId PlayerId, const FString& Reason);
+
+		private:
+			void KeepAlive();
+
+			void OnConnected();
+			void OnConnectionError(const FString& Error);
+			void OnClosed(int32 StatusCode, const FString& Reason, bool bWasClean);
+			void OnMessage(const FString& Msg);
+
+			using FJsonObjectPtr = TSharedPtr<FJsonObject>;
+
+			void OnIdRequested();
+			void OnConfig(const FJsonObjectPtr& Json);
+			void OnSessionDescription(const FJsonObjectPtr& Json);
+			void OnStreamerIceCandidate(const FJsonObjectPtr& Json);
+			void OnPlayerIceCandidate(const FJsonObjectPtr& Json);
+			void OnPlayerCount(const FJsonObjectPtr& Json);
+			void OnPlayerConnected(const FJsonObjectPtr& Json);
+			void OnPlayerDisconnected(const FJsonObjectPtr& Json);
+			void SetPlayerIdJson(FJsonObjectPtr& JsonObject, FPixelStreamingPlayerId PlayerId);
+			bool GetPlayerIdJson(const FJsonObjectPtr& Json, FPixelStreamingPlayerId& OutPlayerId);
+
+		private:
+			FSignallingServerConnectionObserver& Observer;
+			FString StreamerId;
+
+			FDelegateHandle OnConnectedHandle;
+			FDelegateHandle OnConnectionErrorHandle;
+			FDelegateHandle OnClosedHandle;
+			FDelegateHandle OnMessageHandle;
+
+			/** Handle for efficient management of KeepAlive timer */
+			FTimerHandle TimerHandle_KeepAlive;
+			const float KEEP_ALIVE_INTERVAL = 60.0f;
+
+			TSharedPtr<IWebSocket> WS;
+		};
 	}
-
-	virtual void OnSignallingServerDisconnected() = 0;
-	virtual void OnConfig(const webrtc::PeerConnectionInterface::RTCConfiguration& Config) = 0;
-	virtual void OnSessionDescription(FPlayerId PlayerId, webrtc::SdpType Type, const FString& Sdp)
-	{
-		unimplemented();
-	}
-
-	// Streamer-only
-	virtual void OnRemoteIceCandidate(FPlayerId PlayerId, const std::string& SdpMid, int SdpMLineIndex, const std::string& Sdp)
-	{
-		unimplemented();
-	}
-	virtual void OnPlayerConnected(FPlayerId PlayerId, int Flags)
-	{
-		unimplemented();
-	}
-	virtual void OnPlayerDisconnected(FPlayerId PlayerId)
-	{
-		unimplemented();
-	}
-
-	// Player-only
-	virtual void OnRemoteIceCandidate(TUniquePtr<webrtc::IceCandidateInterface> Candidate)
-	{
-		unimplemented();
-	}
-	virtual void OnPlayerCount(uint32 Count)
-	{
-		unimplemented();
-	}
-};
-
-class FSignallingServerConnection final
-{
-public:
-	explicit FSignallingServerConnection(FSignallingServerConnectionObserver& Observer, const FString& StreamerId);
-	~FSignallingServerConnection();
-	void Connect(const FString& Url);
-	void Disconnect();
-
-	void SendOffer(FPlayerId PlayerId, const webrtc::SessionDescriptionInterface& SDP);
-	void SendAnswer(FPlayerId PlayerId, const webrtc::SessionDescriptionInterface& SDP);
-	void SendIceCandidate(const webrtc::IceCandidateInterface& IceCandidate);
-	void SendIceCandidate(FPlayerId PlayerId, const webrtc::IceCandidateInterface& IceCandidate);
-	void SendDisconnectPlayer(FPlayerId PlayerId, const FString& Reason);
-
-private:
-	void KeepAlive();
-
-	void OnConnected();
-	void OnConnectionError(const FString& Error);
-	void OnClosed(int32 StatusCode, const FString& Reason, bool bWasClean);
-	void OnMessage(const FString& Msg);
-
-	using FJsonObjectPtr = TSharedPtr<FJsonObject>;
-
-	void OnIdRequested();
-	void OnConfig(const FJsonObjectPtr& Json);
-	void OnSessionDescription(const FJsonObjectPtr& Json);
-	void OnStreamerIceCandidate(const FJsonObjectPtr& Json);
-	void OnPlayerIceCandidate(const FJsonObjectPtr& Json);
-	void OnPlayerCount(const FJsonObjectPtr& Json);
-	void OnPlayerConnected(const FJsonObjectPtr& Json);
-	void OnPlayerDisconnected(const FJsonObjectPtr& Json);
-	void SetPlayerIdJson(FJsonObjectPtr& JsonObject, FPlayerId PlayerId);
-	bool GetPlayerIdJson(const FJsonObjectPtr& Json, FPlayerId& OutPlayerId);
-
-private:
-	FSignallingServerConnectionObserver& Observer;
-	FString StreamerId;
-
-	FDelegateHandle OnConnectedHandle;
-	FDelegateHandle OnConnectionErrorHandle;
-	FDelegateHandle OnClosedHandle;
-	FDelegateHandle OnMessageHandle;
-
-	/** Handle for efficient management of KeepAlive timer */
-	FTimerHandle TimerHandle_KeepAlive;
-	const float KEEP_ALIVE_INTERVAL = 60.0f;
-
-	TSharedPtr<IWebSocket> WS;
-};
+}
diff --git a/Engine/Plugins/Media/PixelStreaming/Source/PixelStreaming/Private/SimulcastEncoderAdapter.cpp b/Engine/Plugins/Media/PixelStreaming/Source/PixelStreaming/Private/SimulcastEncoderAdapter.cpp
new file mode 100644
--- /dev/null
+++ b/Engine/Plugins/Media/PixelStreaming/Source/PixelStreaming/Private/SimulcastEncoderAdapter.cpp
@@ -0,0 +1,475 @@
+// Copyright Epic Games, Inc. All Rights Reserved.
+#include "SimulcastEncoderAdapter.h"
+#include "FrameBuffer.h"
+
+namespace
+{
+	uint32_t SumStreamMaxBitrate(int Streams, const webrtc::VideoCodec& Codec)
+	{
+		uint32_t BitrateSum = 0;
+		for (int i = 0; i < Streams; ++i)
+		{
+			BitrateSum += Codec.simulcastStream[i].maxBitrate;
+		}
+		return BitrateSum;
+	}
+
+	int GetNumberOfStreams(const webrtc::VideoCodec& Codec)
+	{
+		int Streams = Codec.numberOfSimulcastStreams < 1 ? 1 : Codec.numberOfSimulcastStreams;
+		uint32_t SimulcastMaxBitrate = SumStreamMaxBitrate(Streams, Codec);
+		if (SimulcastMaxBitrate == 0)
+		{
+			Streams = 1;
+		}
+		return Streams;
+	}
+
+	int GetNumActiveStreams(const webrtc::VideoCodec& Codec)
+	{
+		int NumConfiguredStreams = GetNumberOfStreams(Codec);
+		int NumActiveStreams = 0;
+		for (int i = 0; i < NumConfiguredStreams; ++i)
+		{
+			if (Codec.simulcastStream[i].active)
+			{
+				++NumActiveStreams;
+			}
+		}
+		return NumActiveStreams;
+	}
+
+	bool StreamResolutionCompare(const webrtc::SimulcastStream& A, const webrtc::SimulcastStream& B)
+	{
+		return std::tie(A.height, A.width, A.maxBitrate, A.maxFramerate) < std::tie(B.height, B.width, B.maxBitrate, B.maxFramerate);
+	}
+
+	void PopulateStreamCodec(const webrtc::VideoCodec& CodecSettings, int StreamIndex, uint32_t StartBitrateKbps, webrtc::VideoCodec* StreamCodec)
+	{
+		*StreamCodec = CodecSettings;
+
+		// Stream specific settings.
+		StreamCodec->numberOfSimulcastStreams = 0;
+		StreamCodec->width = CodecSettings.simulcastStream[StreamIndex].width;
+		StreamCodec->height = CodecSettings.simulcastStream[StreamIndex].height;
+		StreamCodec->maxBitrate = CodecSettings.simulcastStream[StreamIndex].maxBitrate;
+		StreamCodec->minBitrate = CodecSettings.simulcastStream[StreamIndex].minBitrate;
+		StreamCodec->maxFramerate = CodecSettings.simulcastStream[StreamIndex].maxFramerate;
+		StreamCodec->qpMax = CodecSettings.simulcastStream[StreamIndex].qpMax;
+		StreamCodec->active = CodecSettings.simulcastStream[StreamIndex].active;
+		if (CodecSettings.codecType == webrtc::kVideoCodecH264)
+		{
+			StreamCodec->H264()->numberOfTemporalLayers = CodecSettings.simulcastStream[StreamIndex].numberOfTemporalLayers;
+		}
+		StreamCodec->startBitrate = StartBitrateKbps;
+	}
+
+	// An EncodedImageCallback implementation that forwards on calls to a
+	// UE::PixelStreaming::FSimulcastEncoderAdapter, but with the stream index it's registered with as
+	// the first parameter to Encoded.
+	class AdapterEncodedImageCallback : public webrtc::EncodedImageCallback
+	{
+	public:
+		AdapterEncodedImageCallback(UE::PixelStreaming::FSimulcastEncoderAdapter* adapter,
+			size_t stream_idx)
+			: adapter_(adapter), stream_idx_(stream_idx)
+		{
+		}
+
+		EncodedImageCallback::Result OnEncodedImage(
+			const webrtc::EncodedImage& encoded_image,
+			const webrtc::CodecSpecificInfo* codec_specific_info,
+			const webrtc::RTPFragmentationHeader* fragmentation) override
+		{
+			return adapter_->OnEncodedImage(stream_idx_, encoded_image,
+				codec_specific_info, fragmentation);
+		}
+
+	private:
+		UE::PixelStreaming::FSimulcastEncoderAdapter* const adapter_;
+		const size_t stream_idx_;
+	};
+} // namespace
+
+UE::PixelStreaming::FSimulcastEncoderAdapter::FSimulcastEncoderAdapter(webrtc::VideoEncoderFactory* primary_factory, const webrtc::SdpVideoFormat& format)
+	: Initialized(false)
+	, PrimaryEncoderFactory(primary_factory)
+	, VideoFormat(format)
+	, EncodedCompleteCallback(nullptr)
+{
+	check(PrimaryEncoderFactory);
+	memset(&CurrentCodec, 0, sizeof(webrtc::VideoCodec));
+}
+
+UE::PixelStreaming::FSimulcastEncoderAdapter::~FSimulcastEncoderAdapter()
+{
+	check(!IsInitialized());
+}
+
+int UE::PixelStreaming::FSimulcastEncoderAdapter::Release()
+{
+	while (!StreamInfos.empty())
+	{
+		std::unique_ptr<VideoEncoder> Encoder = std::move(StreamInfos.back().Encoder);
+		// Even though it seems very unlikely, there are no guarantees that the
+		// encoder will not call back after being Release()'d. Therefore, we first
+		// disable the callbacks here.
+		Encoder->RegisterEncodeCompleteCallback(nullptr);
+		Encoder->Release();
+		StreamInfos.pop_back(); // Deletes callback adapter.
+	}
+
+	Initialized = false;
+
+	return WEBRTC_VIDEO_CODEC_OK;
+}
+
+int UE::PixelStreaming::FSimulcastEncoderAdapter::InitEncode(const webrtc::VideoCodec* codec_settings, const webrtc::VideoEncoder::Settings& settings)
+{
+	int NumberOfStreams = GetNumberOfStreams(*codec_settings);
+	int NumActiveStreams = GetNumActiveStreams(*codec_settings);
+
+	CurrentCodec = *codec_settings;
+
+	const auto MinMax = std::minmax_element(std::begin(CurrentCodec.simulcastStream), std::begin(CurrentCodec.simulcastStream) + NumberOfStreams, StreamResolutionCompare);
+	LowestResolutionStreamIndex = std::distance(std::begin(CurrentCodec.simulcastStream), MinMax.first);
+	HighestResolutionStreamIndex = std::distance(std::begin(CurrentCodec.simulcastStream), MinMax.second);
+
+	const webrtc::SdpVideoFormat Format(CurrentCodec.codecType == webrtc::kVideoCodecVP8 ? "VP8" : "H264", VideoFormat.parameters);
+
+	if (NumActiveStreams == 1)
+	{
+		// with one stream we just proxy the pixelstreaming encoder
+		std::unique_ptr<VideoEncoder> Encoder = PrimaryEncoderFactory->CreateVideoEncoder(Format);
+		int ReturnCode = Encoder->InitEncode(&CurrentCodec, settings);
+		if (ReturnCode < 0)
+		{
+			// Explicitly destroy the current encoder; because we haven't registered
+			// a StreamInfo for it yet, Release won't do anything about it.
+			Encoder.reset();
+			Release();
+			return ReturnCode;
+		}
+
+		std::unique_ptr<webrtc::EncodedImageCallback> Callback(new AdapterEncodedImageCallback(this, 0));
+		Encoder->RegisterEncodeCompleteCallback(Callback.get());
+		StreamInfos.emplace_back(
+			std::move(Encoder), std::move(Callback),
+			std::make_unique<webrtc::FramerateController>(CurrentCodec.maxFramerate),
+			CurrentCodec.width, CurrentCodec.height, true, true);
+	}
+	else
+	{
+		for (int i = 0; i < NumberOfStreams; ++i)
+		{
+			webrtc::VideoCodec StreamCodec;
+			uint32_t StartBitrateKbps = CurrentCodec.simulcastStream[i].targetBitrate;
+			PopulateStreamCodec(CurrentCodec, i, StartBitrateKbps, &StreamCodec);
+
+			std::unique_ptr<VideoEncoder> Encoder = PrimaryEncoderFactory->CreateVideoEncoder(Format);
+			int ReturnCode = Encoder->InitEncode(&StreamCodec, settings);
+			if (ReturnCode < 0)
+			{
+				// Explicitly destroy the current encoder; because we haven't registered
+				// a StreamInfo for it yet, Release won't do anything about it.
+				Encoder.reset();
+				Release();
+				return ReturnCode;
+			}
+
+			std::unique_ptr<webrtc::EncodedImageCallback> Callback(new AdapterEncodedImageCallback(this, i));
+			Encoder->RegisterEncodeCompleteCallback(Callback.get());
+			StreamInfos.emplace_back(
+				std::move(Encoder), std::move(Callback),
+				std::make_unique<webrtc::FramerateController>(StreamCodec.maxFramerate),
+				StreamCodec.width, StreamCodec.height, true, true);
+		}
+	}
+
+	Initialized = true;
+
+	return WEBRTC_VIDEO_CODEC_OK;
+}
+
+int UE::PixelStreaming::FSimulcastEncoderAdapter::EncodeStream(const webrtc::VideoFrame& InputImage, TSharedPtr<UE::PixelStreaming::ITextureSource> LayerFrameSource, size_t StreamIdx, bool bSendKeyFrame)
+{
+	// grab the simulcast frame source, extract the frame source for this layer and wrap that in a new frame buffer
+	rtc::scoped_refptr<UE::PixelStreaming::FLayerFrameBuffer> LayerFrameBuffer = new rtc::RefCountedObject<UE::PixelStreaming::FLayerFrameBuffer>(LayerFrameSource);
+	webrtc::VideoFrame NewFrame(InputImage);
+	NewFrame.set_video_frame_buffer(LayerFrameBuffer);
+
+	const uint32_t FrameTimestampMs = 1000 * NewFrame.timestamp() / 90000;
+
+	// If adapter is passed through and only one sw encoder does simulcast,
+	// frame types for all streams should be passed to the encoder unchanged.
+	// Otherwise a single per-encoder frame type is passed.
+	std::vector<webrtc::VideoFrameType> StreamFrameTypes(StreamInfos.size() == 1 ? GetNumberOfStreams(CurrentCodec) : 1);
+	if (bSendKeyFrame)
+	{
+		std::fill(StreamFrameTypes.begin(), StreamFrameTypes.end(), webrtc::VideoFrameType::kVideoFrameKey);
+		StreamInfos[StreamIdx].KeyFrameRequest = false;
+	}
+	else
+	{
+		if (StreamInfos[StreamIdx].FramerateController->DropFrame(FrameTimestampMs))
+		{
+			return WEBRTC_VIDEO_CODEC_OK;
+		}
+		std::fill(StreamFrameTypes.begin(), StreamFrameTypes.end(), webrtc::VideoFrameType::kVideoFrameDelta);
+	}
+	StreamInfos[StreamIdx].FramerateController->AddFrame(FrameTimestampMs);
+
+	return StreamInfos[StreamIdx].Encoder->Encode(NewFrame, &StreamFrameTypes);
+}
+
+int UE::PixelStreaming::FSimulcastEncoderAdapter::Encode(const webrtc::VideoFrame& input_image, const std::vector<webrtc::VideoFrameType>* frame_types)
+{
+	if (!IsInitialized())
+	{
+		return WEBRTC_VIDEO_CODEC_UNINITIALIZED;
+	}
+
+	if (EncodedCompleteCallback == nullptr)
+	{
+		return WEBRTC_VIDEO_CODEC_UNINITIALIZED;
+	}
+
+	// ignore any init frames
+	// NOTE: It seems most of the time they never even get here because the internals of WebRTC decide to drop it. It seems it thinks
+	// the encoder is paused or something.
+	const UE::PixelStreaming::FFrameBuffer* FrameBuffer = static_cast<UE::PixelStreaming::FFrameBuffer*>(input_image.video_frame_buffer().get());
+	if (FrameBuffer->GetFrameBufferType() == UE::PixelStreaming::FFrameBufferType::Initialize)
+	{
+		return WEBRTC_VIDEO_CODEC_OK;
+	}
+
+	// separate out our pixelstreaming image sources
+	check(FrameBuffer->GetFrameBufferType() == UE::PixelStreaming::FFrameBufferType::Simulcast);
+	const UE::PixelStreaming::FSimulcastFrameBuffer* SimulcastFrameBuffer = static_cast<const UE::PixelStreaming::FSimulcastFrameBuffer*>(FrameBuffer);
+
+	// All active streams should generate a key frame if
+	// a key frame is requested by any stream.
+	bool bSendKeyFrame = false;
+	if (frame_types)
+	{
+		for (size_t i = 0; i < frame_types->size(); ++i)
+		{
+			if (frame_types->at(i) == webrtc::VideoFrameType::kVideoFrameKey)
+			{
+				bSendKeyFrame = true;
+				break;
+			}
+		}
+	}
+
+	for (size_t StreamIdx = 0; StreamIdx < StreamInfos.size(); ++StreamIdx)
+	{
+		if (StreamInfos[StreamIdx].KeyFrameRequest && StreamInfos[StreamIdx].bSendStream)
+		{
+			bSendKeyFrame = true;
+			break;
+		}
+	}
+
+	for (size_t StreamIdx = 0; StreamIdx < StreamInfos.size(); ++StreamIdx)
+	{
+		// Don't encode frames in resolutions that we don't intend to send.
+		if (!StreamInfos[StreamIdx].bSendStream)
+		{
+			continue;
+		}
+
+		TSharedPtr<ITextureSource> LayerFrameSource = SimulcastFrameBuffer->GetLayerFrameSource(StreamInfos.size() == 1 ? SimulcastFrameBuffer->GetNumLayers() - 1 : StreamIdx);
+		checkf(LayerFrameSource->GetSourceWidth() == StreamInfos[StreamIdx].Width, TEXT("Width of simulcast layer must match width of stream info it is associated with."));
+		checkf(LayerFrameSource->GetSourceHeight() == StreamInfos[StreamIdx].Height, TEXT("Height of simulcast layer must match height of stream info it is associated with."));
+
+		int RtcError = EncodeStream(input_image, LayerFrameSource, StreamIdx, bSendKeyFrame);
+		if (RtcError != WEBRTC_VIDEO_CODEC_OK)
+		{
+			return RtcError;
+		}
+	}
+	return WEBRTC_VIDEO_CODEC_OK;
+}
+
+int UE::PixelStreaming::FSimulcastEncoderAdapter::RegisterEncodeCompleteCallback(webrtc::EncodedImageCallback* callback)
+{
+	EncodedCompleteCallback = callback;
+	return WEBRTC_VIDEO_CODEC_OK;
+}
+
+void UE::PixelStreaming::FSimulcastEncoderAdapter::SetRates(const RateControlParameters& parameters)
+{
+	if (!IsInitialized())
+	{
+		RTC_LOG(LS_WARNING) << "SetRates while not initialized";
+		return;
+	}
+
+	if (parameters.framerate_fps < 1.0)
+	{
+		RTC_LOG(LS_WARNING) << "Invalid framerate: " << parameters.framerate_fps;
+		return;
+	}
+
+	CurrentCodec.maxFramerate = static_cast<uint32_t>(parameters.framerate_fps + 0.5);
+
+	if (StreamInfos.size() == 1)
+	{
+		// Not doing simulcast.
+		StreamInfos[0].Encoder->SetRates(parameters);
+		return;
+	}
+
+	for (size_t StreamIdx = 0; StreamIdx < StreamInfos.size(); ++StreamIdx)
+	{
+		uint32_t StreamBitrateKbps = parameters.bitrate.GetSpatialLayerSum(StreamIdx) / 1000;
+
+		// Need a key frame if we have not sent this stream before.
+		if (StreamBitrateKbps > 0 && !StreamInfos[StreamIdx].bSendStream)
+		{
+			StreamInfos[StreamIdx].KeyFrameRequest = true;
+		}
+		StreamInfos[StreamIdx].bSendStream = StreamBitrateKbps > 0;
+
+		// Slice the temporal layers out of the full allocation and pass it on to
+		// the encoder handling the current simulcast stream.
+		RateControlParameters StreamParameters = parameters;
+		StreamParameters.bitrate = webrtc::VideoBitrateAllocation();
+		for (int i = 0; i < webrtc::kMaxTemporalStreams; ++i)
+		{
+			if (parameters.bitrate.HasBitrate(StreamIdx, i))
+			{
+				StreamParameters.bitrate.SetBitrate(0, i, parameters.bitrate.GetBitrate(StreamIdx, i));
+			}
+		}
+
+		// Assign link allocation proportionally to spatial layer allocation.
+		if (!parameters.bandwidth_allocation.IsZero() && parameters.bitrate.get_sum_bps() > 0)
+		{
+			StreamParameters.bandwidth_allocation =
+				webrtc::DataRate::BitsPerSec((parameters.bandwidth_allocation.bps() * StreamParameters.bitrate.get_sum_bps()) / parameters.bitrate.get_sum_bps());
+			// Make sure we don't allocate bandwidth lower than target bitrate.
+			if (StreamParameters.bandwidth_allocation.bps() < StreamParameters.bitrate.get_sum_bps())
+			{
+				StreamParameters.bandwidth_allocation =
+					webrtc::DataRate::BitsPerSec(StreamParameters.bitrate.get_sum_bps());
+			}
+		}
+
+		StreamParameters.framerate_fps = std::min<double>(parameters.framerate_fps, StreamInfos[StreamIdx].FramerateController->GetTargetRate());
+
+		StreamInfos[StreamIdx].Encoder->SetRates(StreamParameters);
+	}
+}
+
+void UE::PixelStreaming::FSimulcastEncoderAdapter::OnPacketLossRateUpdate(float packet_loss_rate)
+{
+	for (StreamInfo& Info : StreamInfos)
+	{
+		Info.Encoder->OnPacketLossRateUpdate(packet_loss_rate);
+	}
+}
+
+void UE::PixelStreaming::FSimulcastEncoderAdapter::OnRttUpdate(int64_t rtt_ms)
+{
+	for (StreamInfo& Info : StreamInfos)
+	{
+		Info.Encoder->OnRttUpdate(rtt_ms);
+	}
+}
+
+void UE::PixelStreaming::FSimulcastEncoderAdapter::OnLossNotification(const LossNotification& loss_notification)
+{
+	for (StreamInfo& Info : StreamInfos)
+	{
+		Info.Encoder->OnLossNotification(loss_notification);
+	}
+}
+
+webrtc::EncodedImageCallback::Result UE::PixelStreaming::FSimulcastEncoderAdapter::OnEncodedImage(
+	size_t stream_idx,
+	const webrtc::EncodedImage& encodedImage,
+	const webrtc::CodecSpecificInfo* codecSpecificInfo,
+	const webrtc::RTPFragmentationHeader* fragmentation)
+{
+	webrtc::EncodedImage StreamImage(encodedImage);
+	webrtc::CodecSpecificInfo StreamCodecSpecific = *codecSpecificInfo;
+
+	StreamImage.SetSpatialIndex(stream_idx);
+
+	return EncodedCompleteCallback->OnEncodedImage(StreamImage, &StreamCodecSpecific, fragmentation);
+}
+
+bool UE::PixelStreaming::FSimulcastEncoderAdapter::IsInitialized() const
+{
+	return Initialized;
+}
+
+webrtc::VideoEncoder::EncoderInfo UE::PixelStreaming::FSimulcastEncoderAdapter::GetEncoderInfo() const
+{
+	if (StreamInfos.size() == 1)
+	{
+		// Not using simulcast adapting functionality, just pass through.
+		return StreamInfos[0].Encoder->GetEncoderInfo();
+	}
+
+	VideoEncoder::EncoderInfo EncoderInfo;
+	EncoderInfo.implementation_name = "PixelStreamingSimulcastEncoderAdapter";
+	EncoderInfo.requested_resolution_alignment = 1;
+	EncoderInfo.supports_native_handle = true;
+	EncoderInfo.scaling_settings.thresholds = absl::nullopt;
+	if (StreamInfos.empty())
+	{
+		return EncoderInfo;
+	}
+
+	EncoderInfo.scaling_settings = VideoEncoder::ScalingSettings::kOff;
+	int NumActiveStreams = GetNumActiveStreams(CurrentCodec);
+
+	for (size_t i = 0; i < StreamInfos.size(); ++i)
+	{
+		VideoEncoder::EncoderInfo EncoderImplInfo = StreamInfos[i].Encoder->GetEncoderInfo();
+
+		if (i == 0)
+		{
+			// Encoder name indicates names of all sub-encoders.
+			EncoderInfo.implementation_name += " (";
+			EncoderInfo.implementation_name += EncoderImplInfo.implementation_name;
+
+			EncoderInfo.supports_native_handle = EncoderImplInfo.supports_native_handle;
+			EncoderInfo.has_trusted_rate_controller = EncoderImplInfo.has_trusted_rate_controller;
+			EncoderInfo.is_hardware_accelerated = EncoderImplInfo.is_hardware_accelerated;
+			EncoderInfo.has_internal_source = EncoderImplInfo.has_internal_source;
+		}
+		else
+		{
+			EncoderInfo.implementation_name += ", ";
+			EncoderInfo.implementation_name += EncoderImplInfo.implementation_name;
+
+			// Native handle supported if any encoder supports it.
+			EncoderInfo.supports_native_handle |= EncoderImplInfo.supports_native_handle;
+
+			// Trusted rate controller only if all encoders have it.
+			EncoderInfo.has_trusted_rate_controller &= EncoderImplInfo.has_trusted_rate_controller;
+
+			// Uses hardware support if any of the encoders uses it.
+			// For example, if we are having issues with down-scaling due to
+			// pipelining delay in HW encoders we need higher encoder usage
+			// thresholds in CPU adaptation.
+			EncoderInfo.is_hardware_accelerated |= EncoderImplInfo.is_hardware_accelerated;
+
+			// Has internal source only if all encoders have it.
+			EncoderInfo.has_internal_source &= EncoderImplInfo.has_internal_source;
+		}
+		EncoderInfo.fps_allocation[i] = EncoderImplInfo.fps_allocation[0];
+		EncoderInfo.requested_resolution_alignment = cricket::LeastCommonMultiple(EncoderInfo.requested_resolution_alignment, EncoderImplInfo.requested_resolution_alignment);
+		if (NumActiveStreams == 1 && CurrentCodec.simulcastStream[i].active)
+		{
+			EncoderInfo.scaling_settings = EncoderImplInfo.scaling_settings;
+		}
+	}
+	EncoderInfo.implementation_name += ")";
+
+	return EncoderInfo;
+}
diff --git a/Engine/Plugins/Media/PixelStreaming/Source/PixelStreaming/Private/SimulcastEncoderAdapter.h b/Engine/Plugins/Media/PixelStreaming/Source/PixelStreaming/Private/SimulcastEncoderAdapter.h
new file mode 100644
--- /dev/null
+++ b/Engine/Plugins/Media/PixelStreaming/Source/PixelStreaming/Private/SimulcastEncoderAdapter.h
@@ -0,0 +1,89 @@
+// Copyright Epic Games, Inc. All Rights Reserved.
+#pragma once
+
+#include "CoreMinimal.h"
+#include "WebRTCIncludes.h"
+
+namespace webrtc
+{
+	class SimulcastRateAllocator;
+	class VideoEncoderFactory;
+} // namespace webrtc
+
+namespace UE {
+	namespace PixelStreaming {
+		class ITextureSource;
+
+		// This is highly modified version of webrtc::simulcast_encoder_adapter
+
+		class FSimulcastEncoderAdapter : public webrtc::VideoEncoder
+		{
+		public:
+			FSimulcastEncoderAdapter(webrtc::VideoEncoderFactory* primary_factory, const webrtc::SdpVideoFormat& format);
+			~FSimulcastEncoderAdapter() override;
+
+			// Implements VideoEncoder.
+			int Release() override;
+			int InitEncode(const webrtc::VideoCodec* codec_settings, const webrtc::VideoEncoder::Settings& settings) override;
+			int Encode(const webrtc::VideoFrame& input_image, const std::vector<webrtc::VideoFrameType>* frame_types) override;
+			int RegisterEncodeCompleteCallback(webrtc::EncodedImageCallback* callback) override;
+
+			void SetRates(const RateControlParameters& parameters) override;
+			void OnPacketLossRateUpdate(float packet_loss_rate) override;
+			void OnRttUpdate(int64_t rtt_ms) override;
+			void OnLossNotification(const LossNotification& loss_notification) override;
+
+			// Eventual handler for the contained encoders' EncodedImageCallbacks, but
+			// called from an internal helper that also knows the correct stream
+			// index.
+			webrtc::EncodedImageCallback::Result OnEncodedImage(
+				size_t stream_idx,
+				const webrtc::EncodedImage& encoded_image,
+				const webrtc::CodecSpecificInfo* codec_specific_info,
+				const webrtc::RTPFragmentationHeader* fragmentation);
+
+			EncoderInfo GetEncoderInfo() const override;
+
+		private:
+			struct StreamInfo
+			{
+				StreamInfo(std::unique_ptr<webrtc::VideoEncoder> InEncoder,
+					std::unique_ptr<webrtc::EncodedImageCallback> InCallback,
+					std::unique_ptr<webrtc::FramerateController> InFramerateController,
+					uint16_t InWidth,
+					uint16_t InHeight,
+					bool bInKeyFrame,
+					bool bInSendStream)
+					: Encoder(std::move(InEncoder))
+					, Callback(std::move(InCallback))
+					, FramerateController(std::move(InFramerateController))
+					, Width(InWidth)
+					, Height(InHeight)
+					, KeyFrameRequest(bInKeyFrame)
+					, bSendStream(bInSendStream)
+				{
+				}
+				std::unique_ptr<webrtc::VideoEncoder> Encoder;
+				std::unique_ptr<webrtc::EncodedImageCallback> Callback;
+				std::unique_ptr<webrtc::FramerateController> FramerateController;
+				uint16_t Width;
+				uint16_t Height;
+				bool KeyFrameRequest;
+				bool bSendStream;
+			};
+
+			bool IsInitialized() const;
+
+			TAtomic<bool> Initialized;
+			webrtc::VideoEncoderFactory* const PrimaryEncoderFactory;
+			const webrtc::SdpVideoFormat VideoFormat;
+			webrtc::VideoCodec CurrentCodec;
+			std::vector<StreamInfo> StreamInfos;
+			webrtc::EncodedImageCallback* EncodedCompleteCallback;
+
+			int LowestResolutionStreamIndex;
+			int HighestResolutionStreamIndex;
+			int EncodeStream(const webrtc::VideoFrame& InputImage, TSharedPtr<ITextureSource> LayerFrameSource, size_t StreamIdx, bool bSendKeyFrame);
+		};
+	}
+}
diff --git a/Engine/Plugins/Media/PixelStreaming/Source/PixelStreaming/Private/Stats.cpp b/Engine/Plugins/Media/PixelStreaming/Source/PixelStreaming/Private/Stats.cpp
new file mode 100644
--- /dev/null
+++ b/Engine/Plugins/Media/PixelStreaming/Source/PixelStreaming/Private/Stats.cpp
@@ -0,0 +1,494 @@
+// Copyright Epic Games, Inc. All Rights Reserved.
+
+#include "Stats.h"
+#include "PixelStreamingPrivate.h"
+#include "Engine/Engine.h"
+#include "CoreGlobals.h"
+#include "Async/Async.h"
+#include "CanvasItem.h"
+#include "CanvasTypes.h"
+#include "Stats/Stats.h"
+#include "PixelStreamingDelegates.h"
+
+UE::PixelStreaming::FStats* UE::PixelStreaming::FStats::Instance = nullptr;
+
+UE::PixelStreaming::FStats* UE::PixelStreaming::FStats::Get()
+{
+	return UE::PixelStreaming::FStats::Instance;
+}
+
+UE::PixelStreaming::FStats::FStats(UE::PixelStreaming::IPixelStreamingSessions* InSessions)
+	: Sessions(InSessions)
+{
+	checkf(UE::PixelStreaming::FStats::Instance == nullptr, TEXT("There should only ever been one PixelStreaming stats object."));
+	checkf(Sessions, TEXT("To make stats object the sessions object must not be nullptr."));
+	UE::PixelStreaming::FStats::Instance = this;
+}
+
+void UE::PixelStreaming::FStats::QueryPeerStat(FPixelStreamingPlayerId PlayerId, FName StatToQuery, TFunction<void(bool, double)> QueryCallback) const
+{
+	if (IsInGameThread())
+	{
+		double OutValue = 0.0;
+		bool bHadStat = QueryPeerStat_GameThread(PlayerId, StatToQuery, OutValue);
+		QueryCallback(bHadStat, OutValue);
+	}
+	else
+	{
+		AsyncTask(ENamedThreads::GameThread, [this, PlayerId, StatToQuery, QueryCallback]() {
+			double OutValue = 0.0;
+			bool bHadStat = QueryPeerStat_GameThread(PlayerId, StatToQuery, OutValue);
+			QueryCallback(bHadStat, OutValue);
+		});
+	}
+}
+
+void UE::PixelStreaming::FStats::RemovePeersStats(FPixelStreamingPlayerId PlayerId)
+{
+	if (IsInGameThread())
+	{
+		RemovePeerStat_GameThread(PlayerId);
+	}
+	else
+	{
+		AsyncTask(ENamedThreads::GameThread, [this, PlayerId]() { RemovePeerStat_GameThread(PlayerId); });
+	}
+}
+
+void UE::PixelStreaming::FStats::StorePeerStat(FPixelStreamingPlayerId PlayerId, UE::PixelStreaming::FStatData Stat)
+{
+	if (IsInGameThread())
+	{
+		//todo: return bool
+		bool bUpdated = StorePeerStat_GameThread(PlayerId, Stat);
+		if (bUpdated)
+		{
+			FireStatChanged_GameThread(PlayerId, Stat.StatName, Stat.StatValue);
+		}
+	}
+	else
+	{
+		AsyncTask(ENamedThreads::GameThread, [this, PlayerId, Stat]() {
+			bool bUpdated = StorePeerStat_GameThread(PlayerId, Stat);
+			if (bUpdated)
+			{
+				FireStatChanged_GameThread(PlayerId, Stat.StatName, Stat.StatValue);
+			}
+		});
+	}
+}
+
+void UE::PixelStreaming::FStats::StoreApplicationStat(UE::PixelStreaming::FStatData Stat)
+{
+	if (IsInGameThread())
+	{
+		bool bUpdated = StoreApplicationStat_GameThread(Stat);
+		if (bUpdated)
+		{
+			FireStatChanged_GameThread(FPixelStreamingPlayerId(TEXT("Application")), Stat.StatName, Stat.StatValue);
+		}
+	}
+	else
+	{
+		AsyncTask(ENamedThreads::GameThread, [this, Stat]() {
+			bool bUpdated = StoreApplicationStat_GameThread(Stat);
+			if (bUpdated)
+			{
+				FireStatChanged_GameThread(FPixelStreamingPlayerId(TEXT("Application")), Stat.StatName, Stat.StatValue);
+			}
+		});
+	}
+}
+
+void UE::PixelStreaming::FStats::FireStatChanged_GameThread(FPixelStreamingPlayerId PlayerId, FName StatName, float StatValue)
+{
+	checkf(IsInGameThread(), TEXT("This method must be called from the game thread"));
+
+	if (UPixelStreamingDelegates* Delegates = UPixelStreamingDelegates::GetPixelStreamingDelegates())
+	{
+		Delegates->OnPixelStreamerStatChangedNative.Broadcast(PlayerId, StatName, StatValue);
+	}
+}
+
+void UE::PixelStreaming::FStats::RemovePeerStat_GameThread(FPixelStreamingPlayerId PlayerId)
+{
+	PeerStats.Remove(PlayerId);
+
+	if (PlayerId == SFU_PLAYER_ID)
+	{
+		TArray<FPixelStreamingPlayerId> ToRemove;
+
+		for (auto& Entry : PeerStats)
+		{
+			FPixelStreamingPlayerId PeerId = Entry.Key;
+			if (PeerId.Contains(TEXT("Simulcast"), ESearchCase::IgnoreCase, ESearchDir::FromStart))
+			{
+				ToRemove.Add(PeerId);
+			}
+		}
+
+		for (FPixelStreamingPlayerId SimulcastLayerId : ToRemove)
+		{
+			PeerStats.Remove(SimulcastLayerId);
+		}
+	}
+}
+
+bool UE::PixelStreaming::FStats::QueryPeerStat_GameThread(FPixelStreamingPlayerId PlayerId, FName StatToQuery, double& OutValue) const
+{
+	verifyf(IsInGameThread(), TEXT("This method must be called from the game thread"));
+
+	const UE::PixelStreaming::FPeerStats* SinglePeerStats = PeerStats.Find(PlayerId);
+	if (!SinglePeerStats)
+	{
+		return false;
+	}
+
+	const UE::PixelStreaming::FRenderableStat* StoredStat = SinglePeerStats->StoredStats.Find(StatToQuery);
+	if (!StoredStat)
+	{
+		return false;
+	}
+
+	OutValue = StoredStat->Stat.StatValue;
+	return true;
+}
+
+bool UE::PixelStreaming::FStats::StorePeerStat_GameThread(FPixelStreamingPlayerId PlayerId, UE::PixelStreaming::FStatData Stat)
+{
+	checkf(IsInGameThread(), TEXT("This method must be called from the game thread"));
+
+	if (IsEngineExitRequested())
+	{
+		return false;
+	}
+
+	if (!PeerStats.Contains(PlayerId))
+	{
+		PeerStats.Add(PlayerId, UE::PixelStreaming::FPeerStats(PlayerId));
+		return true;
+	}
+	return PeerStats[PlayerId].StoreStat_GameThread(Stat);
+}
+
+bool UE::PixelStreaming::FStats::StoreApplicationStat_GameThread(UE::PixelStreaming::FStatData Stat)
+{
+	checkf(IsInGameThread(), TEXT("This method must be called from the game thread"));
+
+	if (IsEngineExitRequested())
+	{
+		return false;
+	}
+
+	bool bUpdated = false;
+
+	if (ApplicationStats.Contains(Stat.StatName))
+	{
+		UE::PixelStreaming::FRenderableStat* RenderableStat = ApplicationStats.Find(Stat.StatName);
+
+		if (Stat.bSmooth && RenderableStat->Stat.StatValue != 0)
+		{
+			double CurValue = RenderableStat->Stat.StatValue;
+			double PercentageDrift = FMath::Abs(CurValue - Stat.StatValue) / CurValue;
+			if (PercentageDrift > UE::PixelStreaming::FStats::SmoothingFactor)
+			{
+				RenderableStat->Stat.StatValue = Stat.StatValue;
+				bUpdated = true;
+			}
+			else
+			{
+				return false;
+			}
+		}
+		else
+		{
+			bUpdated = RenderableStat->Stat.StatValue != Stat.StatValue;
+			RenderableStat->Stat.StatValue = Stat.StatValue;
+		}
+
+		if (bUpdated)
+		{
+			FText TextToDisplay = FText::FromString(FString::Printf(TEXT("%s: %.*f"), *Stat.StatName.ToString(), Stat.NDecimalPlacesToPrint, RenderableStat->Stat.StatValue));
+			RenderableStat->CanvasItem.Text = TextToDisplay;
+		}
+	}
+	else
+	{
+		FText TextToDisplay = FText::FromString(FString::Printf(TEXT("%s: %.*f"), *Stat.StatName.ToString(), Stat.NDecimalPlacesToPrint, Stat.StatValue));
+
+		UE::PixelStreaming::FRenderableStat RenderableStat{
+			Stat,
+			FCanvasTextItem(FVector2D(0, 0), TextToDisplay, FSlateFontInfo(FSlateFontInfo(UEngine::GetSmallFont(), 12)), FLinearColor(0, 1, 0))
+		};
+
+		RenderableStat.CanvasItem.EnableShadow(FLinearColor::Black);
+
+		ApplicationStats.Add(RenderableStat.Stat.StatName, RenderableStat);
+		bUpdated = true;
+	}
+	return bUpdated;
+}
+
+void UE::PixelStreaming::FStats::AddOnPeerStatChangedCallback_GameThread(FPixelStreamingPlayerId PlayerId, FName StatToListenOn, TWeakPtr<IPixelStreamingStatsConsumer> Callback)
+{
+	checkf(IsInGameThread(), TEXT("This method was not called from the game thread."));
+
+	UE::PixelStreaming::FPeerStats* SinglePeerStats = PeerStats.Find(PlayerId);
+	if (!SinglePeerStats)
+	{
+		PeerStats.Add(PlayerId, UE::PixelStreaming::FPeerStats(PlayerId));
+		SinglePeerStats = PeerStats.Find(PlayerId);
+	}
+
+	TArray<TWeakPtr<IPixelStreamingStatsConsumer>>* Callbacks = SinglePeerStats->SingleStatConsumers.Find(StatToListenOn);
+
+	if (!Callbacks)
+	{
+		SinglePeerStats->SingleStatConsumers.Add(StatToListenOn, TArray<TWeakPtr<IPixelStreamingStatsConsumer>>());
+		Callbacks = SinglePeerStats->SingleStatConsumers.Find(StatToListenOn);
+	}
+
+	Callbacks->Add(Callback);
+}
+
+int32 UE::PixelStreaming::FStats::OnRenderStats(UWorld* World, FViewport* Viewport, FCanvas* Canvas, int32 X, int32 Y, const FVector* ViewLocation, const FRotator* ViewRotation)
+{
+	if (GAreScreenMessagesEnabled)
+	{
+		Y += 250;
+		// Draw each peer's stats in a column, so we must recall where Y starts for each column
+		int32 YStart = Y;
+
+		// --------- Draw stats for this Pixel Streaming instance ----------
+
+		for (auto& ApplicationStatEntry : ApplicationStats)
+		{
+			UE::PixelStreaming::FRenderableStat& StatToDraw = ApplicationStatEntry.Value;
+			StatToDraw.CanvasItem.Position.X = X;
+			StatToDraw.CanvasItem.Position.Y = Y;
+			Canvas->DrawItem(StatToDraw.CanvasItem);
+			Y += StatToDraw.CanvasItem.DrawnSize.Y;
+		}
+
+		// --------- Draw stats for each peer ----------
+
+		// increment X now we are done drawing application stats
+		X += 435;
+
+		// TMap<FPixelStreamingPlayerId, UE::PixelStreaming::FPeerStats> PeerStats;
+		for (auto& EachPeerEntry : PeerStats)
+		{
+			UE::PixelStreaming::FPeerStats& SinglePeerStats = EachPeerEntry.Value;
+			if (SinglePeerStats.StoredStats.Num() == 0)
+			{
+				continue;
+			}
+
+			// Reset Y for each peer as each peer gets it own column
+			Y = YStart;
+
+			SinglePeerStats.PlayerIdCanvasItem.Position.X = X;
+			SinglePeerStats.PlayerIdCanvasItem.Position.Y = Y;
+			Canvas->DrawItem(SinglePeerStats.PlayerIdCanvasItem);
+			Y += SinglePeerStats.PlayerIdCanvasItem.DrawnSize.Y;
+
+			for (auto& NameStatKeyVal : SinglePeerStats.StoredStats)
+			{
+				UE::PixelStreaming::FRenderableStat& Stat = NameStatKeyVal.Value;
+				Stat.CanvasItem.Position.X = X;
+				Stat.CanvasItem.Position.Y = Y;
+				Canvas->DrawItem(Stat.CanvasItem);
+				Y += Stat.CanvasItem.DrawnSize.Y;
+			}
+
+			// Each peer's stats gets its own column
+			X += 250;
+		}
+	}
+	return Y;
+}
+
+bool UE::PixelStreaming::FStats::OnToggleStats(UWorld* World, FCommonViewportClient* ViewportClient, const TCHAR* Stream)
+{
+	return false;
+}
+
+void UE::PixelStreaming::FStats::PollPixelStreamingSettings()
+{
+	double DeltaSeconds = FPlatformTime::ToSeconds64(FPlatformTime::Cycles64() - LastTimeSettingsPolledCycles);
+	if (DeltaSeconds > 1)
+	{
+		StoreApplicationStat(UE::PixelStreaming::FStatData(FName(TEXT("PixelStreaming.Encoder.MinQP")), UE::PixelStreaming::Settings::CVarPixelStreamingEncoderMinQP.GetValueOnAnyThread(), 0));
+		StoreApplicationStat(UE::PixelStreaming::FStatData(FName(TEXT("PixelStreaming.Encoder.MaxQP")), UE::PixelStreaming::Settings::CVarPixelStreamingEncoderMaxQP.GetValueOnAnyThread(), 0));
+		StoreApplicationStat(UE::PixelStreaming::FStatData(FName(TEXT("PixelStreaming.Encoder.KeyframeInterval (frames)")), UE::PixelStreaming::Settings::CVarPixelStreamingEncoderKeyframeInterval.GetValueOnAnyThread(), 0));
+		StoreApplicationStat(UE::PixelStreaming::FStatData(FName(TEXT("PixelStreaming.WebRTC.Fps")), UE::PixelStreaming::Settings::CVarPixelStreamingWebRTCFps.GetValueOnAnyThread(), 0));
+		StoreApplicationStat(UE::PixelStreaming::FStatData(FName(TEXT("PixelStreaming.WebRTC.StartBitrate")), UE::PixelStreaming::Settings::CVarPixelStreamingWebRTCStartBitrate.GetValueOnAnyThread(), 0));
+		StoreApplicationStat(UE::PixelStreaming::FStatData(FName(TEXT("PixelStreaming.WebRTC.MinBitrate")), UE::PixelStreaming::Settings::CVarPixelStreamingWebRTCMinBitrate.GetValueOnAnyThread(), 0));
+		StoreApplicationStat(UE::PixelStreaming::FStatData(FName(TEXT("PixelStreaming.WebRTC.MaxBitrate")), UE::PixelStreaming::Settings::CVarPixelStreamingWebRTCMaxBitrate.GetValueOnAnyThread(), 0));
+
+		LastTimeSettingsPolledCycles = FPlatformTime::Cycles64();
+	}
+}
+
+void UE::PixelStreaming::FStats::AddOnPeerStatChangedCallback(FPixelStreamingPlayerId PlayerId, FName StatToListenOn, TWeakPtr<IPixelStreamingStatsConsumer> Callback)
+{
+	if (IsInGameThread())
+	{
+		AddOnPeerStatChangedCallback_GameThread(PlayerId, StatToListenOn, Callback);
+	}
+	else
+	{
+		AsyncTask(ENamedThreads::GameThread, [this, PlayerId, StatToListenOn, Callback]() { AddOnPeerStatChangedCallback_GameThread(PlayerId, StatToListenOn, Callback); });
+	}
+}
+
+void UE::PixelStreaming::FStats::Tick(float DeltaTime)
+{
+	// Note (Luke): If we want more metrics from WebRTC there is also the histogram counts.
+	// For example:
+	// RTC_HISTOGRAM_COUNTS("WebRTC.Video.NacksSent", nacks_sent, 1, 100000, 100);
+	// webrtc::metrics::Histogram* Hist1 = webrtc::metrics::HistogramFactoryGetCounts("WebRTC.Video.NacksSent", 0, 100000, 100);
+	// Will require calling webrtc::metrics::Enable();
+
+	Sessions->PollWebRTCStats();
+	PollPixelStreamingSettings();
+
+	if (!GEngine)
+	{
+		return;
+	}
+
+	// Check if user has enabled on screen stats
+	if (!bRegisterEngineStats && UE::PixelStreaming::Settings::CVarPixelStreamingOnScreenStats.GetValueOnGameThread())
+	{
+		GAreScreenMessagesEnabled = true;
+		UEngine::FEngineStatRender RenderFunc = UEngine::FEngineStatRender::CreateRaw(this, &UE::PixelStreaming::FStats::OnRenderStats);
+		UEngine::FEngineStatToggle ToggleFunc = UEngine::FEngineStatToggle::CreateRaw(this, &UE::PixelStreaming::FStats::OnToggleStats);
+		GEngine->AddEngineStat(PixelStreamingStatName, PixelStreamingStatCategory, PixelStreamingStatDescription, RenderFunc, ToggleFunc, false);
+
+		// Turn on the Engine stat for Pixel Streaming
+		for (const FWorldContext& WorldContext : GEngine->GetWorldContexts())
+		{
+			if (WorldContext.WorldType == EWorldType::Game || WorldContext.WorldType == EWorldType::PIE)
+			{
+				UWorld* World = WorldContext.World();
+				UGameViewportClient* ViewportClient = World->GetGameViewport();
+				GEngine->SetEngineStat(World, ViewportClient, TEXT("PixelStreaming"), true);
+			}
+		}
+
+		bRegisterEngineStats = true;
+	}
+}
+
+//
+// ---------------- PixelStreamingPeerStats ---------------------------
+// Stats specific to a particular peer, as opposed to the entire app.
+//
+
+void UE::PixelStreaming::FPeerStats::StoreStat(UE::PixelStreaming::FStatData StatToStore)
+{
+	if (IsInGameThread())
+	{
+		StoreStat_GameThread(StatToStore);
+	}
+	else
+	{
+		AsyncTask(ENamedThreads::GameThread, [this, StatToStore]() { StoreStat_GameThread(StatToStore); });
+	}
+}
+
+bool UE::PixelStreaming::FPeerStats::StoreStat_GameThread(UE::PixelStreaming::FStatData StatToStore)
+{
+	checkf(IsInGameThread(), TEXT("This method was not called from the game thread."));
+
+	if (IsEngineExitRequested())
+	{
+		return false;
+	}
+
+	bool bUpdated = false;
+
+	if (!StoredStats.Contains(StatToStore.StatName))
+	{
+		FText TextToDisplay = FText::FromString(FString::Printf(TEXT("%s: %.*f"), *StatToStore.StatName.ToString(), StatToStore.NDecimalPlacesToPrint, StatToStore.StatValue));
+
+		UE::PixelStreaming::FRenderableStat RenderableStat{
+			StatToStore,
+			FCanvasTextItem(FVector2D(0, 0), TextToDisplay, FSlateFontInfo(FSlateFontInfo(UEngine::GetSmallFont(), 12)), FLinearColor(0, 1, 0))
+		};
+
+		RenderableStat.CanvasItem.EnableShadow(FLinearColor::Black);
+
+		StoredStats.Add(StatToStore.StatName, RenderableStat);
+
+		// first time this stat has been stored, so we also need to sort our stats so they render in consistent order
+		StoredStats.KeySort([](const FName& A, const FName& B) {
+			return A.FastLess(B);
+		});
+
+		return true;
+	}
+	else
+	{
+		// We already have this stat, so just update it
+
+		UE::PixelStreaming::FRenderableStat* RenderableStat = StoredStats.Find(StatToStore.StatName);
+		if (!RenderableStat)
+		{
+			return false;
+		}
+
+		if (RenderableStat->Stat.bSmooth && RenderableStat->Stat.StatValue != 0)
+		{
+			double CurValue = RenderableStat->Stat.StatValue;
+			double PercentageDrift = FMath::Abs(CurValue - StatToStore.StatValue) / CurValue;
+			if (PercentageDrift > UE::PixelStreaming::FStats::SmoothingFactor)
+			{
+				RenderableStat->Stat.StatValue = StatToStore.StatValue;
+				bUpdated = true;
+			}
+			else
+			{
+				return false;
+			}
+		}
+		else
+		{
+			bUpdated = RenderableStat->Stat.StatValue != StatToStore.StatValue;
+			RenderableStat->Stat.StatValue = StatToStore.StatValue;
+		}
+
+		if (!bUpdated)
+		{
+			return false;
+		}
+		else
+		{
+			FText TextToDisplay = FText::FromString(FString::Printf(TEXT("%s: %.*f"), *StatToStore.StatName.ToString(), StatToStore.NDecimalPlacesToPrint, RenderableStat->Stat.StatValue));
+			RenderableStat->CanvasItem.Text = TextToDisplay;
+		}
+
+		// Fire any callbacks to anyone listening for changing in this stat
+		if (!SingleStatConsumers.Contains(StatToStore.StatName))
+		{
+			return bUpdated;
+		}
+
+		TArray<TWeakPtr<IPixelStreamingStatsConsumer>>* StatConsumersArr = SingleStatConsumers.Find(StatToStore.StatName);
+		if (!StatConsumersArr)
+		{
+			return bUpdated;
+		}
+
+		for (TWeakPtr<IPixelStreamingStatsConsumer> Consumer : *StatConsumersArr)
+		{
+			if (Consumer.Pin())
+			{
+				Consumer.Pin()->ConsumeStat(AssociatedPlayer, StatToStore.StatName, RenderableStat->Stat.StatValue);
+			}
+		}
+
+		return bUpdated;
+	}
+}
\ No newline at end of file
diff --git a/Engine/Plugins/Media/PixelStreaming/Source/PixelStreaming/Private/Stats.h b/Engine/Plugins/Media/PixelStreaming/Source/PixelStreaming/Private/Stats.h
new file mode 100644
--- /dev/null
+++ b/Engine/Plugins/Media/PixelStreaming/Source/PixelStreaming/Private/Stats.h
@@ -0,0 +1,131 @@
+// Copyright Epic Games, Inc. All Rights Reserved.
+
+#pragma once
+
+#include "Settings.h"
+#include "Tickable.h"
+#include "PixelStreamingPlayerId.h"
+#include "CanvasItem.h"
+#include "IPixelStreamingSessions.h"
+#include "Containers/UnrealString.h"
+#include "IPixelStreamingStatsConsumer.h"
+#include "UnrealEngine.h"
+
+namespace UE {
+	namespace PixelStreaming {
+		// ----------------- FStatData -----------------
+
+		struct FStatData
+		{
+		public:
+			FStatData(FName InStatName, double InStatValue, int InNDecimalPlacesToPrint, bool bInSmooth = false)
+				: StatName(InStatName)
+				, StatValue(InStatValue)
+				, NDecimalPlacesToPrint(InNDecimalPlacesToPrint)
+				, bSmooth(bInSmooth) {}
+
+			bool operator==(const FStatData& Other) const
+			{
+				return Equals(Other);
+			}
+
+			bool Equals(const FStatData& Other) const
+			{
+				return StatName == Other.StatName;
+			}
+
+			FName StatName;
+			double StatValue;
+			int NDecimalPlacesToPrint;
+			bool bSmooth;
+		};
+
+		FORCEINLINE uint32 GetTypeHash(const FStatData& Obj)
+		{
+			// From UnrealString.h
+			return GetTypeHash(Obj.StatName);
+		}
+
+		// ----------------- FRenderableStat -----------------
+
+		struct FRenderableStat
+		{
+			FStatData Stat;
+			FCanvasTextItem CanvasItem;
+		};
+
+		// ----------------- FPeerStats -----------------
+
+		// Pixel Streaming stats that are associated with a specific peer.
+		class FPeerStats
+		{
+
+		public:
+			FPeerStats(FPixelStreamingPlayerId InAssociatedPlayer)
+				: AssociatedPlayer(InAssociatedPlayer)
+				, PlayerIdCanvasItem(FVector2D(0, 0), FText::FromString(FString::Printf(TEXT("[Peer Stats(%s)]"), *AssociatedPlayer)), FSlateFontInfo(FSlateFontInfo(UEngine::GetSmallFont(), 12)), FLinearColor(0, 1, 0))
+			{
+				PlayerIdCanvasItem.EnableShadow(FLinearColor::Black);
+			};
+
+			void StoreStat(FStatData StatToStore);
+			bool StoreStat_GameThread(FStatData StatToStore);
+
+		private:
+			int DisplayId = 0;
+			FPixelStreamingPlayerId AssociatedPlayer;
+
+		public:
+			TMap<FName, FRenderableStat> StoredStats;
+			TMap<FName, TArray<TWeakPtr<IPixelStreamingStatsConsumer>>> SingleStatConsumers;
+			FCanvasTextItem PlayerIdCanvasItem;
+		};
+
+		// ----------------- FStats -----------------
+
+		// Stats about Pixel Streaming that can displayed either in the in-application HUD, in the log, or simply reported to some subscriber.
+		class FStats : FTickableGameObject
+		{
+		public:
+			static constexpr double SmoothingPeriod = 3.0 * 60.0;
+			static constexpr double SmoothingFactor = 10.0 / 100.0;
+			static FStats* Get();
+
+			FStats(IPixelStreamingSessions* Sessions);
+			void QueryPeerStat(FPixelStreamingPlayerId PlayerId, FName StatToQuery, TFunction<void(bool, double)> QueryCallback) const;
+			bool QueryPeerStat_GameThread(FPixelStreamingPlayerId PlayerId, FName StatToQuery, double& OutStatValue) const;
+			void RemovePeersStats(FPixelStreamingPlayerId PlayerId);
+			void StorePeerStat(FPixelStreamingPlayerId PlayerId, FStatData Stat);
+			void StoreApplicationStat(FStatData PeerStat);
+			void Tick(float DeltaTime);
+			void AddOnPeerStatChangedCallback(FPixelStreamingPlayerId PlayerId, FName StatToListenOn, TWeakPtr<IPixelStreamingStatsConsumer> Callback);
+			int32 OnRenderStats(UWorld* World, FViewport* Viewport, FCanvas* Canvas, int32 X, int32 Y, const FVector* ViewLocation, const FRotator* ViewRotation);
+			bool OnToggleStats(UWorld* World, FCommonViewportClient* ViewportClient, const TCHAR* Stream);
+
+			FORCEINLINE TStatId GetStatId() const { RETURN_QUICK_DECLARE_CYCLE_STAT(PixelStreamingStats, STATGROUP_Tickables); }
+
+		private:
+			void PollPixelStreamingSettings();
+			void RemovePeerStat_GameThread(FPixelStreamingPlayerId PlayerId);
+			bool StorePeerStat_GameThread(FPixelStreamingPlayerId PlayerId, FStatData Stat);
+			bool StoreApplicationStat_GameThread(FStatData Stat);
+			void AddOnPeerStatChangedCallback_GameThread(FPixelStreamingPlayerId PlayerId, FName StatToListenOn, TWeakPtr<IPixelStreamingStatsConsumer> Callback);
+			void FireStatChanged_GameThread(FPixelStreamingPlayerId PlayerId, FName StatName, float StatValue);
+
+		private:
+			static FStats* Instance;
+
+			IPixelStreamingSessions* Sessions;
+
+			bool bRegisterEngineStats = false;
+			FName PixelStreamingStatName = FName(TEXT("STAT_PixelStreaming"));
+			FName PixelStreamingStatCategory = FName(TEXT("STATCAT_PixelStreaming"));
+			FText PixelStreamingStatDescription = FText::FromString(FString(TEXT("Pixel Streaming stats for all connected peers.")));
+
+			TMap<FPixelStreamingPlayerId, FPeerStats> PeerStats;
+			TMap<FName, FRenderableStat> ApplicationStats;
+
+			int64 LastTimeSettingsPolledCycles = 0;
+		};
+	}
+}
