Merging Pixel Streaming refactor using //UE5/Dev-TensorWorks-to-//UE5/Release-5.0

diff --git a/Engine/Plugins/Media/PixelStreaming/Source/PixelStreaming/Private/Streamer.cpp b/Engine/Plugins/Media/PixelStreaming/Source/PixelStreaming/Private/Streamer.cpp
--- a/Engine/Plugins/Media/PixelStreaming/Source/PixelStreaming/Private/Streamer.cpp
+++ b/Engine/Plugins/Media/PixelStreaming/Source/PixelStreaming/Private/Streamer.cpp
@@ -1,33 +1,31 @@
 // Copyright Epic Games, Inc. All Rights Reserved.
 
 #include "Streamer.h"
+#include "ToStringExtensions.h"
 #include "AudioCapturer.h"
 #include "PlayerSession.h"
 #include "VideoEncoderFactory.h"
-#include "PixelStreamerDelegates.h"
-#include "PixelStreamingEncoderFactory.h"
-#include "PixelStreamingSettings.h"
+#include "EncoderFactory.h"
+#include "Settings.h"
 #include "WebRTCLogging.h"
 #include "WebSocketsModule.h"
-#include "PixelStreamingAudioDeviceModule.h"
-#include "PixelStreamingAudioSink.h"
-#include "WebRtcObservers.h"
-#include "PlayerVideoSource.h"
-#include "PixelStreamingRTCStatsCollector.h"
+#include "AudioDeviceModule.h"
+#include "AudioSink.h"
+#include "SessionDescriptionObservers.h"
+#include "VideoSource.h"
+#include "RTCStatsCollector.h"
+#include "PixelStreamingDelegates.h"
 
-DEFINE_LOG_CATEGORY(PixelStreamer);
-
-bool FStreamer::CheckPlatformCompatibility()
+bool UE::PixelStreaming::FStreamer::IsPlatformCompatible()
 {
 	return AVEncoder::FVideoEncoderFactory::Get().HasEncoderForCodec(AVEncoder::ECodecType::H264);
 }
 
-FStreamer::FStreamer(const FString& InSignallingServerUrl, const FString& InStreamerId)
+UE::PixelStreaming::FStreamer::FStreamer(const FString& InSignallingServerUrl, const FString& InStreamerId)
 	: SignallingServerUrl(InSignallingServerUrl)
 	, StreamerId(InStreamerId)
 	, WebRtcSignallingThread(MakeUnique<rtc::Thread>(rtc::SocketServer::CreateDefault()))
-	, SignallingServerConnection(MakeUnique<FSignallingServerConnection>(*this, InStreamerId))
-	, FramePump(&FrameSource)
+	, SignallingServerConnection(MakeUnique<UE::PixelStreaming::FSignallingServerConnection>(*this, InStreamerId))
 	, PlayerSessions(WebRtcSignallingThread.Get())
 	, Stats(&PlayerSessions)
 {
@@ -41,21 +39,25 @@ FStreamer::FStreamer(const FString& InSignallingServerUrl, const FString& InStre
 	StartWebRtcSignallingThread();
 	ConnectToSignallingServer();
 
-	this->PlayerSessions.OnQualityControllerChanged.AddRaw(this, &FStreamer::OnQualityControllerChanged);
-	this->PlayerSessions.OnPlayerDeleted.AddRaw(this, &FStreamer::PostPlayerDeleted);
+
+	if (UPixelStreamingDelegates* Delegates = UPixelStreamingDelegates::GetPixelStreamingDelegates())
+	{
+		Delegates->OnClosedConnectionNative.AddRaw(this, &UE::PixelStreaming::FStreamer::PostPlayerDeleted);
+		Delegates->OnQualityControllerChangedNative.AddRaw(this, &UE::PixelStreaming::FStreamer::OnQualityControllerChanged);
+	}
 }
 
-FStreamer::~FStreamer()
+UE::PixelStreaming::FStreamer::~FStreamer()
 {
 	DeleteAllPlayerSessions();
-	PeerConnectionFactory = nullptr;
+	P2PPeerConnectionFactory = nullptr;
 	WebRtcSignallingThread->Stop();
 	rtc::CleanupSSL();
 }
 
-void FStreamer::StartWebRtcSignallingThread()
+void UE::PixelStreaming::FStreamer::StartWebRtcSignallingThread()
 {
-	// initialisation of WebRTC stuff and things that depends on it should happen in WebRTC signalling thread
+	// Initialisation of WebRTC. Note: That interfacing with WebRTC should generally happen in WebRTC signalling thread.
 
 	// Create our own WebRTC thread for signalling
 	WebRtcSignallingThread->SetName("WebRtcSignallingThread", nullptr);
@@ -65,35 +67,58 @@ void FStreamer::StartWebRtcSignallingThread()
 
 	PeerConnectionConfig = {};
 
-	std::unique_ptr<FPixelStreamingSimulcastEncoderFactory> videoEncoderFactory = std::make_unique<FPixelStreamingSimulcastEncoderFactory>(this);
-	VideoEncoderFactory = videoEncoderFactory->GetRealFactory();
-
-	bool bUseLegacyAudioDeviceModule = PixelStreamingSettings::CVarPixelStreamingWebRTCUseLegacyAudioDevice.GetValueOnAnyThread();
+	bool bUseLegacyAudioDeviceModule = UE::PixelStreaming::Settings::CVarPixelStreamingWebRTCUseLegacyAudioDevice.GetValueOnAnyThread();
 	rtc::scoped_refptr<webrtc::AudioDeviceModule> AudioDeviceModule;
 	if (bUseLegacyAudioDeviceModule)
 	{
-		AudioDeviceModule = new rtc::RefCountedObject<FAudioCapturer>();
+		AudioDeviceModule = new rtc::RefCountedObject<UE::PixelStreaming::FAudioCapturer>();
 	}
 	else
 	{
-		AudioDeviceModule = new rtc::RefCountedObject<FPixelStreamingAudioDeviceModule>();
+		AudioDeviceModule = new rtc::RefCountedObject<UE::PixelStreaming::FAudioDeviceModule>();
 	}
 
-	PeerConnectionFactory = webrtc::CreatePeerConnectionFactory(
+	/* ---------- P2P Peer Connection Factory ---------- */
+
+	std::unique_ptr<UE::PixelStreaming::FVideoEncoderFactory> P2PPeerConnectionFactoryPtr = std::make_unique<UE::PixelStreaming::FVideoEncoderFactory>();
+	P2PVideoEncoderFactory = P2PPeerConnectionFactoryPtr.get(); 
+
+	P2PPeerConnectionFactory = webrtc::CreatePeerConnectionFactory(
 		nullptr,					  // network_thread
 		nullptr,					  // worker_thread
 		WebRtcSignallingThread.Get(), // signal_thread
 		AudioDeviceModule,			  // audio device manager
 		webrtc::CreateAudioEncoderFactory<webrtc::AudioEncoderOpus>(),
 		webrtc::CreateAudioDecoderFactory<webrtc::AudioDecoderOpus>(),
-		std::move(videoEncoderFactory),
+		std::move(P2PPeerConnectionFactoryPtr),
 		std::make_unique<webrtc::InternalDecoderFactory>(),
-		nullptr,							 // audio_mixer
-		this->SetupAudioProcessingModule()); // audio_processing
-	check(PeerConnectionFactory);
+		nullptr,					   // audio_mixer
+		SetupAudioProcessingModule()); // audio_processing
+
+	check(P2PPeerConnectionFactory.get() != nullptr);
+
+	/* ---------- SFU Peer Connection Factory ---------- */
+
+	std::unique_ptr<UE::PixelStreaming::FSimulcastEncoderFactory> SFUPeerConnectionFactoryPtr = std::make_unique<UE::PixelStreaming::FSimulcastEncoderFactory>();
+	SFUVideoEncoderFactory = SFUPeerConnectionFactoryPtr.get(); 
+
+	SFUPeerConnectionFactory = webrtc::CreatePeerConnectionFactory(
+		nullptr,					  // network_thread
+		nullptr,					  // worker_thread
+		WebRtcSignallingThread.Get(), // signal_thread
+		AudioDeviceModule,			  // audio device manager
+		webrtc::CreateAudioEncoderFactory<webrtc::AudioEncoderOpus>(),
+		webrtc::CreateAudioDecoderFactory<webrtc::AudioDecoderOpus>(),
+		std::move(SFUPeerConnectionFactoryPtr),
+		std::make_unique<webrtc::InternalDecoderFactory>(),
+		nullptr,					   // audio_mixer
+		SetupAudioProcessingModule()); // audio_processing
+
+	check(SFUPeerConnectionFactory.get() != nullptr);
+
 }
 
-webrtc::AudioProcessing* FStreamer::SetupAudioProcessingModule()
+webrtc::AudioProcessing* UE::PixelStreaming::FStreamer::SetupAudioProcessingModule()
 {
 	webrtc::AudioProcessing* AudioProcessingModule = webrtc::AudioProcessingBuilder().Create();
 	webrtc::AudioProcessing::Config Config;
@@ -121,57 +146,57 @@ webrtc::AudioProcessing* FStreamer::SetupAudioProcessingModule()
 	return AudioProcessingModule;
 }
 
-void FStreamer::ConnectToSignallingServer()
+void UE::PixelStreaming::FStreamer::ConnectToSignallingServer()
 {
-	this->SignallingServerConnection->Connect(SignallingServerUrl);
+	SignallingServerConnection->Connect(SignallingServerUrl);
 }
 
-void FStreamer::OnFrameBufferReady(const FTexture2DRHIRef& FrameBuffer)
+void UE::PixelStreaming::FStreamer::OnFrameBufferReady(const FTexture2DRHIRef& FrameBuffer)
 {
-	if (bStreamingStarted)
-	{
-		FrameSource.OnFrameReady(FrameBuffer);
-	}
 }
 
-void FStreamer::OnQualityControllerChanged(FPlayerId PlayerId)
+void UE::PixelStreaming::FStreamer::OnQualityControllerChanged(FPixelStreamingPlayerId PlayerId)
 {
-	FramePump.SetQualityController(PlayerId);
 }
 
-void FStreamer::PostPlayerDeleted(FPlayerId PlayerId)
+void UE::PixelStreaming::FStreamer::PostPlayerDeleted(FPixelStreamingPlayerId PlayerId, bool WasQualityController)
 {
-	FramePump.RemovePlayerVideoSource(PlayerId);
 }
 
-void FStreamer::OnConfig(const webrtc::PeerConnectionInterface::RTCConfiguration& Config)
+void UE::PixelStreaming::FStreamer::OnConfig(const webrtc::PeerConnectionInterface::RTCConfiguration& Config)
 {
-	this->PeerConnectionConfig = Config;
+	PeerConnectionConfig = Config;
 }
 
-void FStreamer::OnDataChannelOpen(FPlayerId PlayerId, webrtc::DataChannelInterface* DataChannel)
+void UE::PixelStreaming::FStreamer::OnDataChannelOpen(FPixelStreamingPlayerId PlayerId, webrtc::DataChannelInterface* DataChannel)
 {
 	if (DataChannel)
 	{
 		// When data channel is open, try to send cached freeze frame (if we have one)
-		this->SendCachedFreezeFrameTo(PlayerId);
+		SendCachedFreezeFrameTo(PlayerId);
 	}
 }
 
-webrtc::PeerConnectionInterface* FStreamer::CreateSession(FPlayerId PlayerId, int Flags)
+webrtc::PeerConnectionInterface* UE::PixelStreaming::FStreamer::CreateSession(FPixelStreamingPlayerId PlayerId, int Flags)
 {
 	PeerConnectionConfig.enable_simulcast_stats = true;
 
+	bool bIsSFU = (Flags & UE::PixelStreaming::Protocol::EPlayerFlags::PSPFlag_IsSFU) != UE::PixelStreaming::Protocol::EPlayerFlags::PSPFlag_None;
+	
+	rtc::scoped_refptr<webrtc::PeerConnectionFactoryInterface> PCFactory = bIsSFU ?
+		SFUPeerConnectionFactory :
+		P2PPeerConnectionFactory;
+
 	webrtc::PeerConnectionInterface* PeerConnection = PlayerSessions.CreatePlayerSession(
 		PlayerId,
-		PeerConnectionFactory,
+		PCFactory,
 		PeerConnectionConfig,
 		SignallingServerConnection.Get(),
 		Flags);
 
 	if (PeerConnection != nullptr)
 	{
-		if ((Flags & PixelStreamingProtocol::EPixelStreamingPlayerFlags::PSPFlag_SupportsDataChannel) != PixelStreamingProtocol::EPixelStreamingPlayerFlags::PSPFlag_None)
+		if ((Flags & UE::PixelStreaming::Protocol::EPlayerFlags::PSPFlag_SupportsDataChannel) != UE::PixelStreaming::Protocol::EPlayerFlags::PSPFlag_None)
 		{
 			webrtc::DataChannelInit DataChannelConfig;
 			DataChannelConfig.reliable = true;
@@ -181,19 +206,20 @@ webrtc::PeerConnectionInterface* FStreamer::CreateSession(FPlayerId PlayerId, in
 			PlayerSessions.SetPlayerSessionDataChannel(PlayerId, DataChannel);
 
 			// Bind to the on data channel open delegate (we use this as the earliest time peer is ready for freeze frame)
-			FPixelStreamingDataChannelObserver* DataChannelObserver = PlayerSessions.GetDataChannelObserver(PlayerId);
+			UE::PixelStreaming::FDataChannelObserver* DataChannelObserver = PlayerSessions.GetDataChannelObserver(PlayerId);
 			check(DataChannelObserver);
-			DataChannelObserver->OnDataChannelOpen.AddRaw(this, &FStreamer::OnDataChannelOpen);
+			if (UPixelStreamingDelegates* Delegates = UPixelStreamingDelegates::GetPixelStreamingDelegates())
+			{
+				Delegates->OnDataChannelOpenNative.AddRaw(this, &UE::PixelStreaming::FStreamer::OnDataChannelOpen);
+			}
 			DataChannelObserver->Register(DataChannel);
 		}
-
-		AddStreams(PlayerId, PeerConnection, Flags);
 	}
 
 	return PeerConnection;
 }
 
-void FStreamer::OnSessionDescription(FPlayerId PlayerId, webrtc::SdpType Type, const FString& Sdp)
+void UE::PixelStreaming::FStreamer::OnSessionDescription(FPixelStreamingPlayerId PlayerId, webrtc::SdpType Type, const FString& Sdp)
 {
 	switch (Type)
 	{
@@ -206,21 +232,21 @@ void FStreamer::OnSessionDescription(FPlayerId PlayerId, webrtc::SdpType Type, c
 			SetStreamingStarted(true);
 			break;
 		case webrtc::SdpType::kRollback:
-			UE_LOG(PixelStreamer, Error, TEXT("Rollback SDP is currently unsupported. SDP is: %s"), *Sdp);
+			UE_LOG(LogPixelStreaming, Error, TEXT("Rollback SDP is currently unsupported. SDP is: %s"), *Sdp);
 			break;
 	}
 }
 
-void FStreamer::OnOffer(FPlayerId PlayerId, const FString& Sdp)
+void UE::PixelStreaming::FStreamer::OnOffer(FPixelStreamingPlayerId PlayerId, const FString& Sdp)
 {
-	webrtc::PeerConnectionInterface* PeerConnection = CreateSession(PlayerId, PixelStreamingProtocol::EPixelStreamingPlayerFlags::PSPFlag_SupportsDataChannel);
+	webrtc::PeerConnectionInterface* PeerConnection = CreateSession(PlayerId, UE::PixelStreaming::Protocol::EPlayerFlags::PSPFlag_SupportsDataChannel);
 	if (PeerConnection)
 	{
 		webrtc::SdpParseError Error;
 		std::unique_ptr<webrtc::SessionDescriptionInterface> SessionDesc = webrtc::CreateSessionDescription(webrtc::SdpType::kOffer, to_string(Sdp), &Error);
 		if (!SessionDesc)
 		{
-			UE_LOG(PixelStreamer, Error, TEXT("Failed to parse offer's SDP: %s\n%s"), Error.description.c_str(), *Sdp);
+			UE_LOG(LogPixelStreaming, Error, TEXT("Failed to parse offer's SDP: %s\n%s"), Error.description.c_str(), *Sdp);
 			return;
 		}
 
@@ -229,22 +255,16 @@ void FStreamer::OnOffer(FPlayerId PlayerId, const FString& Sdp)
 	}
 	else
 	{
-		UE_LOG(PixelStreamer, Error, TEXT("Failed to create player session, peer connection was nullptr."));
+		UE_LOG(LogPixelStreaming, Error, TEXT("Failed to create player session, peer connection was nullptr."));
 	}
 }
 
-void FStreamer::SetLocalDescription(webrtc::PeerConnectionInterface* PeerConnection, FSetSessionDescriptionObserver* Observer, webrtc::SessionDescriptionInterface* SDP)
+void UE::PixelStreaming::FStreamer::SetLocalDescription(webrtc::PeerConnectionInterface* PeerConnection, UE::PixelStreaming::FSetSessionDescriptionObserver* Observer, webrtc::SessionDescriptionInterface* SDP)
 {
 	// Note from Luke about WebRTC: the sink of video capturer will be added as a direct result
 	// of `PeerConnection->SetLocalDescription()` call but video encoder will be created later on
 	// when the first frame is pushed into the WebRTC pipeline (by the capturer calling `OnFrame`).
 
-	// Note from Luke about Pixel Streaming: Internally we associate `FPlayerId` with each `FVideoCapturer` and
-	// pass an "initialisation frame" that contains `FPlayerId` through `OnFrame` to the encoder.
-	// This initialization frame establishes the correct association between the player and `FPixelStreamingVideoEncoder`.
-	// The reason we want this association between encoder<->peer is is so that the quality controlling peer
-	// can change encoder bitrate etc, while the other peers cannot.
-
 	PeerConnection->SetLocalDescription(Observer, SDP);
 
 	// Once local description has been set we can start setting some encoding information for the video stream rtp sender
@@ -256,22 +276,20 @@ void FStreamer::SetLocalDescription(webrtc::PeerConnectionInterface* PeerConnect
 			webrtc::RtpParameters ExistingParams = Sender->GetParameters();
 
 			// Set the degradation preference based on our CVar for it.
-			ExistingParams.degradation_preference = PixelStreamingSettings::GetDegradationPreference();
+			ExistingParams.degradation_preference = UE::PixelStreaming::Settings::GetDegradationPreference();
 
 			webrtc::RTCError Err = Sender->SetParameters(ExistingParams);
 			if (!Err.ok())
 			{
 				const char* ErrMsg = Err.message();
 				FString ErrorStr(ErrMsg);
-				UE_LOG(PixelStreamer, Error, TEXT("Failed to set RTP Sender params: %s"), *ErrorStr);
+				UE_LOG(LogPixelStreaming, Error, TEXT("Failed to set RTP Sender params: %s"), *ErrorStr);
 			}
 		}
 	}
-
-	// Note: if desired, manually limiting total bitrate for a peer is possible in PeerConnection::SetBitrate(const BitrateSettings& bitrate)
 }
 
-void FStreamer::SendAnswer(FPlayerId PlayerId, webrtc::PeerConnectionInterface* PeerConnection, TUniquePtr<webrtc::SessionDescriptionInterface> Sdp)
+void UE::PixelStreaming::FStreamer::SendAnswer(FPixelStreamingPlayerId PlayerId, webrtc::PeerConnectionInterface* PeerConnection, TUniquePtr<webrtc::SessionDescriptionInterface> Sdp)
 {
 	// below is async execution (with error handling) of:
 	//		PeerConnection.SetRemoteDescription(SDP);
@@ -279,34 +297,34 @@ void FStreamer::SendAnswer(FPlayerId PlayerId, webrtc::PeerConnectionInterface*
 	//		PeerConnection.SetLocalDescription(Answer);
 	//		SignallingServerConnection.SendAnswer(Answer);
 
-	FSetSessionDescriptionObserver* SetLocalDescriptionObserver = FSetSessionDescriptionObserver::Create(
+	UE::PixelStreaming::FSetSessionDescriptionObserver* SetLocalDescriptionObserver = UE::PixelStreaming::FSetSessionDescriptionObserver::Create(
 		[this, PlayerId, PeerConnection]() // on success
 		{
-			this->GetSignallingServerConnection()->SendAnswer(PlayerId, *PeerConnection->local_description());
-			this->SetStreamingStarted(true);
+			GetSignallingServerConnection()->SendAnswer(PlayerId, *PeerConnection->local_description());
+			SetStreamingStarted(true);
 		},
 		[this, PlayerId](const FString& Error) // on failure
 		{
-			UE_LOG(PixelStreamer, Error, TEXT("Failed to set local description: %s"), *Error);
-			this->PlayerSessions.DisconnectPlayer(PlayerId, Error);
+			UE_LOG(LogPixelStreaming, Error, TEXT("Failed to set local description: %s"), *Error);
+			PlayerSessions.DisconnectPlayer(PlayerId, Error);
 		});
 
 	auto OnCreateAnswerSuccess = [this, PlayerId, PeerConnection, SetLocalDescriptionObserver](webrtc::SessionDescriptionInterface* SDP) {
 		SetLocalDescription(PeerConnection, SetLocalDescriptionObserver, SDP);
 	};
 
-	FCreateSessionDescriptionObserver* CreateAnswerObserver = FCreateSessionDescriptionObserver::Create(
+	UE::PixelStreaming::FCreateSessionDescriptionObserver* CreateAnswerObserver = UE::PixelStreaming::FCreateSessionDescriptionObserver::Create(
 		MoveTemp(OnCreateAnswerSuccess),
 		[this, PlayerId](const FString& Error) // on failure
 		{
-			UE_LOG(PixelStreamer, Error, TEXT("Failed to create answer: %s"), *Error);
-			this->PlayerSessions.DisconnectPlayer(PlayerId, Error);
+			UE_LOG(LogPixelStreaming, Error, TEXT("Failed to create answer: %s"), *Error);
+			PlayerSessions.DisconnectPlayer(PlayerId, Error);
 		});
 
-	auto OnSetRemoteDescriptionSuccess = [this, PeerConnection, CreateAnswerObserver]() {
-		// Note: these offer to receive at superseded now we are use transceivers to setup our peer connection media
-		int offer_to_receive_video = 0;
-		int offer_to_receive_audio = PixelStreamingSettings::CVarPixelStreamingWebRTCDisableReceiveAudio.GetValueOnAnyThread() ? 0 : 1;
+	auto OnSetRemoteDescriptionSuccess = [this, PlayerId, PeerConnection, CreateAnswerObserver]() {
+		// Note: these offer to receive are superseded now we are use transceivers to setup our peer connection media
+		int offer_to_receive_video = webrtc::PeerConnectionInterface::RTCOfferAnswerOptions::kUndefined; 
+		int offer_to_receive_audio = webrtc::PeerConnectionInterface::RTCOfferAnswerOptions::kUndefined; 
 		bool voice_activity_detection = false;
 		bool ice_restart = true;
 		bool use_rtp_mux = true;
@@ -319,263 +337,267 @@ void FStreamer::SendAnswer(FPlayerId PlayerId, webrtc::PeerConnectionInterface*
 			use_rtp_mux
 		};
 
-		// modify audio transceiver direction based on CVars just before creating answer
-		this->ModifyAudioTransceiverDirection(PeerConnection);
+		AddStreams(PlayerId, PeerConnection, UE::PixelStreaming::Protocol::EPlayerFlags::PSPFlag_SupportsDataChannel);
+		ModifyTransceivers(PeerConnection->GetTransceivers(), UE::PixelStreaming::Protocol::EPlayerFlags::PSPFlag_SupportsDataChannel);
 
 		PeerConnection->CreateAnswer(CreateAnswerObserver, AnswerOption);
 	};
 
-	FSetSessionDescriptionObserver* SetRemoteDescriptionObserver = FSetSessionDescriptionObserver::Create(
+	UE::PixelStreaming::FSetSessionDescriptionObserver* SetRemoteDescriptionObserver = UE::PixelStreaming::FSetSessionDescriptionObserver::Create(
 		MoveTemp(OnSetRemoteDescriptionSuccess),
 		[this, PlayerId](const FString& Error) // on failure
 		{
-			UE_LOG(PixelStreamer, Error, TEXT("Failed to set remote description: %s"), *Error);
-			this->PlayerSessions.DisconnectPlayer(PlayerId, Error);
+			UE_LOG(LogPixelStreaming, Error, TEXT("Failed to set remote description: %s"), *Error);
+			PlayerSessions.DisconnectPlayer(PlayerId, Error);
 		});
 
+	cricket::SessionDescription* RemoteDescription = Sdp->description();
+	MungeRemoteSDP(RemoteDescription);
+
 	PeerConnection->SetRemoteDescription(SetRemoteDescriptionObserver, Sdp.Release());
 }
 
-void FStreamer::ModifyAudioTransceiverDirection(webrtc::PeerConnectionInterface* PeerConnection)
+void UE::PixelStreaming::FStreamer::MungeRemoteSDP(cricket::SessionDescription* RemoteDescription)
 {
-	//virtual std::vector<rtc::scoped_refptr<RtpTransceiverInterface>>GetTransceivers()
-	std::vector<rtc::scoped_refptr<webrtc::RtpTransceiverInterface>> Transceivers = PeerConnection->GetTransceivers();
-	for (auto& Transceiver : Transceivers)
+	// Munge SDP of remote description to inject min, max, start bitrates
+	std::vector<cricket::ContentInfo>& ContentInfos = RemoteDescription->contents();
+	for(cricket::ContentInfo& Content : ContentInfos)
 	{
-		if (Transceiver->media_type() == cricket::MediaType::MEDIA_TYPE_AUDIO)
+		cricket::MediaContentDescription* MediaDescription = Content.media_description();
+		if(MediaDescription->type() == cricket::MediaType::MEDIA_TYPE_VIDEO)
 		{
-
-			bool bTransmitUEAudio = !PixelStreamingSettings::CVarPixelStreamingWebRTCDisableTransmitAudio.GetValueOnAnyThread();
-			bool bReceiveBrowserAudio = !PixelStreamingSettings::CVarPixelStreamingWebRTCDisableReceiveAudio.GetValueOnAnyThread();
-
-			// Determine the direction of the transceiver
-			webrtc::RtpTransceiverDirection AudioTransceiverDirection;
-			if (bTransmitUEAudio && bReceiveBrowserAudio)
+			cricket::VideoContentDescription* VideoDescription = MediaDescription->as_video();
+			std::vector<cricket::VideoCodec> CodecsCopy = VideoDescription->codecs();
+			for(cricket::VideoCodec& Codec : CodecsCopy)
 			{
-				AudioTransceiverDirection = webrtc::RtpTransceiverDirection::kSendRecv;
-			}
-			else if (bTransmitUEAudio)
-			{
-				AudioTransceiverDirection = webrtc::RtpTransceiverDirection::kSendOnly;
+				// Note: These params are passed as kilobits, so divide by 1000.
+				Codec.SetParam(cricket::kCodecParamMinBitrate, UE::PixelStreaming::Settings::CVarPixelStreamingWebRTCMinBitrate.GetValueOnAnyThread() / 1000);
+				Codec.SetParam(cricket::kCodecParamStartBitrate, UE::PixelStreaming::Settings::CVarPixelStreamingWebRTCStartBitrate.GetValueOnAnyThread() / 1000);
+				Codec.SetParam(cricket::kCodecParamMaxBitrate, UE::PixelStreaming::Settings::CVarPixelStreamingWebRTCMaxBitrate.GetValueOnAnyThread() / 1000);
 			}
-			else if (bReceiveBrowserAudio)
-			{
-				AudioTransceiverDirection = webrtc::RtpTransceiverDirection::kRecvOnly;
-			}
-			else
-			{
-				AudioTransceiverDirection = webrtc::RtpTransceiverDirection::kInactive;
-			}
-
-			Transceiver->SetDirection(AudioTransceiverDirection);
-		}
+			VideoDescription->set_codecs(CodecsCopy);
+		}	
 	}
 }
 
-void FStreamer::OnRemoteIceCandidate(FPlayerId PlayerId, const std::string& SdpMid, int SdpMLineIndex, const std::string& Sdp)
+void UE::PixelStreaming::FStreamer::OnRemoteIceCandidate(FPixelStreamingPlayerId PlayerId, const std::string& SdpMid, int SdpMLineIndex, const std::string& Sdp)
 {
-	this->PlayerSessions.OnRemoteIceCandidate(PlayerId, SdpMid, SdpMLineIndex, Sdp);
+	PlayerSessions.OnRemoteIceCandidate(PlayerId, SdpMid, SdpMLineIndex, Sdp);
 }
 
-void FStreamer::OnPlayerConnected(FPlayerId PlayerId, int Flags)
+void UE::PixelStreaming::FStreamer::OnPlayerConnected(FPixelStreamingPlayerId PlayerId, int Flags)
 {
 	// create peer connection
-	if (webrtc::PeerConnectionInterface* PlayerSession = CreateSession(PlayerId, Flags))
+	if (webrtc::PeerConnectionInterface* PeerConnection = CreateSession(PlayerId, Flags))
 	{
-		// make them send only
-		std::vector<rtc::scoped_refptr<webrtc::RtpTransceiverInterface>> Transceivers = PlayerSession->GetTransceivers();
-		for (auto& Transceiver : Transceivers)
-		{
-			Transceiver->SetDirection(webrtc::RtpTransceiverDirection::kSendOnly);
-		}
+		AddStreams(PlayerId, PeerConnection, Flags);
+		ModifyTransceivers(PeerConnection->GetTransceivers(), Flags);
 
 		// observer for creating offer
-		FCreateSessionDescriptionObserver* CreateOfferObserver = FCreateSessionDescriptionObserver::Create(
-			[this, PlayerId, PlayerSession](webrtc::SessionDescriptionInterface* SDP) // on SDP create success
+		UE::PixelStreaming::FCreateSessionDescriptionObserver* CreateOfferObserver = UE::PixelStreaming::FCreateSessionDescriptionObserver::Create(
+			[this, PlayerId, PeerConnection](webrtc::SessionDescriptionInterface* SDP) // on SDP create success
 			{
-				FSetSessionDescriptionObserver* SetLocalDescriptionObserver = FSetSessionDescriptionObserver::Create(
+				UE::PixelStreaming::FSetSessionDescriptionObserver* SetLocalDescriptionObserver = UE::PixelStreaming::FSetSessionDescriptionObserver::Create(
 					[this, PlayerId, SDP]() // on SDP set success
 					{
 						GetSignallingServerConnection()->SendOffer(PlayerId, *SDP);
 					},
 					[](const FString& Error) // on SDP set failure
 					{
-						UE_LOG(PixelStreamer, Error, TEXT("Failed to set local description: %s"), *Error);
+						UE_LOG(LogPixelStreaming, Error, TEXT("Failed to set local description: %s"), *Error);
 					});
 
-				SetLocalDescription(PlayerSession, SetLocalDescriptionObserver, SDP);
+				SetLocalDescription(PeerConnection, SetLocalDescriptionObserver, SDP);
 			},
 			[](const FString& Error) // on SDP create failure
 			{
-				UE_LOG(PixelStreamer, Error, TEXT("Failed to create offer: %s"), *Error);
+				UE_LOG(LogPixelStreaming, Error, TEXT("Failed to create offer: %s"), *Error);
 			});
 
-		PlayerSession->CreateOffer(CreateOfferObserver, {});
+		PeerConnection->CreateOffer(CreateOfferObserver, {});
 	}
 }
 
-void FStreamer::OnPlayerDisconnected(FPlayerId PlayerId)
+void UE::PixelStreaming::FStreamer::OnPlayerDisconnected(FPixelStreamingPlayerId PlayerId)
 {
-	UE_LOG(PixelStreamer, Log, TEXT("player %s disconnected"), *PlayerId);
+	UE_LOG(LogPixelStreaming, Log, TEXT("player %s disconnected"), *PlayerId);
 	DeletePlayerSession(PlayerId);
 }
 
-void FStreamer::OnSignallingServerDisconnected()
+void UE::PixelStreaming::FStreamer::OnSignallingServerDisconnected()
 {
 	DeleteAllPlayerSessions();
 
 	// Call disconnect here makes sure any other websocket callbacks etc are cleared
-	this->SignallingServerConnection->Disconnect();
+	SignallingServerConnection->Disconnect();
 	ConnectToSignallingServer();
 }
 
-int FStreamer::GetNumPlayers() const
+int UE::PixelStreaming::FStreamer::GetNumPlayers() const
 {
-	return this->PlayerSessions.GetNumPlayers();
+	return PlayerSessions.GetNumPlayers();
 }
 
-void FStreamer::ForceKeyFrame()
+void UE::PixelStreaming::FStreamer::ForceKeyFrame()
 {
-	if (this->VideoEncoderFactory)
+	if (P2PVideoEncoderFactory)
 	{
-		this->VideoEncoderFactory->ForceKeyFrame();
+		P2PVideoEncoderFactory->ForceKeyFrame();
 	}
 	else
 	{
-		UE_LOG(PixelStreamer, Log, TEXT("Cannot force a key frame - video encoder factory is nullptr."));
+		UE_LOG(LogPixelStreaming, Log, TEXT("Cannot force a key frame - video encoder factory is nullptr."));
 	}
 }
 
-void FStreamer::SetQualityController(FPlayerId PlayerId)
+void UE::PixelStreaming::FStreamer::SetQualityController(FPixelStreamingPlayerId PlayerId)
 {
 
-	this->PlayerSessions.SetQualityController(PlayerId);
+	PlayerSessions.SetQualityController(PlayerId);
 }
 
-bool FStreamer::IsQualityController(FPlayerId PlayerId) const
+bool UE::PixelStreaming::FStreamer::IsQualityController(FPixelStreamingPlayerId PlayerId) const
 {
-	return this->PlayerSessions.IsQualityController(PlayerId);
+	return PlayerSessions.IsQualityController(PlayerId);
 }
 
-IPixelStreamingAudioSink* FStreamer::GetAudioSink(FPlayerId PlayerId) const
+IPixelStreamingAudioSink* UE::PixelStreaming::FStreamer::GetAudioSink(FPixelStreamingPlayerId PlayerId) const
 {
-	return this->PlayerSessions.GetAudioSink(PlayerId);
+	return PlayerSessions.GetAudioSink(PlayerId);
 }
 
-IPixelStreamingAudioSink* FStreamer::GetUnlistenedAudioSink() const
+IPixelStreamingAudioSink* UE::PixelStreaming::FStreamer::GetUnlistenedAudioSink() const
 {
-	return this->PlayerSessions.GetUnlistenedAudioSink();
+	return PlayerSessions.GetUnlistenedAudioSink();
 }
 
-bool FStreamer::SendMessage(FPlayerId PlayerId, PixelStreamingProtocol::EToPlayerMsg Type, const FString& Descriptor) const
+bool UE::PixelStreaming::FStreamer::SendMessage(FPixelStreamingPlayerId PlayerId, UE::PixelStreaming::Protocol::EToPlayerMsg Type, const FString& Descriptor) const
 {
-	return this->PlayerSessions.SendMessage(PlayerId, Type, Descriptor);
+	return PlayerSessions.SendMessage(PlayerId, Type, Descriptor);
 }
 
-void FStreamer::SendLatestQP(FPlayerId PlayerId, int LatestQP) const
+void UE::PixelStreaming::FStreamer::SendLatestQP(FPixelStreamingPlayerId PlayerId, int LatestQP) const
 {
-	this->PlayerSessions.SendLatestQP(PlayerId, LatestQP);
+	PlayerSessions.SendLatestQP(PlayerId, LatestQP);
 }
 
-void FStreamer::PollWebRTCStats() const
+void UE::PixelStreaming::FStreamer::PollWebRTCStats() const
 {
-	this->PlayerSessions.PollWebRTCStats();
+	PlayerSessions.PollWebRTCStats();
 }
 
-void FStreamer::DeleteAllPlayerSessions()
+void UE::PixelStreaming::FStreamer::DeleteAllPlayerSessions()
 {
-	this->PlayerSessions.DeleteAllPlayerSessions();
-	this->bStreamingStarted = false;
+	PlayerSessions.DeleteAllPlayerSessions();
+	bStreamingStarted = false;
 }
 
-void FStreamer::DeletePlayerSession(FPlayerId PlayerId)
+void UE::PixelStreaming::FStreamer::DeletePlayerSession(FPixelStreamingPlayerId PlayerId)
 {
-	int NumRemainingPlayers = this->PlayerSessions.DeletePlayerSession(PlayerId);
+	int NumRemainingPlayers = PlayerSessions.DeletePlayerSession(PlayerId);
 
 	if (NumRemainingPlayers == 0)
 	{
-		this->bStreamingStarted = false;
+		bStreamingStarted = false;
 	}
 }
 
-void FStreamer::AddStreams(FPlayerId PlayerId, webrtc::PeerConnectionInterface* PeerConnection, int Flags)
+FString UE::PixelStreaming::FStreamer::GetAudioStreamID() const
 {
-	check(PeerConnection);
-
-	bool bSyncVideoAndAudio = !PixelStreamingSettings::CVarPixelStreamingWebRTCDisableAudioSync.GetValueOnAnyThread();
-
-	const FString AudioStreamId = bSyncVideoAndAudio ? TEXT("pixelstreaming_av_stream_id") : TEXT("pixelstreaming_audio_stream_id");
-	const FString VideoStreamId = bSyncVideoAndAudio ? TEXT("pixelstreaming_av_stream_id") : TEXT("pixelstreaming_video_stream_id");
-	const FString AudioTrackLabel = TEXT("pixelstreaming_audio_track_label");
-	const FString VideoTrackLabel = TEXT("pixelstreaming_video_track_label");
-
-	if (!PeerConnection->GetSenders().empty())
-	{
-		return; // Already added tracks
-	}
+	bool bSyncVideoAndAudio = !UE::PixelStreaming::Settings::CVarPixelStreamingWebRTCDisableAudioSync.GetValueOnAnyThread();
+	return bSyncVideoAndAudio ? TEXT("pixelstreaming_av_stream_id") : TEXT("pixelstreaming_audio_stream_id");
+}
 
-	// Use PeerConnection's transceiver API to add create audio/video tracks will correct directionality.
-	// These tracks are only thin wrappers around the underlying sources (the sources are shared among all peer's tracks).
-	// As per the WebRTC source: "The same source can be used by multiple VideoTracks."
-	this->SetupVideoTrack(PlayerId, PeerConnection, VideoStreamId, VideoTrackLabel, Flags);
-	this->SetupAudioTrack(PeerConnection, AudioStreamId, AudioTrackLabel);
+FString UE::PixelStreaming::FStreamer::GetVideoStreamID() const
+{
+	bool bSyncVideoAndAudio = !UE::PixelStreaming::Settings::CVarPixelStreamingWebRTCDisableAudioSync.GetValueOnAnyThread();
+	return bSyncVideoAndAudio ? TEXT("pixelstreaming_av_stream_id") : TEXT("pixelstreaming_video_stream_id");
 }
 
-void FStreamer::SetupVideoTrack(FPlayerId PlayerId, webrtc::PeerConnectionInterface* PeerConnection, const FString VideoStreamId, const FString VideoTrackLabel, int Flags)
+void UE::PixelStreaming::FStreamer::AddStreams(FPixelStreamingPlayerId PlayerId, webrtc::PeerConnectionInterface* PeerConnection, int Flags)
 {
-	// Create a video source for this player
-	FPlayerVideoSource* PlayerVideoSource = FramePump.CreatePlayerVideoSource(PlayerId, Flags);
+	check(PeerConnection);
 
-	// Create video track
-	rtc::scoped_refptr<webrtc::VideoTrackInterface> VideoTrack = PeerConnectionFactory->CreateVideoTrack(TCHAR_TO_UTF8(*VideoTrackLabel), PlayerVideoSource);
+	// Use PeerConnection's transceiver API to add create audio/video tracks.
+	SetupVideoTrack(PlayerId, PeerConnection, GetVideoStreamID(), VideoTrackLabel, Flags);
+	SetupAudioTrack(PeerConnection, GetAudioStreamID(), AudioTrackLabel, Flags);
+}
 
-	// Add tranceiver
-	webrtc::RtpTransceiverInit TransceiverOptions;
-	TransceiverOptions.stream_ids = { TCHAR_TO_UTF8(*VideoStreamId) };
+std::vector<webrtc::RtpEncodingParameters> UE::PixelStreaming::FStreamer::CreateRTPEncodingParams(int Flags)
+{
+	bool bIsSFU = (Flags & UE::PixelStreaming::Protocol::EPlayerFlags::PSPFlag_IsSFU) != UE::PixelStreaming::Protocol::EPlayerFlags::PSPFlag_None;
 
-	bool bIsSFU = (Flags & PixelStreamingProtocol::EPixelStreamingPlayerFlags::PSPFlag_IsSFU) != PixelStreamingProtocol::EPixelStreamingPlayerFlags::PSPFlag_None;
+	std::vector<webrtc::RtpEncodingParameters> EncodingParams;
 
-	if (bIsSFU && PixelStreamingSettings::SimulcastParameters.Layers.Num() > 0)
+	if (bIsSFU && UE::PixelStreaming::Settings::SimulcastParameters.Layers.Num() > 0)
 	{
+		using FLayer = UE::PixelStreaming::Settings::FSimulcastParameters::FLayer;
 
 		// encodings should be lowest res to highest
-		TArray<PixelStreamingSettings::FSimulcastParameters::FLayer*> SortedLayers;
-		for (auto& Layer : PixelStreamingSettings::SimulcastParameters.Layers)
+		TArray<FLayer*> SortedLayers;
+		for (FLayer& Layer : UE::PixelStreaming::Settings::SimulcastParameters.Layers)
 		{
 			SortedLayers.Add(&Layer);
 		}
 
-		SortedLayers.Sort([](const auto& LayerA, const auto& LayerB) { return LayerA.Scaling > LayerB.Scaling; });
+		SortedLayers.Sort([](const FLayer& LayerA, const FLayer& LayerB) { return LayerA.Scaling > LayerB.Scaling; });
 
 		const int LayerCount = SortedLayers.Num();
 		for (int i = 0; i < LayerCount; ++i)
 		{
-			const auto SimulcastLayer = SortedLayers[i];
+			const FLayer* SimulcastLayer = SortedLayers[i];
 			webrtc::RtpEncodingParameters LayerEncoding{};
 			LayerEncoding.rid = TCHAR_TO_UTF8(*(FString("simulcast") + FString::FromInt(LayerCount - i)));
 			LayerEncoding.min_bitrate_bps = SimulcastLayer->MinBitrate;
 			LayerEncoding.max_bitrate_bps = SimulcastLayer->MaxBitrate;
 			LayerEncoding.scale_resolution_down_by = SimulcastLayer->Scaling;
-			LayerEncoding.max_framerate = FMath::Max(60, PixelStreamingSettings::CVarPixelStreamingWebRTCFps.GetValueOnAnyThread());
-			TransceiverOptions.send_encodings.push_back(LayerEncoding);
+			LayerEncoding.max_framerate = FMath::Max(60, UE::PixelStreaming::Settings::CVarPixelStreamingWebRTCFps.GetValueOnAnyThread());
+			EncodingParams.push_back(LayerEncoding);
 		}
 	}
 	else
 	{
 		webrtc::RtpEncodingParameters encoding{};
 		encoding.rid = "base";
-		encoding.max_bitrate_bps = PixelStreamingSettings::CVarPixelStreamingWebRTCMaxBitrate.GetValueOnAnyThread();
-		encoding.min_bitrate_bps = PixelStreamingSettings::CVarPixelStreamingWebRTCMinBitrate.GetValueOnAnyThread();
-		encoding.max_framerate = FMath::Max(60, PixelStreamingSettings::CVarPixelStreamingWebRTCFps.GetValueOnAnyThread());
+		encoding.max_bitrate_bps = UE::PixelStreaming::Settings::CVarPixelStreamingWebRTCMaxBitrate.GetValueOnAnyThread();
+		encoding.min_bitrate_bps = UE::PixelStreaming::Settings::CVarPixelStreamingWebRTCMinBitrate.GetValueOnAnyThread();
+		encoding.max_framerate = FMath::Max(60, UE::PixelStreaming::Settings::CVarPixelStreamingWebRTCFps.GetValueOnAnyThread());
 		encoding.scale_resolution_down_by.reset();
-		TransceiverOptions.send_encodings.push_back(encoding);
+		EncodingParams.push_back(encoding);
+	}
+
+	return EncodingParams;
+}
+
+void UE::PixelStreaming::FStreamer::SetupVideoTrack(FPixelStreamingPlayerId PlayerId, webrtc::PeerConnectionInterface* PeerConnection, const FString InVideoStreamId, const FString InVideoTrackLabel, int Flags)
+{
+
+	bool bIsSFU = (Flags & UE::PixelStreaming::Protocol::EPlayerFlags::PSPFlag_IsSFU) != UE::PixelStreaming::Protocol::EPlayerFlags::PSPFlag_None;
+	
+	rtc::scoped_refptr<webrtc::PeerConnectionFactoryInterface> PCFactory = bIsSFU ?
+		SFUPeerConnectionFactory :
+		P2PPeerConnectionFactory;
+
+	// Create a video source for this player
+	rtc::scoped_refptr<FVideoSourceBase> VideoSource;
+	if(bIsSFU)
+	{
+		VideoSource = new FVideoSourceSFU();
+	}
+	else
+	{
+		VideoSource = new FVideoSourceP2P(PlayerId, this);
 	}
 
-	webrtc::RTCErrorOr<rtc::scoped_refptr<webrtc::RtpTransceiverInterface>> Result = PeerConnection->AddTransceiver(VideoTrack, TransceiverOptions);
-	checkf(Result.ok(), TEXT("Failed to add Video transceiver to PeerConnection. Msg=%s"), *FString(Result.error().message()));
+	PlayerSessions.SetVideoSource(PlayerId, VideoSource);
+	VideoSource->Initialize();
+
+	// Create video track
+	rtc::scoped_refptr<webrtc::VideoTrackInterface> VideoTrack = PCFactory->CreateVideoTrack(TCHAR_TO_UTF8(*InVideoTrackLabel), VideoSource.get());
+	VideoTrack->set_enabled(true);
 
 	// Set some content hints based on degradation prefs, WebRTC uses these internally.
-	webrtc::DegradationPreference DegradationPref = PixelStreamingSettings::GetDegradationPreference();
+	webrtc::DegradationPreference DegradationPref = UE::PixelStreaming::Settings::GetDegradationPreference();
 	switch (DegradationPref)
 	{
 		case webrtc::DegradationPreference::MAINTAIN_FRAMERATE:
@@ -587,12 +609,42 @@ void FStreamer::SetupVideoTrack(FPlayerId PlayerId, webrtc::PeerConnectionInterf
 		default:
 			break;
 	}
+
+	bool bHasVideoTransceiver = false;
+
+	for (auto& Transceiver : PeerConnection->GetTransceivers())
+	{
+		rtc::scoped_refptr<webrtc::RtpSenderInterface> Sender = Transceiver->sender();
+		if (Transceiver->media_type() == cricket::MediaType::MEDIA_TYPE_VIDEO)
+		{
+			bHasVideoTransceiver = true;
+			Sender->SetTrack(VideoTrack);
+		}
+	}
+
+	// If there is no existing video transceiver, add one.
+	if(!bHasVideoTransceiver)
+	{
+		webrtc::RtpTransceiverInit TransceiverOptions;
+		TransceiverOptions.stream_ids = { TCHAR_TO_UTF8(*InVideoStreamId) };
+		TransceiverOptions.direction = webrtc::RtpTransceiverDirection::kSendOnly;
+		TransceiverOptions.send_encodings = CreateRTPEncodingParams(Flags);
+
+		webrtc::RTCErrorOr<rtc::scoped_refptr<webrtc::RtpTransceiverInterface>> Result = PeerConnection->AddTransceiver(VideoTrack, TransceiverOptions);
+		checkf(Result.ok(), TEXT("Failed to add Video transceiver to PeerConnection. Msg=%s"), *FString(Result.error().message()));
+	}
 }
 
-void FStreamer::SetupAudioTrack(webrtc::PeerConnectionInterface* PeerConnection, FString const AudioStreamId, FString const AudioTrackLabel)
+void UE::PixelStreaming::FStreamer::SetupAudioTrack(webrtc::PeerConnectionInterface* PeerConnection, FString const InAudioStreamId, FString const InAudioTrackLabel, int Flags)
 {
-	bool bTransmitUEAudio = !PixelStreamingSettings::CVarPixelStreamingWebRTCDisableTransmitAudio.GetValueOnAnyThread();
-	bool bReceiveBrowserAudio = !PixelStreamingSettings::CVarPixelStreamingWebRTCDisableReceiveAudio.GetValueOnAnyThread();
+
+	bool bIsSFU = (Flags & UE::PixelStreaming::Protocol::EPlayerFlags::PSPFlag_IsSFU) != UE::PixelStreaming::Protocol::EPlayerFlags::PSPFlag_None;
+	
+	rtc::scoped_refptr<webrtc::PeerConnectionFactoryInterface> PCFactory = bIsSFU ?
+		SFUPeerConnectionFactory :
+		P2PPeerConnectionFactory;
+
+	bool bTransmitUEAudio = !UE::PixelStreaming::Settings::CVarPixelStreamingWebRTCDisableTransmitAudio.GetValueOnAnyThread();
 
 	// Create one and only one audio source for Pixel Streaming.
 	if (!AudioSource && bTransmitUEAudio)
@@ -600,90 +652,144 @@ void FStreamer::SetupAudioTrack(webrtc::PeerConnectionInterface* PeerConnection,
 		// Setup audio source options, we turn off many of the "nice" audio settings that
 		// would traditionally be used in a conference call because the audio source we are
 		// transmitting is UE application audio (not some unknown microphone).
-		this->AudioSourceOptions.echo_cancellation = false;
-		this->AudioSourceOptions.auto_gain_control = false;
-		this->AudioSourceOptions.noise_suppression = false;
-		this->AudioSourceOptions.highpass_filter = false;
-		this->AudioSourceOptions.stereo_swapping = false;
-		this->AudioSourceOptions.audio_jitter_buffer_max_packets = 1000;
-		this->AudioSourceOptions.audio_jitter_buffer_fast_accelerate = false;
-		this->AudioSourceOptions.audio_jitter_buffer_min_delay_ms = 0;
-		this->AudioSourceOptions.audio_jitter_buffer_enable_rtx_handling = false;
-		this->AudioSourceOptions.typing_detection = false;
-		this->AudioSourceOptions.experimental_agc = false;
-		this->AudioSourceOptions.experimental_ns = false;
-		this->AudioSourceOptions.residual_echo_detector = false;
+		AudioSourceOptions.echo_cancellation = false;
+		AudioSourceOptions.auto_gain_control = false;
+		AudioSourceOptions.noise_suppression = false;
+		AudioSourceOptions.highpass_filter = false;
+		AudioSourceOptions.stereo_swapping = false;
+		AudioSourceOptions.audio_jitter_buffer_max_packets = 1000;
+		AudioSourceOptions.audio_jitter_buffer_fast_accelerate = false;
+		AudioSourceOptions.audio_jitter_buffer_min_delay_ms = 0;
+		AudioSourceOptions.audio_jitter_buffer_enable_rtx_handling = false;
+		AudioSourceOptions.typing_detection = false;
+		AudioSourceOptions.experimental_agc = false;
+		AudioSourceOptions.experimental_ns = false;
+		AudioSourceOptions.residual_echo_detector = false;
 		// Create audio source
-		AudioSource = PeerConnectionFactory->CreateAudioSource(this->AudioSourceOptions);
+		AudioSource = PCFactory->CreateAudioSource(AudioSourceOptions);
 	}
 
 	// Add the audio track to the audio transceiver's sender if we are transmitting audio
-	if (bTransmitUEAudio)
+	if(!bTransmitUEAudio)
 	{
+		return;
+	}
 
-		rtc::scoped_refptr<webrtc::AudioTrackInterface> AudioTrack = PeerConnectionFactory->CreateAudioTrack(TCHAR_TO_UTF8(*AudioTrackLabel), AudioSource);
+	rtc::scoped_refptr<webrtc::AudioTrackInterface> AudioTrack = PCFactory->CreateAudioTrack(TCHAR_TO_UTF8(*InAudioTrackLabel), AudioSource);
 
+	bool bHasAudioTransceiver = false;
+	for (auto& Transceiver : PeerConnection->GetTransceivers())
+	{
+		rtc::scoped_refptr<webrtc::RtpSenderInterface> Sender = Transceiver->sender();
+		if (Transceiver->media_type() == cricket::MediaType::MEDIA_TYPE_AUDIO)
+		{
+			bHasAudioTransceiver = true;
+			Sender->SetTrack(AudioTrack);
+		}
+	}
+
+	if(!bHasAudioTransceiver)
+	{
 		// Add the track
-		webrtc::RTCErrorOr<rtc::scoped_refptr<webrtc::RtpSenderInterface>> Result = PeerConnection->AddTrack(AudioTrack, { TCHAR_TO_UTF8(*AudioStreamId) });
+		webrtc::RTCErrorOr<rtc::scoped_refptr<webrtc::RtpSenderInterface>> Result = PeerConnection->AddTrack(AudioTrack, { TCHAR_TO_UTF8(*InAudioStreamId) });
 
 		if (!Result.ok())
 		{
-			UE_LOG(PixelStreamer, Error, TEXT("Failed to add audio track to PeerConnection. Msg=%s"), TCHAR_TO_UTF8(Result.error().message()));
+			UE_LOG(LogPixelStreaming, Error, TEXT("Failed to add audio track to PeerConnection. Msg=%s"), TCHAR_TO_UTF8(Result.error().message()));
 		}
 	}
+
 }
 
-void FStreamer::SendPlayerMessage(PixelStreamingProtocol::EToPlayerMsg Type, const FString& Descriptor)
+void UE::PixelStreaming::FStreamer::ModifyTransceivers(std::vector<rtc::scoped_refptr<webrtc::RtpTransceiverInterface>> Transceivers, int Flags)
 {
-	this->PlayerSessions.SendMessageAll(Type, Descriptor);
+	bool bTransmitUEAudio = !UE::PixelStreaming::Settings::CVarPixelStreamingWebRTCDisableTransmitAudio.GetValueOnAnyThread();
+	bool bIsSFU = (Flags & UE::PixelStreaming::Protocol::EPlayerFlags::PSPFlag_IsSFU) != UE::PixelStreaming::Protocol::EPlayerFlags::PSPFlag_None;
+	bool bReceiveBrowserAudio = !bIsSFU && !UE::PixelStreaming::Settings::CVarPixelStreamingWebRTCDisableReceiveAudio.GetValueOnAnyThread();
+
+	for (auto& Transceiver : Transceivers)
+	{
+		rtc::scoped_refptr<webrtc::RtpSenderInterface> Sender = Transceiver->sender();
+
+		if (Transceiver->media_type() == cricket::MediaType::MEDIA_TYPE_AUDIO)
+		{
+
+			// Determine the direction of the transceiver
+			webrtc::RtpTransceiverDirection AudioTransceiverDirection;
+			if (bTransmitUEAudio && bReceiveBrowserAudio)
+			{
+				AudioTransceiverDirection = webrtc::RtpTransceiverDirection::kSendRecv;
+			}
+			else if (bTransmitUEAudio)
+			{
+				AudioTransceiverDirection = webrtc::RtpTransceiverDirection::kSendOnly;
+			}
+			else if (bReceiveBrowserAudio)
+			{
+				AudioTransceiverDirection = webrtc::RtpTransceiverDirection::kRecvOnly;
+			}
+			else
+			{
+				AudioTransceiverDirection = webrtc::RtpTransceiverDirection::kInactive;
+			}
+
+			Transceiver->SetDirection(AudioTransceiverDirection);
+			Sender->SetStreams( { TCHAR_TO_UTF8(*GetAudioStreamID()) } );
+		}
+		else if (Transceiver->media_type() == cricket::MediaType::MEDIA_TYPE_VIDEO)
+		{
+			Transceiver->SetDirection(webrtc::RtpTransceiverDirection::kSendOnly);
+			Sender->SetStreams( { TCHAR_TO_UTF8(*GetVideoStreamID()) } );
+
+			webrtc::RtpParameters RtpParams = Sender->GetParameters();
+			RtpParams.encodings = CreateRTPEncodingParams(Flags);
+			Sender->SetParameters(RtpParams);
+
+		}
+	}
 }
 
-void FStreamer::SendFreezeFrame(const TArray64<uint8>& JpegBytes)
+void UE::PixelStreaming::FStreamer::SendPlayerMessage(UE::PixelStreaming::Protocol::EToPlayerMsg Type, const FString& Descriptor)
 {
-	if (!this->bStreamingStarted)
+	PlayerSessions.SendMessageAll(Type, Descriptor);
+}
+
+void UE::PixelStreaming::FStreamer::SendFreezeFrame(const TArray64<uint8>& JpegBytes)
+{
+	if (!bStreamingStarted)
 	{
 		return;
 	}
 
-	this->PlayerSessions.SendFreezeFrame(JpegBytes);
+	PlayerSessions.SendFreezeFrame(JpegBytes);
 
 	CachedJpegBytes = JpegBytes;
 }
 
-void FStreamer::SendCachedFreezeFrameTo(FPlayerId PlayerId) const
+void UE::PixelStreaming::FStreamer::SendCachedFreezeFrameTo(FPixelStreamingPlayerId PlayerId) const
 {
-	if (this->CachedJpegBytes.Num() > 0)
+	if (CachedJpegBytes.Num() > 0)
 	{
-		this->PlayerSessions.SendFreezeFrameTo(PlayerId, this->CachedJpegBytes);
+		PlayerSessions.SendFreezeFrameTo(PlayerId, CachedJpegBytes);
 	}
 }
 
-void FStreamer::SendFreezeFrameTo(FPlayerId PlayerId, const TArray64<uint8>& JpegBytes) const
+void UE::PixelStreaming::FStreamer::SendFreezeFrameTo(FPixelStreamingPlayerId PlayerId, const TArray64<uint8>& JpegBytes) const
 {
-	this->PlayerSessions.SendFreezeFrameTo(PlayerId, JpegBytes);
+	PlayerSessions.SendFreezeFrameTo(PlayerId, JpegBytes);
 }
 
-void FStreamer::SendUnfreezeFrame()
+void UE::PixelStreaming::FStreamer::SendUnfreezeFrame()
 {
 	// Force a keyframe so when stream unfreezes if player has never received a h.264 frame before they can still connect.
-	this->ForceKeyFrame();
+	ForceKeyFrame();
 
-	this->PlayerSessions.SendUnfreezeFrame();
+	PlayerSessions.SendUnfreezeFrame();
 
 	CachedJpegBytes.Empty();
 }
 
-void FStreamer::SendFileData(TArray<uint8>& ByteData, FString& MimeType, FString& FileExtension)
-{
-	this->PlayerSessions.SendFileData(ByteData, MimeType, FileExtension);
-}
-
-void FStreamer::AddAnyStatChangedCallback(TWeakPtr<IPixelStreamingStatsConsumer> Callback)
-{
-	this->Stats.AddOnAnyStatChangedCallback(Callback);
-}
-
-void FStreamer::RemoveAnyStatChangedCallback(TWeakPtr<IPixelStreamingStatsConsumer> Callback)
+void UE::PixelStreaming::FStreamer::SendFileData(TArray<uint8>& ByteData, FString& MimeType, FString& FileExtension)
 {
-	this->Stats.RemoveOnAnyStatChangedCallback(Callback);
+	PlayerSessions.SendFileData(ByteData, MimeType, FileExtension);
 }
diff --git a/Engine/Plugins/Media/PixelStreaming/Source/PixelStreaming/Private/Streamer.h b/Engine/Plugins/Media/PixelStreaming/Source/PixelStreaming/Private/Streamer.h
--- a/Engine/Plugins/Media/PixelStreaming/Source/PixelStreaming/Private/Streamer.h
+++ b/Engine/Plugins/Media/PixelStreaming/Source/PixelStreaming/Private/Streamer.h
@@ -9,112 +9,121 @@
 #include "HAL/Thread.h"
 #include "IPixelStreamingSessions.h"
 #include "ThreadSafePlayerSessions.h"
-#include "PixelStreamingFramePump.h"
-#include "PixelStreamingFrameSource.h"
-#include "PixelStreamingFramePump.h"
-
-class FVideoCapturer;
-class FPlayerSession;
-class FPixelStreamingVideoEncoderFactory;
-class FStreamer;
-class FSetSessionDescriptionObserver;
-
-class FStreamer : public FSignallingServerConnectionObserver, public IPixelStreamingSessions
-{
-public:
-	static bool CheckPlatformCompatibility();
-
-	explicit FStreamer(const FString& SignallingServerUrl, const FString& StreamerId);
-	virtual ~FStreamer() override;
-
-	void SendPlayerMessage(PixelStreamingProtocol::EToPlayerMsg Type, const FString& Descriptor);
-	void SendFreezeFrame(const TArray64<uint8>& JpegBytes);
-	void SendUnfreezeFrame();
-	void SendFileData(TArray<uint8>& ByteData, FString& MimeType, FString& FileExtension);
-	void ForceKeyFrame();
-
-	bool IsStreaming() const { return bStreamingStarted; }
-	void SetStreamingStarted(bool started) { bStreamingStarted = started; }
-	FSignallingServerConnection* GetSignallingServerConnection() const { return SignallingServerConnection.Get(); }
-	FPixelStreamingVideoEncoderFactory* GetVideoEncoderFactory() const { return VideoEncoderFactory; }
-
-	// data coming from the engine
-	void OnFrameBufferReady(const FTexture2DRHIRef& FrameBuffer);
-
-	// Begin IPixelStreamingSessions interface.
-	virtual int GetNumPlayers() const override;
-	virtual IPixelStreamingAudioSink* GetAudioSink(FPlayerId PlayerId) const override;
-	virtual IPixelStreamingAudioSink* GetUnlistenedAudioSink() const override;
-	virtual bool IsQualityController(FPlayerId PlayerId) const override;
-	virtual void SetQualityController(FPlayerId PlayerId) override;
-	virtual bool SendMessage(FPlayerId PlayerId, PixelStreamingProtocol::EToPlayerMsg Type, const FString& Descriptor) const override;
-	virtual void SendLatestQP(FPlayerId PlayerId, int LatestQP) const override;
-	virtual void SendFreezeFrameTo(FPlayerId PlayerId, const TArray64<uint8>& JpegBytes) const override;
-	virtual void PollWebRTCStats() const override;
-	// End IPixelStreamingSessions interface.
-
-	void AddAnyStatChangedCallback(TWeakPtr<IPixelStreamingStatsConsumer> Callback);
-	void RemoveAnyStatChangedCallback(TWeakPtr<IPixelStreamingStatsConsumer> Callback);
-	void SendCachedFreezeFrameTo(FPlayerId PlayerId) const;
-
-private:
-	webrtc::PeerConnectionInterface* CreateSession(FPlayerId PlayerId, int Flags);
-	webrtc::AudioProcessing* SetupAudioProcessingModule();
-
-	// Procedure for WebRTC inter-thread communication
-	void StartWebRtcSignallingThread();
-	void ConnectToSignallingServer();
-
-	// ISignallingServerConnectionObserver impl
-	virtual void OnConfig(const webrtc::PeerConnectionInterface::RTCConfiguration& Config) override;
-	virtual void OnSessionDescription(FPlayerId PlayerId, webrtc::SdpType Type, const FString& Sdp) override;
-	virtual void OnRemoteIceCandidate(FPlayerId PlayerId, const std::string& SdpMid, int SdpMLineIndex, const std::string& Sdp) override;
-	virtual void OnPlayerConnected(FPlayerId PlayerId, int Flags) override;
-	virtual void OnPlayerDisconnected(FPlayerId PlayerId) override;
-	virtual void OnSignallingServerDisconnected() override;
-
-	// own methods
-	void OnOffer(FPlayerId PlayerId, const FString& Sdp);
-	void ModifyAudioTransceiverDirection(webrtc::PeerConnectionInterface* PeerConnection);
-	void DeletePlayerSession(FPlayerId PlayerId);
-	void DeleteAllPlayerSessions();
-	void AddStreams(FPlayerId PlayerId, webrtc::PeerConnectionInterface* PeerConnection, int Flags);
-	void SetupVideoTrack(FPlayerId PlayerId, webrtc::PeerConnectionInterface* PeerConnection, FString const VideoStreamId, FString const VideoTrackLabel, int Flags);
-	void SetupAudioTrack(webrtc::PeerConnectionInterface* PeerConnection, const FString AudioStreamId, const FString AudioTrackLabel);
-	void SendAnswer(FPlayerId PlayerId, webrtc::PeerConnectionInterface* PeerConnection, TUniquePtr<webrtc::SessionDescriptionInterface> Sdp);
-	void OnDataChannelOpen(FPlayerId PlayerId, webrtc::DataChannelInterface* DataChannel);
-	void OnQualityControllerChanged(FPlayerId PlayerId);
-	void PostPlayerDeleted(FPlayerId PlayerId);
-	void SetLocalDescription(webrtc::PeerConnectionInterface* PeerConnection, FSetSessionDescriptionObserver* Observer, webrtc::SessionDescriptionInterface* SDP);
-
-private:
-	FString SignallingServerUrl;
-	FString StreamerId;
-
-	TUniquePtr<rtc::Thread> WebRtcSignallingThread;
-
-	TUniquePtr<FSignallingServerConnection> SignallingServerConnection;
-	double LastSignallingServerConnectionAttemptTimestamp = 0;
-
-	rtc::scoped_refptr<webrtc::PeerConnectionFactoryInterface> PeerConnectionFactory;
-	webrtc::PeerConnectionInterface::RTCConfiguration PeerConnectionConfig;
-
-	FPixelStreamingVideoEncoderFactory* VideoEncoderFactory;
-
-	FPixelStreamingFrameSource FrameSource;
-	FPixelStreamingFramePump FramePump;
-
-	rtc::scoped_refptr<webrtc::AudioSourceInterface> AudioSource;
-	cricket::AudioOptions AudioSourceOptions;
-
-	// When we send a freeze frame we retain the data so we send freeze frame to new peers if they join during a freeze frame.
-	TArray64<uint8> CachedJpegBytes;
-
-	FThreadSafeBool bStreamingStarted = false;
-
-	FThreadSafePlayerSessions PlayerSessions;
-	FPixelStreamingStats Stats;
-
-	TUniquePtr<webrtc::SessionDescriptionInterface> SFULocalDescription;
-	TUniquePtr<webrtc::SessionDescriptionInterface> SFURemoteDescription;
-};
+
+namespace UE {
+	namespace PixelStreaming {
+		class FPlayerSession;
+		class FVideoEncoderFactory;
+		class FSimulcastEncoderFactory;
+		class FStreamer;
+		class FSetSessionDescriptionObserver;
+
+		class FStreamer : public FSignallingServerConnectionObserver, public IPixelStreamingSessions
+		{
+		public:
+			static bool IsPlatformCompatible();
+
+			explicit FStreamer(const FString& SignallingServerUrl, const FString& StreamerId);
+			virtual ~FStreamer() override;
+
+			void SendPlayerMessage(Protocol::EToPlayerMsg Type, const FString& Descriptor);
+			void SendFreezeFrame(const TArray64<uint8>& JpegBytes);
+			void SendUnfreezeFrame();
+			void SendFileData(TArray<uint8>& ByteData, FString& MimeType, FString& FileExtension);
+			void ForceKeyFrame();
+
+			bool IsStreaming() const { return bStreamingStarted; }
+			void SetStreamingStarted(bool bStarted) { bStreamingStarted = bStarted; }
+			FSignallingServerConnection* GetSignallingServerConnection() const { return SignallingServerConnection.Get(); }
+			FVideoEncoderFactory* GetP2PVideoEncoderFactory() const { return P2PVideoEncoderFactory; }
+
+			// data coming from the engine
+			void OnFrameBufferReady(const FTexture2DRHIRef& FrameBuffer);
+
+			// Begin IPixelStreamingSessions interface.
+			virtual int GetNumPlayers() const override;
+			virtual IPixelStreamingAudioSink* GetAudioSink(FPixelStreamingPlayerId PlayerId) const override;
+			virtual IPixelStreamingAudioSink* GetUnlistenedAudioSink() const override;
+			virtual bool IsQualityController(FPixelStreamingPlayerId PlayerId) const override;
+			virtual void SetQualityController(FPixelStreamingPlayerId PlayerId) override;
+			virtual bool SendMessage(FPixelStreamingPlayerId PlayerId, Protocol::EToPlayerMsg Type, const FString& Descriptor) const override;
+			virtual void SendLatestQP(FPixelStreamingPlayerId PlayerId, int LatestQP) const override;
+			virtual void SendFreezeFrameTo(FPixelStreamingPlayerId PlayerId, const TArray64<uint8>& JpegBytes) const override;
+			virtual void PollWebRTCStats() const override;
+			// End IPixelStreamingSessions interface.
+
+			void SendCachedFreezeFrameTo(FPixelStreamingPlayerId PlayerId) const;
+
+		private:
+			webrtc::PeerConnectionInterface* CreateSession(FPixelStreamingPlayerId PlayerId, int Flags);
+			webrtc::AudioProcessing* SetupAudioProcessingModule();
+
+			// Procedure for WebRTC inter-thread communication
+			void StartWebRtcSignallingThread();
+			void ConnectToSignallingServer();
+
+			// ISignallingServerConnectionObserver impl
+			virtual void OnConfig(const webrtc::PeerConnectionInterface::RTCConfiguration& Config) override;
+			virtual void OnSessionDescription(FPixelStreamingPlayerId PlayerId, webrtc::SdpType Type, const FString& Sdp) override;
+			virtual void OnRemoteIceCandidate(FPixelStreamingPlayerId PlayerId, const std::string& SdpMid, int SdpMLineIndex, const std::string& Sdp) override;
+			virtual void OnPlayerConnected(FPixelStreamingPlayerId PlayerId, int Flags) override;
+			virtual void OnPlayerDisconnected(FPixelStreamingPlayerId PlayerId) override;
+			virtual void OnSignallingServerDisconnected() override;
+
+			// own methods
+			void OnOffer(FPixelStreamingPlayerId PlayerId, const FString& Sdp);
+			void ModifyTransceivers(std::vector<rtc::scoped_refptr<webrtc::RtpTransceiverInterface>> Transceivers, int Flags);
+			void MungeRemoteSDP(cricket::SessionDescription* SessionDescription);
+			void DeletePlayerSession(FPixelStreamingPlayerId PlayerId);
+			void DeleteAllPlayerSessions();
+			void AddStreams(FPixelStreamingPlayerId PlayerId, webrtc::PeerConnectionInterface* PeerConnection, int Flags);
+			void SetupVideoTrack(FPixelStreamingPlayerId PlayerId, webrtc::PeerConnectionInterface* PeerConnection, FString const VideoStreamId, FString const VideoTrackLabel, int Flags);
+			void SetupAudioTrack(webrtc::PeerConnectionInterface* PeerConnection, const FString AudioStreamId, const FString AudioTrackLabel, int Flags);
+			std::vector<webrtc::RtpEncodingParameters> CreateRTPEncodingParams(int Flags);
+			void SendAnswer(FPixelStreamingPlayerId PlayerId, webrtc::PeerConnectionInterface* PeerConnection, TUniquePtr<webrtc::SessionDescriptionInterface> Sdp);
+			void OnDataChannelOpen(FPixelStreamingPlayerId PlayerId, webrtc::DataChannelInterface* DataChannel);
+			void OnQualityControllerChanged(FPixelStreamingPlayerId PlayerId);
+			void PostPlayerDeleted(FPixelStreamingPlayerId PlayerId, bool WasQualityController);
+			void SetLocalDescription(webrtc::PeerConnectionInterface* PeerConnection, FSetSessionDescriptionObserver* Observer, webrtc::SessionDescriptionInterface* SDP);
+			FString GetAudioStreamID() const;
+			FString GetVideoStreamID() const;
+
+		private:
+			FString SignallingServerUrl;
+			FString StreamerId;
+			const FString AudioTrackLabel = TEXT("pixelstreaming_audio_track_label");
+			const FString VideoTrackLabel = TEXT("pixelstreaming_video_track_label");
+
+			TUniquePtr<rtc::Thread> WebRtcSignallingThread;
+
+			TUniquePtr<FSignallingServerConnection> SignallingServerConnection;
+			double LastSignallingServerConnectionAttemptTimestamp = 0;
+
+			webrtc::PeerConnectionInterface::RTCConfiguration PeerConnectionConfig;
+
+			/* P2P Peer Connections use these exclusively */
+			rtc::scoped_refptr<webrtc::PeerConnectionFactoryInterface> P2PPeerConnectionFactory;
+			FVideoEncoderFactory* P2PVideoEncoderFactory;
+			/* P2P Peer Connections use these exclusively */
+
+			/* SFU Peer Connections use these exclusively */
+			rtc::scoped_refptr<webrtc::PeerConnectionFactoryInterface> SFUPeerConnectionFactory;
+			FSimulcastEncoderFactory* SFUVideoEncoderFactory;
+			/* SFU Peer Connections use these exclusively */
+
+			rtc::scoped_refptr<webrtc::AudioSourceInterface> AudioSource;
+			cricket::AudioOptions AudioSourceOptions;
+
+			// When we send a freeze frame we retain the data so we send freeze frame to new peers if they join during a freeze frame.
+			TArray64<uint8> CachedJpegBytes;
+
+			FThreadSafeBool bStreamingStarted = false;
+
+			FThreadSafePlayerSessions PlayerSessions;
+			FStats Stats;
+
+			TUniquePtr<webrtc::SessionDescriptionInterface> SFULocalDescription;
+			TUniquePtr<webrtc::SessionDescriptionInterface> SFURemoteDescription;
+		};
+	}
+}
diff --git a/Engine/Plugins/Media/PixelStreaming/Source/PixelStreaming/Private/TextureSource.cpp b/Engine/Plugins/Media/PixelStreaming/Source/PixelStreaming/Private/TextureSource.cpp
new file mode 100644
--- /dev/null
+++ b/Engine/Plugins/Media/PixelStreaming/Source/PixelStreaming/Private/TextureSource.cpp
@@ -0,0 +1,135 @@
+// Copyright Epic Games, Inc. All Rights Reserved.
+
+#include "TextureSource.h"
+#include "Stats.h"
+#include "Utils.h"
+#include "Async/Async.h"
+#include "GPUFencePoller.h"
+#include "Framework/Application/SlateApplication.h"
+
+
+
+/*
+* -------- UE::PixelStreaming::FTextureSourceBackBuffer ----------------
+*/
+
+UE::PixelStreaming::FTextureSourceBackBuffer::FTextureSourceBackBuffer(float InFrameScale)
+	: FrameScale(InFrameScale)
+{
+	// The backbuffer delegate can only be accessed using GameThread
+	AsyncTask(ENamedThreads::GameThread, [this]() {
+		OnBackbuffer = &FSlateApplication::Get().GetRenderer()->OnBackBufferReadyToPresent();
+		BackbufferDelegateHandle = OnBackbuffer->AddSP(this, &UE::PixelStreaming::FTextureSourceBackBuffer::OnBackBufferReady_RenderThread);
+	});
+}
+
+UE::PixelStreaming::FTextureSourceBackBuffer::FTextureSourceBackBuffer()
+	: UE::PixelStreaming::FTextureSourceBackBuffer(1.0) {}
+
+UE::PixelStreaming::FTextureSourceBackBuffer::~FTextureSourceBackBuffer()
+{
+	if(OnBackbuffer)
+	{
+		OnBackbuffer->Remove(BackbufferDelegateHandle);
+	}
+	
+	SetEnabled(false);
+}
+
+void UE::PixelStreaming::FTextureSourceBackBuffer::SetEnabled(bool bInEnabled) 
+{ 
+	*bEnabled = bInEnabled; 
+
+	// This source has been disabled, so set `bInitialized` to false so `OnBackBufferReady_RenderThread` 
+	// will make new textures next time it is called.
+	if(bInitialized && bInEnabled == false)
+	{
+		bInitialized = false;
+	}
+	
+}
+
+void UE::PixelStreaming::FTextureSourceBackBuffer::OnBackBufferReady_RenderThread(SWindow& SlateWindow, const FTexture2DRHIRef& FrameBuffer)
+{
+	if (!bInitialized)
+	{
+		Initialize(FrameBuffer->GetSizeXY().X * FrameScale, FrameBuffer->GetSizeXY().Y * FrameScale);
+	}
+
+	if (!IsEnabled())
+	{
+		return;
+	}
+
+	auto& WriteBuffer = bWriteParity ? WriteBuffers[0] : WriteBuffers[1];
+	bWriteParity	  = !bWriteParity;
+
+	// for safety we just make sure that the buffer is not currently waiting for a copy
+	if (WriteBuffer.bAvailable)
+	{
+		WriteBuffer.bAvailable = false;
+
+		FRHICommandListImmediate& RHICmdList = FRHICommandListExecutor::GetImmediateCommandList();
+		WriteBuffer.Fence->Clear();
+
+		RHICmdList.EnqueueLambda([&WriteBuffer](FRHICommandListImmediate& RHICmdList) {
+			WriteBuffer.PreWaitingOnCopy = FPlatformTime::Cycles64();
+		});
+
+		UE::PixelStreaming::CopyTexture(FrameBuffer, WriteBuffer.Texture, WriteBuffer.Fence);
+
+		UE::PixelStreaming::FGPUFencePoller::Get().AddJob(WriteBuffer.Fence, bEnabled, [this, &WriteBuffer]() {
+			// This lambda is called only once the GPUFence is done
+			{
+				FScopeLock Lock(&CriticalSection);
+				TempBuffer.Swap(WriteBuffer.Texture);
+				WriteBuffer.Fence->Clear();
+				WriteBuffer.bAvailable = true;
+
+				bIsTempDirty = true;
+			}
+
+			// For debugging timing information about the copy operation
+			// Turning it on all the time is a bit too much log spam if logging stats
+			uint64 PostWaitingOnCopy	= FPlatformTime::Cycles64();
+			UE::PixelStreaming::FStats* Stats = UE::PixelStreaming::FStats::Get();
+			if (Stats)
+			{
+				double CaptureLatencyMs = FPlatformTime::ToMilliseconds64(PostWaitingOnCopy - WriteBuffer.PreWaitingOnCopy);
+				Stats->StoreApplicationStat(UE::PixelStreaming::FStatData(FName(*FString::Printf(TEXT("Layer (x%.2f) Capture time (ms)"), FrameScale)), CaptureLatencyMs, 2, true));
+			}
+		});
+
+	}
+}
+
+FTexture2DRHIRef UE::PixelStreaming::FTextureSourceBackBuffer::GetTexture()
+{
+	if (bIsTempDirty)
+	{
+		FScopeLock Lock(&CriticalSection);
+		ReadBuffer.Swap(TempBuffer);
+		bIsTempDirty = false;
+	}
+	return ReadBuffer;
+}
+
+void UE::PixelStreaming::FTextureSourceBackBuffer::Initialize(int Width, int Height)
+{
+	SourceWidth	 = Width;
+	SourceHeight = Height;
+
+	for (auto& Buffer : WriteBuffers)
+	{
+		Buffer.Texture	  = UE::PixelStreaming::CreateTexture(SourceWidth, SourceHeight);
+		Buffer.Fence	  = GDynamicRHI->RHICreateGPUFence(TEXT("VideoCapturerCopyFence"));
+		Buffer.bAvailable = true;
+	}
+	bWriteParity = true;
+
+	TempBuffer	 = UE::PixelStreaming::CreateTexture(SourceWidth, SourceHeight);
+	ReadBuffer	 = UE::PixelStreaming::CreateTexture(SourceWidth, SourceHeight);
+	bIsTempDirty = false;
+
+	bInitialized = true;
+}
diff --git a/Engine/Plugins/Media/PixelStreaming/Source/PixelStreaming/Private/TextureSource.h b/Engine/Plugins/Media/PixelStreaming/Source/PixelStreaming/Private/TextureSource.h
new file mode 100644
--- /dev/null
+++ b/Engine/Plugins/Media/PixelStreaming/Source/PixelStreaming/Private/TextureSource.h
@@ -0,0 +1,89 @@
+// Copyright Epic Games, Inc. All Rights Reserved.
+#pragma once
+
+#include "RHI.h"
+#include "HAL/ThreadSafeBool.h"
+#include "Rendering/SlateRenderer.h"
+#include "Templates/SharedPointer.h"
+
+namespace UE {
+	namespace PixelStreaming {
+
+		/*
+		* Interface for all texture sources in Pixel Streaming.
+		* These texture sources are used to populate video sources for video tracks.
+		*/
+		class ITextureSource
+		{
+			public:
+				ITextureSource() = default;
+				virtual ~ITextureSource() = default;
+				virtual bool IsAvailable() const = 0;
+				virtual bool IsEnabled() const = 0;
+				virtual void SetEnabled(bool bInEnabled) = 0;
+				virtual int GetSourceHeight() const = 0;
+				virtual int GetSourceWidth() const = 0;
+				virtual FTexture2DRHIRef GetTexture() = 0;
+		};
+
+
+		/*
+		* Source of textures coming from the UE backbuffer.
+		* This class has the additional functionality of scaling textures from the backbuffer.
+		*/
+		class FTextureSourceBackBuffer : public ITextureSource, public TSharedFromThis<FTextureSourceBackBuffer, ESPMode::ThreadSafe>
+		{
+		public:
+			FTextureSourceBackBuffer();
+			FTextureSourceBackBuffer(float InScale);
+			~FTextureSourceBackBuffer();
+			
+			/* Begin ITextureSource interface */
+			void SetEnabled(bool bInEnabled) override;
+			bool IsEnabled() const override { return *bEnabled; }
+			bool IsAvailable() const override { return bInitialized; }
+			int GetSourceWidth() const override { return SourceWidth; }
+			int GetSourceHeight() const override { return SourceHeight; }
+			FTexture2DRHIRef GetTexture() override;
+			/* End ITextureSource interface */
+
+			float GetFrameScaling() const { return FrameScale; }
+
+		private:
+			void OnBackBufferReady_RenderThread(SWindow& SlateWindow, const FTexture2DRHIRef& FrameBuffer);
+			void Initialize(int Width, int Height);
+
+		private:
+			const float FrameScale;
+			int SourceWidth = 0;
+			int SourceHeight = 0;
+			FThreadSafeBool bInitialized = false;
+			TSharedRef<bool, ESPMode::ThreadSafe> bEnabled = MakeShared<bool, ESPMode::ThreadSafe>(true);
+			FSlateRenderer::FOnBackBufferReadyToPresent* OnBackbuffer = nullptr;
+			FDelegateHandle BackbufferDelegateHandle;
+			FCriticalSection DestructorCriticalSection;
+
+			struct FCaptureFrame
+			{
+				FTexture2DRHIRef Texture;
+				FGPUFenceRHIRef Fence;
+				bool bAvailable = true;
+				uint64 PreWaitingOnCopy;
+			};
+
+			/*
+			* Triple buffer setup with queued write buffers (since we have to wait for RHI copy).
+			* 1 Read buffer (read the captured texture)
+			* 1 Temp buffer (for swapping what is read and written)
+			* 2 Write buffers (2 write buffers because UE can render two frames before presenting sometimes)
+			*/
+			FCriticalSection CriticalSection;
+			bool bWriteParity = true;
+			FCaptureFrame WriteBuffers[2];
+			FTexture2DRHIRef TempBuffer;
+			FTexture2DRHIRef ReadBuffer;
+			FThreadSafeBool bIsTempDirty;
+		};
+
+	}
+}
