Merging Pixel Streaming refactor using //UE5/Dev-TensorWorks-to-//UE5/Release-5.0

diff --git a/Engine/Plugins/Media/PixelStreaming/Source/PixelStreaming/Private/ThreadSafePlayerSessions.cpp b/Engine/Plugins/Media/PixelStreaming/Source/PixelStreaming/Private/ThreadSafePlayerSessions.cpp
--- a/Engine/Plugins/Media/PixelStreaming/Source/PixelStreaming/Private/ThreadSafePlayerSessions.cpp
+++ b/Engine/Plugins/Media/PixelStreaming/Source/PixelStreaming/Private/ThreadSafePlayerSessions.cpp
@@ -2,184 +2,189 @@
 
 #include "ThreadSafePlayerSessions.h"
 #include "rtc_base/location.h"
-#include "PixelStreamingSettings.h"
-#include "PixelStreamerDelegates.h"
+#include "Settings.h"
+#include "PixelStreamingDelegates.h"
 #include "Streamer.h"
 
-#define SUBMIT_TASK_WITH_RETURN(ReturnType, FuncWithReturnType)                                                                                \
-	if (this->IsInSignallingThread())                                                                                                          \
-	{                                                                                                                                          \
-		return this->FuncWithReturnType();                                                                                                     \
-	}                                                                                                                                          \
-	else                                                                                                                                       \
-	{                                                                                                                                          \
-		return this->WebRtcSignallingThread->Invoke<ReturnType>(RTC_FROM_HERE, [this]() -> ReturnType { return this->FuncWithReturnType(); }); \
-	}
-
-#define SUBMIT_TASK_WITH_PARAMS_AND_RETURN(ReturnType, FuncWithReturnType, ...)                                                                                        \
-	if (this->IsInSignallingThread())                                                                                                                                  \
-	{                                                                                                                                                                  \
-		return this->FuncWithReturnType(__VA_ARGS__);                                                                                                                  \
-	}                                                                                                                                                                  \
-	else                                                                                                                                                               \
-	{                                                                                                                                                                  \
-		return this->WebRtcSignallingThread->Invoke<ReturnType>(RTC_FROM_HERE, [this, __VA_ARGS__]() -> ReturnType { return this->FuncWithReturnType(__VA_ARGS__); }); \
-	}
-
-#define SUBMIT_TASK_NO_PARAMS(Func)                                                        \
-	if (this->IsInSignallingThread())                                                      \
-	{                                                                                      \
-		this->Func();                                                                      \
-	}                                                                                      \
-	else                                                                                   \
-	{                                                                                      \
-		this->WebRtcSignallingThread->PostTask(RTC_FROM_HERE, [this]() { this->Func(); }); \
-	}
-
-#define SUBMIT_TASK_WITH_PARAMS(Func, ...)                                                                         \
-	if (this->IsInSignallingThread())                                                                              \
-	{                                                                                                              \
-		this->Func(__VA_ARGS__);                                                                                   \
-	}                                                                                                              \
-	else                                                                                                           \
-	{                                                                                                              \
-		this->WebRtcSignallingThread->PostTask(RTC_FROM_HERE, [this, __VA_ARGS__]() { this->Func(__VA_ARGS__); }); \
+#define SUBMIT_TASK_WITH_RETURN(ReturnType, FuncWithReturnType)                                                                    \
+	if (IsInSignallingThread())                                                                                                    \
+	{                                                                                                                              \
+		return FuncWithReturnType();                                                                                               \
+	}                                                                                                                              \
+	else                                                                                                                           \
+	{                                                                                                                              \
+		return WebRtcSignallingThread->Invoke<ReturnType>(RTC_FROM_HERE, [this]() -> ReturnType { return FuncWithReturnType(); }); \
+	}
+
+#define SUBMIT_TASK_WITH_PARAMS_AND_RETURN(ReturnType, FuncWithReturnType, ...)                                                                            \
+	if (IsInSignallingThread())                                                                                                                            \
+	{                                                                                                                                                      \
+		return FuncWithReturnType(__VA_ARGS__);                                                                                                            \
+	}                                                                                                                                                      \
+	else                                                                                                                                                   \
+	{                                                                                                                                                      \
+		return WebRtcSignallingThread->Invoke<ReturnType>(RTC_FROM_HERE, [this, __VA_ARGS__]() -> ReturnType { return FuncWithReturnType(__VA_ARGS__); }); \
+	}
+
+#define SUBMIT_TASK_NO_PARAMS(Func)                                            \
+	if (IsInSignallingThread())                                                \
+	{                                                                          \
+		Func();                                                                \
+	}                                                                          \
+	else                                                                       \
+	{                                                                          \
+		WebRtcSignallingThread->PostTask(RTC_FROM_HERE, [this]() { Func(); }); \
+	}
+
+#define SUBMIT_TASK_WITH_PARAMS(Func, ...)                                                             \
+	if (IsInSignallingThread())                                                                        \
+	{                                                                                                  \
+		Func(__VA_ARGS__);                                                                             \
+	}                                                                                                  \
+	else                                                                                               \
+	{                                                                                                  \
+		WebRtcSignallingThread->PostTask(RTC_FROM_HERE, [this, __VA_ARGS__]() { Func(__VA_ARGS__); }); \
 	}
 
 // NOTE: All public methods should be one-liner functions that call one of the above macros to ensure threadsafety of this class.
 
-FThreadSafePlayerSessions::FThreadSafePlayerSessions(rtc::Thread* InWebRtcSignallingThread)
+UE::PixelStreaming::FThreadSafePlayerSessions::FThreadSafePlayerSessions(rtc::Thread* InWebRtcSignallingThread)
 	: WebRtcSignallingThread(InWebRtcSignallingThread)
 {
 }
 
-bool FThreadSafePlayerSessions::IsInSignallingThread() const
+bool UE::PixelStreaming::FThreadSafePlayerSessions::IsInSignallingThread() const
 {
-	if (this->WebRtcSignallingThread == nullptr)
+	if (WebRtcSignallingThread == nullptr)
 	{
 		return false;
 	}
 
-	return this->WebRtcSignallingThread->IsCurrent();
+	return WebRtcSignallingThread->IsCurrent();
 }
 
-int FThreadSafePlayerSessions::GetNumPlayers() const
+int UE::PixelStreaming::FThreadSafePlayerSessions::GetNumPlayers() const
 {
 	SUBMIT_TASK_WITH_RETURN(int, GetNumPlayers_SignallingThread);
 }
 
-IPixelStreamingAudioSink* FThreadSafePlayerSessions::GetAudioSink(FPlayerId PlayerId) const {
+IPixelStreamingAudioSink* UE::PixelStreaming::FThreadSafePlayerSessions::GetAudioSink(FPixelStreamingPlayerId PlayerId) const {
 	SUBMIT_TASK_WITH_PARAMS_AND_RETURN(IPixelStreamingAudioSink*, GetAudioSink_SignallingThread, PlayerId)
 }
 
-IPixelStreamingAudioSink* FThreadSafePlayerSessions::GetUnlistenedAudioSink() const
+IPixelStreamingAudioSink* UE::PixelStreaming::FThreadSafePlayerSessions::GetUnlistenedAudioSink() const
 {
 	SUBMIT_TASK_WITH_RETURN(IPixelStreamingAudioSink*, GetUnlistenedAudioSink_SignallingThread)
 }
 
-bool FThreadSafePlayerSessions::IsQualityController(FPlayerId PlayerId) const
+bool UE::PixelStreaming::FThreadSafePlayerSessions::IsQualityController(FPixelStreamingPlayerId PlayerId) const
 {
 	// Due to how many theads this particular method is called on we choose not to schedule reading it as a new task and instead just made the
 	// the variable atomic. Additionally some of the calling threads were deadlocking waiting for each other to finish other methods while calling this method.
 	// Not scheduling this as a task is the exception to the general rule of this class though.
 	FScopeLock Lock(&QualityControllerCS);
-	return this->QualityControllingPlayer == PlayerId;
+	return QualityControllingPlayer == PlayerId;
 }
 
-void FThreadSafePlayerSessions::SetQualityController(FPlayerId PlayerId)
+void UE::PixelStreaming::FThreadSafePlayerSessions::SetQualityController(FPixelStreamingPlayerId PlayerId)
 {
 	{
 		FScopeLock Lock(&QualityControllerCS);
-		this->QualityControllingPlayer = PlayerId;
+		QualityControllingPlayer = PlayerId;
 	}
 
 	SUBMIT_TASK_WITH_PARAMS(SetQualityController_SignallingThread, PlayerId)
 }
 
-bool FThreadSafePlayerSessions::SendMessage(FPlayerId PlayerId, PixelStreamingProtocol::EToPlayerMsg Type, const FString& Descriptor) const
+bool UE::PixelStreaming::FThreadSafePlayerSessions::SendMessage(FPixelStreamingPlayerId PlayerId, UE::PixelStreaming::Protocol::EToPlayerMsg Type, const FString& Descriptor) const
 {
 	SUBMIT_TASK_WITH_PARAMS_AND_RETURN(bool, SendMessage_SignallingThread, PlayerId, Type, Descriptor)
 }
 
-void FThreadSafePlayerSessions::SendLatestQP(FPlayerId PlayerId, int LatestQP) const
+void UE::PixelStreaming::FThreadSafePlayerSessions::SendLatestQP(FPixelStreamingPlayerId PlayerId, int LatestQP) const
 {
 	SUBMIT_TASK_WITH_PARAMS(SendLatestQP_SignallingThread, PlayerId, LatestQP)
 }
 
-void FThreadSafePlayerSessions::SendFreezeFrameTo(FPlayerId PlayerId, const TArray64<uint8>& JpegBytes) const
+void UE::PixelStreaming::FThreadSafePlayerSessions::SendFreezeFrameTo(FPixelStreamingPlayerId PlayerId, const TArray64<uint8>& JpegBytes) const
 {
 	SUBMIT_TASK_WITH_PARAMS(SendFreezeFrameTo_SignallingThread, PlayerId, JpegBytes)
 }
 
-void FThreadSafePlayerSessions::SendFreezeFrame(const TArray64<uint8>& JpegBytes)
+void UE::PixelStreaming::FThreadSafePlayerSessions::SendFreezeFrame(const TArray64<uint8>& JpegBytes)
 {
 	SUBMIT_TASK_WITH_PARAMS(SendFreezeFrame_SignallingThread, JpegBytes)
 }
 
-void FThreadSafePlayerSessions::SendUnfreezeFrame()
+void UE::PixelStreaming::FThreadSafePlayerSessions::SendUnfreezeFrame()
 {
 	SUBMIT_TASK_NO_PARAMS(SendUnfreezeFrame_SignallingThread)
 }
 
-void FThreadSafePlayerSessions::SendFileData(const TArray<uint8>& ByteData, const FString& MimeType, const FString& FileExtension){
+void UE::PixelStreaming::FThreadSafePlayerSessions::SendFileData(const TArray<uint8>& ByteData, const FString& MimeType, const FString& FileExtension){
 	SUBMIT_TASK_WITH_PARAMS(SendFileData_SignallingThread, ByteData, MimeType, FileExtension)
 }
 
-webrtc::PeerConnectionInterface* FThreadSafePlayerSessions::CreatePlayerSession(
-    FPlayerId PlayerId,
+webrtc::PeerConnectionInterface* UE::PixelStreaming::FThreadSafePlayerSessions::CreatePlayerSession(
+	FPixelStreamingPlayerId PlayerId,
 	rtc::scoped_refptr<webrtc::PeerConnectionFactoryInterface> PeerConnectionFactory,
 	webrtc::PeerConnectionInterface::RTCConfiguration PeerConnectionConfig,
-	FSignallingServerConnection* SignallingServerConnection,
+	UE::PixelStreaming::FSignallingServerConnection* SignallingServerConnection,
 	int Flags)
 {
 	SUBMIT_TASK_WITH_PARAMS_AND_RETURN(webrtc::PeerConnectionInterface*, CreatePlayerSession_SignallingThread, PlayerId, PeerConnectionFactory, PeerConnectionConfig, SignallingServerConnection, Flags)
 }
 
-void FThreadSafePlayerSessions::SetPlayerSessionDataChannel(FPlayerId PlayerId, rtc::scoped_refptr<webrtc::DataChannelInterface> DataChannel)
+void UE::PixelStreaming::FThreadSafePlayerSessions::SetPlayerSessionDataChannel(FPixelStreamingPlayerId PlayerId, rtc::scoped_refptr<webrtc::DataChannelInterface> DataChannel)
 {
 	SUBMIT_TASK_WITH_PARAMS(SetPlayerSessionDataChannel_SignallingThread, PlayerId, DataChannel)
 }
 
-void FThreadSafePlayerSessions::DeleteAllPlayerSessions()
+void UE::PixelStreaming::FThreadSafePlayerSessions::SetVideoSource(FPixelStreamingPlayerId PlayerId, rtc::scoped_refptr<webrtc::VideoTrackSourceInterface> VideoSource)
+{
+	SUBMIT_TASK_WITH_PARAMS(SetVideoSource_SignallingThread, PlayerId, VideoSource)
+}
+
+void UE::PixelStreaming::FThreadSafePlayerSessions::DeleteAllPlayerSessions()
 {
 	SUBMIT_TASK_NO_PARAMS(DeleteAllPlayerSessions_SignallingThread)
 }
 
-int FThreadSafePlayerSessions::DeletePlayerSession(FPlayerId PlayerId)
+int UE::PixelStreaming::FThreadSafePlayerSessions::DeletePlayerSession(FPixelStreamingPlayerId PlayerId)
 {
 	SUBMIT_TASK_WITH_PARAMS_AND_RETURN(int, DeletePlayerSession_SignallingThread, PlayerId)
 }
 
-void FThreadSafePlayerSessions::DisconnectPlayer(FPlayerId PlayerId, const FString& Reason){
+void UE::PixelStreaming::FThreadSafePlayerSessions::DisconnectPlayer(FPixelStreamingPlayerId PlayerId, const FString& Reason){
 	SUBMIT_TASK_WITH_PARAMS(DisconnectPlayer_SignallingThread, PlayerId, Reason)
 }
 
-FPixelStreamingDataChannelObserver* FThreadSafePlayerSessions::GetDataChannelObserver(FPlayerId PlayerId)
+UE::PixelStreaming::FDataChannelObserver* UE::PixelStreaming::FThreadSafePlayerSessions::GetDataChannelObserver(FPixelStreamingPlayerId PlayerId)
 {
-	SUBMIT_TASK_WITH_PARAMS_AND_RETURN(FPixelStreamingDataChannelObserver*, GetDataChannelObserver_SignallingThread, PlayerId)
+	SUBMIT_TASK_WITH_PARAMS_AND_RETURN(UE::PixelStreaming::FDataChannelObserver*, GetDataChannelObserver_SignallingThread, PlayerId)
 }
 
-void FThreadSafePlayerSessions::SendMessageAll(PixelStreamingProtocol::EToPlayerMsg Type, const FString& Descriptor) const
+void UE::PixelStreaming::FThreadSafePlayerSessions::SendMessageAll(UE::PixelStreaming::Protocol::EToPlayerMsg Type, const FString& Descriptor) const
 {
 	SUBMIT_TASK_WITH_PARAMS(SendMessageAll_SignallingThread, Type, Descriptor)
 }
 
-void FThreadSafePlayerSessions::SendLatestQPAllPlayers(int LatestQP) const
+void UE::PixelStreaming::FThreadSafePlayerSessions::SendLatestQPAllPlayers(int LatestQP) const
 {
 	SUBMIT_TASK_WITH_PARAMS(SendLatestQPAllPlayers_SignallingThread, LatestQP)
 }
 
-void FThreadSafePlayerSessions::OnRemoteIceCandidate(FPlayerId PlayerId, const std::string& SdpMid, int SdpMLineIndex, const std::string& Sdp)
+void UE::PixelStreaming::FThreadSafePlayerSessions::OnRemoteIceCandidate(FPixelStreamingPlayerId PlayerId, const std::string& SdpMid, int SdpMLineIndex, const std::string& Sdp)
 {
 	SUBMIT_TASK_WITH_PARAMS(OnRemoteIceCandidate_SignallingThread, PlayerId, SdpMid, SdpMLineIndex, Sdp)
 }
 
-void FThreadSafePlayerSessions::OnAnswer(FPlayerId PlayerId, FString Sdp)
+void UE::PixelStreaming::FThreadSafePlayerSessions::OnAnswer(FPixelStreamingPlayerId PlayerId, FString Sdp)
 {
 	SUBMIT_TASK_WITH_PARAMS(OnAnswer_SignallingThread, PlayerId, Sdp)
 }
 
-void FThreadSafePlayerSessions::PollWebRTCStats() const
+void UE::PixelStreaming::FThreadSafePlayerSessions::PollWebRTCStats() const
 {
 	SUBMIT_TASK_NO_PARAMS(PollWebRTCStats_SignallingThread)
 }
@@ -188,117 +193,117 @@ void FThreadSafePlayerSessions::PollWebRTCStats() const
 // Internal methods
 /////////////////////////
 
-void FThreadSafePlayerSessions::PollWebRTCStats_SignallingThread() const
+void UE::PixelStreaming::FThreadSafePlayerSessions::PollWebRTCStats_SignallingThread() const
 {
-	checkf(this->IsInSignallingThread(), TEXT("This method must be called on the signalling thread."));
+	checkf(IsInSignallingThread(), TEXT("This method must be called on the signalling thread."));
 
 	for (auto& Entry : Players)
 	{
-		FPlayerSession* Player = Entry.Value;
+		UE::PixelStreaming::FPlayerSession* Player = Entry.Value;
 		if (Player)
 		{
 			Player->PollWebRTCStats();
 		}
 		else
 		{
-			UE_LOG(PixelStreamer, Log, TEXT("Could not poll WebRTC stats for peer because peer was nullptr."));
+			UE_LOG(LogPixelStreaming, Log, TEXT("Could not poll WebRTC stats for peer because peer was nullptr."));
 		}
 	}
 }
 
-void FThreadSafePlayerSessions::OnRemoteIceCandidate_SignallingThread(FPlayerId PlayerId, const std::string& SdpMid, int SdpMLineIndex, const std::string& Sdp)
+void UE::PixelStreaming::FThreadSafePlayerSessions::OnRemoteIceCandidate_SignallingThread(FPixelStreamingPlayerId PlayerId, const std::string& SdpMid, int SdpMLineIndex, const std::string& Sdp)
 {
-	checkf(this->IsInSignallingThread(), TEXT("This method must be called on the signalling thread."));
+	checkf(IsInSignallingThread(), TEXT("This method must be called on the signalling thread."));
 
-	FPlayerSession* Player = this->GetPlayerSession_SignallingThread(PlayerId);
+	UE::PixelStreaming::FPlayerSession* Player = GetPlayerSession_SignallingThread(PlayerId);
 	if (Player)
 	{
 		Player->OnRemoteIceCandidate(SdpMid, SdpMLineIndex, Sdp);
 	}
 	else
 	{
-		UE_LOG(PixelStreamer, Log, TEXT("Could not pass remote ice candidate to player because Player %s no available."), *PlayerId);
+		UE_LOG(LogPixelStreaming, Log, TEXT("Could not pass remote ice candidate to player because Player %s no available."), *PlayerId);
 	}
 }
 
-void FThreadSafePlayerSessions::OnAnswer_SignallingThread(FPlayerId PlayerId, FString Sdp)
+void UE::PixelStreaming::FThreadSafePlayerSessions::OnAnswer_SignallingThread(FPixelStreamingPlayerId PlayerId, FString Sdp)
 {
-	checkf(this->IsInSignallingThread(), TEXT("This method must be called on the signalling thread."));
+	checkf(IsInSignallingThread(), TEXT("This method must be called on the signalling thread."));
 
-	FPlayerSession* Player = this->GetPlayerSession_SignallingThread(PlayerId);
+	UE::PixelStreaming::FPlayerSession* Player = GetPlayerSession_SignallingThread(PlayerId);
 	if (Player)
 	{
 		Player->OnAnswer(Sdp);
 	}
 	else
 	{
-		UE_LOG(PixelStreamer, Log, TEXT("Could not pass answer to player because Player %s no available."), *PlayerId);
+		UE_LOG(LogPixelStreaming, Log, TEXT("Could not pass answer to player because Player %s no available."), *PlayerId);
 	}
 }
 
-IPixelStreamingAudioSink* FThreadSafePlayerSessions::GetUnlistenedAudioSink_SignallingThread() const
+IPixelStreamingAudioSink* UE::PixelStreaming::FThreadSafePlayerSessions::GetUnlistenedAudioSink_SignallingThread() const
 {
-	checkf(this->IsInSignallingThread(), TEXT("This method must be called on the signalling thread."));
+	checkf(IsInSignallingThread(), TEXT("This method must be called on the signalling thread."));
 
 	for (auto& Entry : Players)
 	{
-		FPlayerSession* Session = Entry.Value;
-		FPixelStreamingAudioSink& AudioSink = Session->GetAudioSink();
+		UE::PixelStreaming::FPlayerSession* Session = Entry.Value;
+		UE::PixelStreaming::FAudioSink& AudioSink = Session->GetAudioSink();
 		if (!AudioSink.HasAudioConsumers())
 		{
-			FPixelStreamingAudioSink* AudioSinkPtr = &AudioSink;
+			UE::PixelStreaming::FAudioSink* AudioSinkPtr = &AudioSink;
 			return static_cast<IPixelStreamingAudioSink*>(AudioSinkPtr);
 		}
 	}
 	return nullptr;
 }
 
-IPixelStreamingAudioSink* FThreadSafePlayerSessions::GetAudioSink_SignallingThread(FPlayerId PlayerId) const
+IPixelStreamingAudioSink* UE::PixelStreaming::FThreadSafePlayerSessions::GetAudioSink_SignallingThread(FPixelStreamingPlayerId PlayerId) const
 {
-	checkf(this->IsInSignallingThread(), TEXT("This method must be called on the signalling thread."));
+	checkf(IsInSignallingThread(), TEXT("This method must be called on the signalling thread."));
 
-	FPlayerSession* PlayerSession = this->GetPlayerSession_SignallingThread(PlayerId);
+	UE::PixelStreaming::FPlayerSession* PlayerSession = GetPlayerSession_SignallingThread(PlayerId);
 
 	if (PlayerSession)
 	{
-		FPixelStreamingAudioSink* AudioSink = &PlayerSession->GetAudioSink();
+		UE::PixelStreaming::FAudioSink* AudioSink = &PlayerSession->GetAudioSink();
 		return static_cast<IPixelStreamingAudioSink*>(AudioSink);
 	}
 	return nullptr;
 }
 
-void FThreadSafePlayerSessions::SendLatestQPAllPlayers_SignallingThread(int LatestQP) const
+void UE::PixelStreaming::FThreadSafePlayerSessions::SendLatestQPAllPlayers_SignallingThread(int LatestQP) const
 {
-	checkf(this->IsInSignallingThread(), TEXT("This method must be called on the signalling thread."));
+	checkf(IsInSignallingThread(), TEXT("This method must be called on the signalling thread."));
 
-	for (auto& PlayerEntry : this->Players)
+	for (auto& PlayerEntry : Players)
 	{
 		PlayerEntry.Value->SendVideoEncoderQP(LatestQP);
 	}
 }
 
-void FThreadSafePlayerSessions::SendLatestQP_SignallingThread(FPlayerId PlayerId, int LatestQP) const
+void UE::PixelStreaming::FThreadSafePlayerSessions::SendLatestQP_SignallingThread(FPixelStreamingPlayerId PlayerId, int LatestQP) const
 {
-	checkf(this->IsInSignallingThread(), TEXT("This method must be called on the signalling thread."));
+	checkf(IsInSignallingThread(), TEXT("This method must be called on the signalling thread."));
 
-	FPlayerSession* Session = this->GetPlayerSession_SignallingThread(PlayerId);
+	UE::PixelStreaming::FPlayerSession* Session = GetPlayerSession_SignallingThread(PlayerId);
 	if (Session)
 	{
 		return Session->SendVideoEncoderQP(LatestQP);
 	}
 	else
 	{
-		UE_LOG(PixelStreamer, Log, TEXT("Could not send latest QP for PlayerId=%s because that player was not found."), *PlayerId);
+		UE_LOG(LogPixelStreaming, Log, TEXT("Could not send latest QP for PlayerId=%s because that player was not found."), *PlayerId);
 	}
 }
 
-bool FThreadSafePlayerSessions::SendMessage_SignallingThread(FPlayerId PlayerId, PixelStreamingProtocol::EToPlayerMsg Type, const FString& Descriptor) const
+bool UE::PixelStreaming::FThreadSafePlayerSessions::SendMessage_SignallingThread(FPixelStreamingPlayerId PlayerId, UE::PixelStreaming::Protocol::EToPlayerMsg Type, const FString& Descriptor) const
 {
-	checkf(this->IsInSignallingThread(), TEXT("This method must be called on the signalling thread."));
+	checkf(IsInSignallingThread(), TEXT("This method must be called on the signalling thread."));
 
-	UE_LOG(PixelStreamer, Log, TEXT("SendMessage to: %s | Type: %d | Message: %s"), *PlayerId, static_cast<int32>(Type), *Descriptor);
+	UE_LOG(LogPixelStreaming, Log, TEXT("SendMessage to: %s | Type: %d | Message: %s"), *PlayerId, static_cast<int32>(Type), *Descriptor);
 
-	FPlayerSession* PlayerSession = this->GetPlayerSession_SignallingThread(PlayerId);
+	UE::PixelStreaming::FPlayerSession* PlayerSession = GetPlayerSession_SignallingThread(PlayerId);
 
 	if (PlayerSession)
 	{
@@ -306,92 +311,92 @@ bool FThreadSafePlayerSessions::SendMessage_SignallingThread(FPlayerId PlayerId,
 	}
 	else
 	{
-		UE_LOG(PixelStreamer, Log, TEXT("Cannot send message to player: %s - player does not exist."), *PlayerId);
+		UE_LOG(LogPixelStreaming, Log, TEXT("Cannot send message to player: %s - player does not exist."), *PlayerId);
 		return false;
 	}
 }
 
-void FThreadSafePlayerSessions::SendMessageAll_SignallingThread(PixelStreamingProtocol::EToPlayerMsg Type, const FString& Descriptor) const
+void UE::PixelStreaming::FThreadSafePlayerSessions::SendMessageAll_SignallingThread(UE::PixelStreaming::Protocol::EToPlayerMsg Type, const FString& Descriptor) const
 {
-	checkf(this->IsInSignallingThread(), TEXT("This method must be called on the signalling thread."));
+	checkf(IsInSignallingThread(), TEXT("This method must be called on the signalling thread."));
 
-	UE_LOG(PixelStreamer, Log, TEXT("SendMessageAll: %d - %s"), static_cast<int32>(Type), *Descriptor);
+	UE_LOG(LogPixelStreaming, Log, TEXT("SendMessageAll: %d - %s"), static_cast<int32>(Type), *Descriptor);
 
-	for (auto& PlayerEntry : this->Players)
+	for (auto& PlayerEntry : Players)
 	{
 		PlayerEntry.Value->SendMessage(Type, Descriptor);
 	}
 }
 
-FPixelStreamingDataChannelObserver* FThreadSafePlayerSessions::GetDataChannelObserver_SignallingThread(FPlayerId PlayerId)
+UE::PixelStreaming::FDataChannelObserver* UE::PixelStreaming::FThreadSafePlayerSessions::GetDataChannelObserver_SignallingThread(FPixelStreamingPlayerId PlayerId)
 {
-	checkf(this->IsInSignallingThread(), TEXT("This method must be called on the signalling thread."));
+	checkf(IsInSignallingThread(), TEXT("This method must be called on the signalling thread."));
 
-	FPlayerSession* Player = this->GetPlayerSession_SignallingThread(PlayerId);
+	UE::PixelStreaming::FPlayerSession* Player = GetPlayerSession_SignallingThread(PlayerId);
 	if (Player)
 	{
 		return &Player->GetDataChannelObserver();
 	}
 	else
 	{
-		UE_LOG(PixelStreamer, Log, TEXT("Cannot get data channel observer for player: %s - player does not exist."), *PlayerId);
+		UE_LOG(LogPixelStreaming, Log, TEXT("Cannot get data channel observer for player: %s - player does not exist."), *PlayerId);
 		return nullptr;
 	}
 }
 
-void FThreadSafePlayerSessions::DisconnectPlayer_SignallingThread(FPlayerId PlayerId, const FString& Reason)
+void UE::PixelStreaming::FThreadSafePlayerSessions::DisconnectPlayer_SignallingThread(FPixelStreamingPlayerId PlayerId, const FString& Reason)
 {
-	checkf(this->IsInSignallingThread(), TEXT("This method must be called on the signalling thread."));
+	checkf(IsInSignallingThread(), TEXT("This method must be called on the signalling thread."));
 
-	FPlayerSession* Player = this->GetPlayerSession_SignallingThread(PlayerId);
+	UE::PixelStreaming::FPlayerSession* Player = GetPlayerSession_SignallingThread(PlayerId);
 	if (Player != nullptr)
 	{
 		Player->DisconnectPlayer(Reason);
 	}
 	else
 	{
-		UE_LOG(PixelStreamer, Log, TEXT("Cannot disconnect player: %s - player does not exist."), *PlayerId);
+		UE_LOG(LogPixelStreaming, Log, TEXT("Cannot disconnect player: %s - player does not exist."), *PlayerId);
 	}
 }
 
-int FThreadSafePlayerSessions::GetNumPlayers_SignallingThread() const
+int UE::PixelStreaming::FThreadSafePlayerSessions::GetNumPlayers_SignallingThread() const
 {
-	checkf(this->IsInSignallingThread(), TEXT("This method must be called on the signalling thread."));
+	checkf(IsInSignallingThread(), TEXT("This method must be called on the signalling thread."));
 
-	return this->Players.Num();
+	return Players.Num();
 }
 
-void FThreadSafePlayerSessions::SendFreezeFrame_SignallingThread(const TArray64<uint8>& JpegBytes)
+void UE::PixelStreaming::FThreadSafePlayerSessions::SendFreezeFrame_SignallingThread(const TArray64<uint8>& JpegBytes)
 {
-	checkf(this->IsInSignallingThread(), TEXT("This method must be called on the signalling thread."));
+	checkf(IsInSignallingThread(), TEXT("This method must be called on the signalling thread."));
 
-	UE_LOG(PixelStreamer, Log, TEXT("Sending freeze frame to players: %d bytes"), JpegBytes.Num());
+	UE_LOG(LogPixelStreaming, Log, TEXT("Sending freeze frame to players: %d bytes"), JpegBytes.Num());
 	{
-		for (auto& PlayerEntry : this->Players)
+		for (auto& PlayerEntry : Players)
 		{
 			PlayerEntry.Value->SendFreezeFrame(JpegBytes);
 		}
 	}
 }
 
-void FThreadSafePlayerSessions::SendFileData_SignallingThread(const TArray<uint8>& ByteData, const FString& MimeType, const FString& FileExtension)
+void UE::PixelStreaming::FThreadSafePlayerSessions::SendFileData_SignallingThread(const TArray<uint8>& ByteData, const FString& MimeType, const FString& FileExtension)
 {
-	checkf(this->IsInSignallingThread(), TEXT("This method must be called on the signalling thread."));
+	checkf(IsInSignallingThread(), TEXT("This method must be called on the signalling thread."));
 
-	UE_LOG(PixelStreamer, Log, TEXT("Sending file with Mime type %s to all players: %d bytes"), *MimeType, ByteData.Num());
+	UE_LOG(LogPixelStreaming, Log, TEXT("Sending file with Mime type %s to all players: %d bytes"), *MimeType, ByteData.Num());
 	{
-		for (auto& PlayerEntry : this->Players)
+		for (auto& PlayerEntry : Players)
 		{
 			PlayerEntry.Value->SendFileData(ByteData, MimeType, FileExtension);
 		}
 	}
 }
 
-void FThreadSafePlayerSessions::SendUnfreezeFrame_SignallingThread()
+void UE::PixelStreaming::FThreadSafePlayerSessions::SendUnfreezeFrame_SignallingThread()
 {
-	checkf(this->IsInSignallingThread(), TEXT("This method must be called on the signalling thread."));
+	checkf(IsInSignallingThread(), TEXT("This method must be called on the signalling thread."));
 
-	UE_LOG(PixelStreamer, Log, TEXT("Sending unfreeze message to players"));
+	UE_LOG(LogPixelStreaming, Log, TEXT("Sending unfreeze message to players"));
 
 	for (auto& PlayerEntry : Players)
 	{
@@ -399,100 +404,105 @@ void FThreadSafePlayerSessions::SendUnfreezeFrame_SignallingThread()
 	}
 }
 
-void FThreadSafePlayerSessions::SendFreezeFrameTo_SignallingThread(FPlayerId PlayerId, const TArray64<uint8>& JpegBytes) const
+void UE::PixelStreaming::FThreadSafePlayerSessions::SendFreezeFrameTo_SignallingThread(FPixelStreamingPlayerId PlayerId, const TArray64<uint8>& JpegBytes) const
 {
-	checkf(this->IsInSignallingThread(), TEXT("This method must be called on the signalling thread."));
+	checkf(IsInSignallingThread(), TEXT("This method must be called on the signalling thread."));
 
-	FPlayerSession* Player = this->GetPlayerSession_SignallingThread(PlayerId);
+	UE::PixelStreaming::FPlayerSession* Player = GetPlayerSession_SignallingThread(PlayerId);
 
 	if (Player != nullptr)
 	{
-		this->Players[PlayerId]->SendFreezeFrame(JpegBytes);
+		Players[PlayerId]->SendFreezeFrame(JpegBytes);
 	}
 	else
 	{
-		UE_LOG(PixelStreamer, Log, TEXT("Cannot send freeze frame to player: %s - player does not exist."), *PlayerId);
+		UE_LOG(LogPixelStreaming, Log, TEXT("Cannot send freeze frame to player: %s - player does not exist."), *PlayerId);
 	}
 }
 
-FPlayerSession* FThreadSafePlayerSessions::GetPlayerSession_SignallingThread(FPlayerId PlayerId) const
+UE::PixelStreaming::FPlayerSession* UE::PixelStreaming::FThreadSafePlayerSessions::GetPlayerSession_SignallingThread(FPixelStreamingPlayerId PlayerId) const
 {
-	checkf(this->IsInSignallingThread(), TEXT("This method must be called on the signalling thread."));
+	checkf(IsInSignallingThread(), TEXT("This method must be called on the signalling thread."));
 
-	if (this->Players.Contains(PlayerId))
+	if (Players.Contains(PlayerId))
 	{
-		return this->Players[PlayerId];
+		return Players[PlayerId];
 	}
 	return nullptr;
 }
 
-void FThreadSafePlayerSessions::SetPlayerSessionDataChannel_SignallingThread(FPlayerId PlayerId, rtc::scoped_refptr<webrtc::DataChannelInterface> DataChannel)
+void UE::PixelStreaming::FThreadSafePlayerSessions::SetPlayerSessionDataChannel_SignallingThread(FPixelStreamingPlayerId PlayerId, rtc::scoped_refptr<webrtc::DataChannelInterface> DataChannel)
 {
-	FPlayerSession* Player = this->GetPlayerSession_SignallingThread(PlayerId);
+	UE::PixelStreaming::FPlayerSession* Player = GetPlayerSession_SignallingThread(PlayerId);
 	if (Player)
 	{
 		Player->SetDataChannel(DataChannel);
 	}
 }
 
-void FThreadSafePlayerSessions::DeleteAllPlayerSessions_SignallingThread()
+void UE::PixelStreaming::FThreadSafePlayerSessions::SetVideoSource_SignallingThread(FPixelStreamingPlayerId PlayerId, rtc::scoped_refptr<webrtc::VideoTrackSourceInterface> VideoSource)
 {
-	checkf(this->IsInSignallingThread(), TEXT("This method must be called on the signalling thread."));
+	UE::PixelStreaming::FPlayerSession* Player = GetPlayerSession_SignallingThread(PlayerId);
+	if (Player)
+	{
+		Player->SetVideoSource(VideoSource);
+	}
+}
 
-	UPixelStreamerDelegates* Delegates = UPixelStreamerDelegates::GetPixelStreamerDelegates();
+void UE::PixelStreaming::FThreadSafePlayerSessions::DeleteAllPlayerSessions_SignallingThread()
+{
+	checkf(IsInSignallingThread(), TEXT("This method must be called on the signalling thread."));
+
+	UPixelStreamingDelegates* Delegates = UPixelStreamingDelegates::GetPixelStreamingDelegates();
 
 	// Delete each player
-	for (auto& Entry : this->Players)
+	for (auto& Entry : Players)
 	{
-		FPlayerId PlayerId = Entry.Key;
-		FPlayerSession* Player = Entry.Value;
+		FPixelStreamingPlayerId PlayerId = Entry.Key;
+		UE::PixelStreaming::FPlayerSession* Player = Entry.Value;
 		delete Player;
-		bool bWasQualityController = this->QualityControllingPlayer == PlayerId;
+		bool bWasQualityController = QualityControllingPlayer == PlayerId;
 
 		if (Delegates && FModuleManager::Get().IsModuleLoaded("PixelStreaming"))
 		{
 			Delegates->OnClosedConnection.Broadcast(PlayerId, bWasQualityController);
+			Delegates->OnClosedConnectionNative.Broadcast(PlayerId, bWasQualityController);
 		}
-
-		// Player deleted, tell all our C++ listeners.
-		this->OnPlayerDeleted.Broadcast(PlayerId);
 	}
 
-	this->Players.Empty();
-	this->QualityControllingPlayer = INVALID_PLAYER_ID;
+	Players.Empty();
+	QualityControllingPlayer = INVALID_PLAYER_ID;
 
 	if (Delegates && FModuleManager::Get().IsModuleLoaded("PixelStreaming"))
 	{
 		Delegates->OnAllConnectionsClosed.Broadcast();
+		Delegates->OnAllConnectionsClosedNative.Broadcast();
 	}
 }
 
-int FThreadSafePlayerSessions::DeletePlayerSession_SignallingThread(FPlayerId PlayerId)
+int UE::PixelStreaming::FThreadSafePlayerSessions::DeletePlayerSession_SignallingThread(FPixelStreamingPlayerId PlayerId)
 {
-	checkf(this->IsInSignallingThread(), TEXT("This method must be called on the signalling thread."));
+	checkf(IsInSignallingThread(), TEXT("This method must be called on the signalling thread."));
 
-	FPlayerSession* Player = this->GetPlayerSession_SignallingThread(PlayerId);
+	UE::PixelStreaming::FPlayerSession* Player = GetPlayerSession_SignallingThread(PlayerId);
 	if (!Player)
 	{
-		UE_LOG(PixelStreamer, VeryVerbose, TEXT("Failed to delete player %s - that player was not found."), *PlayerId);
-		return this->GetNumPlayers_SignallingThread();
+		UE_LOG(LogPixelStreaming, VeryVerbose, TEXT("Failed to delete player %s - that player was not found."), *PlayerId);
+		return GetNumPlayers_SignallingThread();
 	}
 
-	bool bWasQualityController = this->IsQualityController(PlayerId);
+	bool bWasQualityController = IsQualityController(PlayerId);
 
 	// The actual modification to the players map
 	Players.Remove(PlayerId);
 	delete Player;
-
-	UPixelStreamerDelegates* Delegates = UPixelStreamerDelegates::GetPixelStreamerDelegates();
+	UPixelStreamingDelegates* Delegates = UPixelStreamingDelegates::GetPixelStreamingDelegates();
 	if (Delegates)
 	{
 		Delegates->OnClosedConnection.Broadcast(PlayerId, bWasQualityController);
+		Delegates->OnClosedConnectionNative.Broadcast(PlayerId, bWasQualityController);
 	}
 
-	// Player deleted, tell all our C++ listeners.
-	this->OnPlayerDeleted.Broadcast(PlayerId);
-
 	// this is called from WebRTC signalling thread, the only thread were `Players` map is modified, so no need to lock it
 	if (Players.Num() == 0)
 	{
@@ -501,29 +511,30 @@ int FThreadSafePlayerSessions::DeletePlayerSession_SignallingThread(FPlayerId Pl
 		if (Delegates)
 		{
 			Delegates->OnAllConnectionsClosed.Broadcast();
+			Delegates->OnAllConnectionsClosedNative.Broadcast();
 		}
 	}
 	else if (bWasQualityController)
 	{
 		// Quality Controller session has been just removed, set quality control to any of remaining sessions
-		for (auto& Entry : this->Players)
+		for (auto& Entry : Players)
 		{
-			this->SetQualityController_SignallingThread(Entry.Key);
+			SetQualityController_SignallingThread(Entry.Key);
 			break;
 		}
 	}
 
-	return this->GetNumPlayers_SignallingThread();
+	return GetNumPlayers_SignallingThread();
 }
 
-webrtc::PeerConnectionInterface* FThreadSafePlayerSessions::CreatePlayerSession_SignallingThread(
-    FPlayerId PlayerId,
+webrtc::PeerConnectionInterface* UE::PixelStreaming::FThreadSafePlayerSessions::CreatePlayerSession_SignallingThread(
+	FPixelStreamingPlayerId PlayerId,
 	rtc::scoped_refptr<webrtc::PeerConnectionFactoryInterface> PeerConnectionFactory,
 	webrtc::PeerConnectionInterface::RTCConfiguration PeerConnectionConfig,
-	FSignallingServerConnection* SignallingServerConnection,
+	UE::PixelStreaming::FSignallingServerConnection* SignallingServerConnection,
 	int Flags)
 {
-	checkf(this->IsInSignallingThread(), TEXT("This method must be called on the signalling thread."));
+	checkf(IsInSignallingThread(), TEXT("This method must be called on the signalling thread."));
 	check(PeerConnectionFactory);
 
 	// With unified plan, we get several calls to OnOffer, which in turn calls
@@ -534,20 +545,20 @@ webrtc::PeerConnectionInterface* FThreadSafePlayerSessions::CreatePlayerSession_
 		return nullptr;
 	}
 
-	UE_LOG(PixelStreamer, Log, TEXT("Creating player session for PlayerId=%s"), *PlayerId);
+	UE_LOG(LogPixelStreaming, Log, TEXT("Creating player session for PlayerId=%s"), *PlayerId);
 
-	// this is called from WebRTC signalling thread, the only thread were `Players` map is modified, so no need to lock it
+	// this is called from WebRTC signalling thread, the only thread where `Players` map is modified, so no need to lock it
 	bool bMakeQualityController = Players.Num() == 0; // first player controls quality by default
-	FPlayerSession* Session = new FPlayerSession(this, SignallingServerConnection, PlayerId);
+	UE::PixelStreaming::FPlayerSession* Session = new UE::PixelStreaming::FPlayerSession(this, SignallingServerConnection, PlayerId);
 
 	rtc::scoped_refptr<webrtc::PeerConnectionInterface> PeerConnection = PeerConnectionFactory->CreatePeerConnection(PeerConnectionConfig, webrtc::PeerConnectionDependencies{ Session });
 	check(PeerConnection);
 
 	// Setup suggested bitrate settings on the Peer Connection based on our CVars
 	webrtc::BitrateSettings BitrateSettings;
-	BitrateSettings.min_bitrate_bps = PixelStreamingSettings::CVarPixelStreamingWebRTCMinBitrate.GetValueOnAnyThread();
-	BitrateSettings.max_bitrate_bps = PixelStreamingSettings::CVarPixelStreamingWebRTCMaxBitrate.GetValueOnAnyThread();
-	BitrateSettings.start_bitrate_bps = PixelStreamingSettings::CVarPixelStreamingWebRTCStartBitrate.GetValueOnAnyThread();
+	BitrateSettings.min_bitrate_bps = UE::PixelStreaming::Settings::CVarPixelStreamingWebRTCMinBitrate.GetValueOnAnyThread();
+	BitrateSettings.max_bitrate_bps = UE::PixelStreaming::Settings::CVarPixelStreamingWebRTCMaxBitrate.GetValueOnAnyThread();
+	BitrateSettings.start_bitrate_bps = UE::PixelStreaming::Settings::CVarPixelStreamingWebRTCStartBitrate.GetValueOnAnyThread();
 	PeerConnection->SetBitrate(BitrateSettings);
 
 	Session->SetPeerConnection(PeerConnection);
@@ -557,40 +568,44 @@ webrtc::PeerConnectionInterface* FThreadSafePlayerSessions::CreatePlayerSession_
 
 	if (bMakeQualityController)
 	{
-		this->SetQualityController_SignallingThread(PlayerId);
+		SetQualityController_SignallingThread(PlayerId);
 	}
 
-	if (UPixelStreamerDelegates* Delegates = UPixelStreamerDelegates::GetPixelStreamerDelegates())
+	if (UPixelStreamingDelegates* Delegates = UPixelStreamingDelegates::GetPixelStreamingDelegates())
 	{
 		Delegates->OnNewConnection.Broadcast(PlayerId, bMakeQualityController);
+		Delegates->OnNewConnectionNative.Broadcast(PlayerId, bMakeQualityController);
 	}
 
 	return PeerConnection.get();
 }
 
-void FThreadSafePlayerSessions::SetQualityController_SignallingThread(FPlayerId PlayerId)
+void UE::PixelStreaming::FThreadSafePlayerSessions::SetQualityController_SignallingThread(FPixelStreamingPlayerId PlayerId)
 {
-	checkf(this->IsInSignallingThread(), TEXT("This method must be called on the signalling thread."));
+	checkf(IsInSignallingThread(), TEXT("This method must be called on the signalling thread."));
 
-	if (this->Players.Contains(PlayerId))
+	if (Players.Contains(PlayerId))
 	{
 
 		// The actual assignment of the quality controlling peeer
 		{
 			FScopeLock Lock(&QualityControllerCS);
-			this->QualityControllingPlayer = PlayerId;
+			QualityControllingPlayer = PlayerId;
 		}
 
 		// Let any listeners know the quality controller has changed
-		this->OnQualityControllerChanged.Broadcast(PlayerId);
+		if (UPixelStreamingDelegates* Delegates = UPixelStreamingDelegates::GetPixelStreamingDelegates())
+		{
+			Delegates->OnQualityControllerChangedNative.Broadcast(PlayerId);
+		}
 
-		UE_LOG(PixelStreamer, Log, TEXT("Quality controller is now PlayerId=%s."), *PlayerId);
+		UE_LOG(LogPixelStreaming, Log, TEXT("Quality controller is now PlayerId=%s."), *PlayerId);
 
 		// Update quality controller status on the browser side too
 		for (auto& Entry : Players)
 		{
-			FPlayerSession* Session = Entry.Value;
-			bool bIsQualityController = Entry.Key == this->QualityControllingPlayer;
+			UE::PixelStreaming::FPlayerSession* Session = Entry.Value;
+			bool bIsQualityController = Entry.Key == QualityControllingPlayer;
 			if (Session)
 			{
 				Session->SendQualityControlStatus(bIsQualityController);
@@ -599,7 +614,7 @@ void FThreadSafePlayerSessions::SetQualityController_SignallingThread(FPlayerId
 	}
 	else
 	{
-		UE_LOG(PixelStreamer, Log, TEXT("Could not set quality controller for PlayerId=%s - that player does not exist."), *PlayerId);
+		UE_LOG(LogPixelStreaming, Log, TEXT("Could not set quality controller for PlayerId=%s - that player does not exist."), *PlayerId);
 	}
 }
 
diff --git a/Engine/Plugins/Media/PixelStreaming/Source/PixelStreaming/Private/ThreadSafePlayerSessions.h b/Engine/Plugins/Media/PixelStreaming/Source/PixelStreaming/Private/ThreadSafePlayerSessions.h
--- a/Engine/Plugins/Media/PixelStreaming/Source/PixelStreaming/Private/ThreadSafePlayerSessions.h
+++ b/Engine/Plugins/Media/PixelStreaming/Source/PixelStreaming/Private/ThreadSafePlayerSessions.h
@@ -3,97 +3,95 @@
 
 #include "IPixelStreamingSessions.h"
 #include "PlayerSession.h"
-#include "PixelStreamingDataChannelObserver.h"
+#include "DataChannelObserver.h"
 #include "Delegates/Delegate.h"
 #include "Delegates/DelegateCombinations.h"
 
-// This class provides a mechanism to read/write FPlayerSession on any thread.
-// This is mostly achieved by submitting all such work to the WebRTC signalling thread
-// using the WebRTC task queue (which is a single thread).
-// Note: Any public methods in this class with a return type will making the calling thread block
-// until the WebRTC thread is done and can return the result.
-class FThreadSafePlayerSessions : public IPixelStreamingSessions
-{
-public:
-	FThreadSafePlayerSessions(rtc::Thread* WebRtcSignallingThread);
-	virtual ~FThreadSafePlayerSessions() = default;
-	bool IsInSignallingThread() const;
-	void SendFreezeFrame(const TArray64<uint8>& JpegBytes);
-	void SendUnfreezeFrame();
-	void SendFileData(const TArray<uint8>& ByteData, const FString& MimeType, const FString& FileExtension);
-	void SendMessageAll(PixelStreamingProtocol::EToPlayerMsg Type, const FString& Descriptor) const;
-	void DisconnectPlayer(FPlayerId PlayerId, const FString& Reason);
-	void SendLatestQPAllPlayers(int LatestQP) const;
-	void OnRemoteIceCandidate(FPlayerId PlayerId, const std::string& SdpMid, int SdpMLineIndex, const std::string& Sdp);
-	void OnAnswer(FPlayerId PlayerId, FString Sdp);
+namespace UE {
+	namespace PixelStreaming {
+		// This class provides a mechanism to read/write FPlayerSession on any thread.
+		// This is mostly achieved by submitting all such work to the WebRTC signalling thread
+		// using the WebRTC task queue (which is a single thread).
+		// Note: Any public methods in this class with a return type will making the calling thread block
+		// until the WebRTC thread is done and can return the result.
+		class FThreadSafePlayerSessions : public IPixelStreamingSessions
+		{
+		public:
+			FThreadSafePlayerSessions(rtc::Thread* WebRtcSignallingThread);
+			virtual ~FThreadSafePlayerSessions() = default;
+			bool IsInSignallingThread() const;
+			void SendFreezeFrame(const TArray64<uint8>& JpegBytes);
+			void SendUnfreezeFrame();
+			void SendFileData(const TArray<uint8>& ByteData, const FString& MimeType, const FString& FileExtension);
+			void SendMessageAll(Protocol::EToPlayerMsg Type, const FString& Descriptor) const;
+			void DisconnectPlayer(FPixelStreamingPlayerId PlayerId, const FString& Reason);
+			void SendLatestQPAllPlayers(int LatestQP) const;
+			void OnRemoteIceCandidate(FPixelStreamingPlayerId PlayerId, const std::string& SdpMid, int SdpMLineIndex, const std::string& Sdp);
+			void OnAnswer(FPixelStreamingPlayerId PlayerId, FString Sdp);
 
-	webrtc::PeerConnectionInterface* CreatePlayerSession(FPlayerId PlayerId,
-		rtc::scoped_refptr<webrtc::PeerConnectionFactoryInterface> PeerConnectionFactory,
-		webrtc::PeerConnectionInterface::RTCConfiguration PeerConnectionConfig,
-		FSignallingServerConnection* SignallingServerConnection,
-		int Flags);
+			webrtc::PeerConnectionInterface* CreatePlayerSession(FPixelStreamingPlayerId PlayerId,
+				rtc::scoped_refptr<webrtc::PeerConnectionFactoryInterface> PeerConnectionFactory,
+				webrtc::PeerConnectionInterface::RTCConfiguration PeerConnectionConfig,
+				FSignallingServerConnection* SignallingServerConnection,
+				int Flags);
 
-	void SetPlayerSessionDataChannel(FPlayerId PlayerId, rtc::scoped_refptr<webrtc::DataChannelInterface> DataChannel);
+			void DeleteAllPlayerSessions();
+			int DeletePlayerSession(FPixelStreamingPlayerId PlayerId);
+			FDataChannelObserver* GetDataChannelObserver(FPixelStreamingPlayerId PlayerId);
+			void SetPlayerSessionDataChannel(FPixelStreamingPlayerId PlayerId, rtc::scoped_refptr<webrtc::DataChannelInterface> DataChannel);
+			void SetVideoSource(FPixelStreamingPlayerId PlayerId, rtc::scoped_refptr<webrtc::VideoTrackSourceInterface> VideoSource);
 
-	void DeleteAllPlayerSessions();
-	int DeletePlayerSession(FPlayerId PlayerId);
-	FPixelStreamingDataChannelObserver* GetDataChannelObserver(FPlayerId PlayerId);
+			// Begin IPixelStreamingSessions
+			virtual int GetNumPlayers() const override;
+			virtual IPixelStreamingAudioSink* GetAudioSink(FPixelStreamingPlayerId PlayerId) const override;
+			virtual IPixelStreamingAudioSink* GetUnlistenedAudioSink() const override;
+			virtual bool IsQualityController(FPixelStreamingPlayerId PlayerId) const override;
+			virtual void SetQualityController(FPixelStreamingPlayerId PlayerId) override;
+			virtual bool SendMessage(FPixelStreamingPlayerId PlayerId, Protocol::EToPlayerMsg Type, const FString& Descriptor) const override;
+			virtual void SendLatestQP(FPixelStreamingPlayerId PlayerId, int LatestQP) const override;
+			virtual void SendFreezeFrameTo(FPixelStreamingPlayerId PlayerId, const TArray64<uint8>& JpegBytes) const override;
+			void PollWebRTCStats() const override;
+			// End IPixelStreamingSessions
 
-	// Begin IPixelStreamingSessions
-	virtual int GetNumPlayers() const override;
-	virtual IPixelStreamingAudioSink* GetAudioSink(FPlayerId PlayerId) const override;
-	virtual IPixelStreamingAudioSink* GetUnlistenedAudioSink() const override;
-	virtual bool IsQualityController(FPlayerId PlayerId) const override;
-	virtual void SetQualityController(FPlayerId PlayerId) override;
-	virtual bool SendMessage(FPlayerId PlayerId, PixelStreamingProtocol::EToPlayerMsg Type, const FString& Descriptor) const override;
-	virtual void SendLatestQP(FPlayerId PlayerId, int LatestQP) const override;
-	virtual void SendFreezeFrameTo(FPlayerId PlayerId, const TArray64<uint8>& JpegBytes) const override;
-	void PollWebRTCStats() const override;
-	// End IPixelStreamingSessions
+		private:
+			// Note: This is very intentionally internal and there is no public version because as soon as we hand it out it isn't thread safe anymore.
+			FPlayerSession* GetPlayerSession_SignallingThread(FPixelStreamingPlayerId PlayerId) const;
 
-public:
-	DECLARE_MULTICAST_DELEGATE_OneParam(FOnQualityControllerChanged, FPlayerId)
-		FOnQualityControllerChanged OnQualityControllerChanged;
+			void SetVideoSource_SignallingThread(FPixelStreamingPlayerId PlayerId, rtc::scoped_refptr<webrtc::VideoTrackSourceInterface> VideoSource);
+			void SetPlayerSessionDataChannel_SignallingThread(FPixelStreamingPlayerId PlayerId, rtc::scoped_refptr<webrtc::DataChannelInterface> DataChannel);
+			void PollWebRTCStats_SignallingThread() const;
+			void OnRemoteIceCandidate_SignallingThread(FPixelStreamingPlayerId PlayerId, const std::string& SdpMid, int SdpMLineIndex, const std::string& Sdp);
+			void OnAnswer_SignallingThread(FPixelStreamingPlayerId PlayerId, FString Sdp);
+			IPixelStreamingAudioSink* GetUnlistenedAudioSink_SignallingThread() const;
+			IPixelStreamingAudioSink* GetAudioSink_SignallingThread(FPixelStreamingPlayerId PlayerId) const;
+			void SendLatestQP_SignallingThread(FPixelStreamingPlayerId PlayerId, int LatestQP) const;
+			int GetNumPlayers_SignallingThread() const;
+			void SendFreezeFrame_SignallingThread(const TArray64<uint8>& JpegBytes);
+			void SendUnfreezeFrame_SignallingThread();
+			void SendFreezeFrameTo_SignallingThread(FPixelStreamingPlayerId PlayerId, const TArray64<uint8>& JpegBytes) const;
+			void SendFileData_SignallingThread(const TArray<uint8>& ByteData, const FString& MimeType, const FString& FileExtension);
+			void SetQualityController_SignallingThread(FPixelStreamingPlayerId PlayerId);
+			void DisconnectPlayer_SignallingThread(FPixelStreamingPlayerId PlayerId, const FString& Reason);
+			void SendMessageAll_SignallingThread(Protocol::EToPlayerMsg Type, const FString& Descriptor) const;
+			bool SendMessage_SignallingThread(FPixelStreamingPlayerId PlayerId, Protocol::EToPlayerMsg Type, const FString& Descriptor) const;
+			void SendLatestQPAllPlayers_SignallingThread(int LatestQP) const;
 
-	DECLARE_MULTICAST_DELEGATE_OneParam(FOnPlayerDeleted, FPlayerId)
-		FOnPlayerDeleted OnPlayerDeleted;
+			webrtc::PeerConnectionInterface* CreatePlayerSession_SignallingThread(FPixelStreamingPlayerId PlayerId,
+				rtc::scoped_refptr<webrtc::PeerConnectionFactoryInterface> PeerConnectionFactory,
+				webrtc::PeerConnectionInterface::RTCConfiguration PeerConnectionConfig,
+				FSignallingServerConnection* SignallingServerConnection,
+				int Flags);
 
-private:
-	// Note: This is very intentionally internal and there is no public version because as soon as we hand it out it isn't thread safe anymore.
-	FPlayerSession* GetPlayerSession_SignallingThread(FPlayerId PlayerId) const;
+			void DeleteAllPlayerSessions_SignallingThread();
+			int DeletePlayerSession_SignallingThread(FPixelStreamingPlayerId PlayerId);
+			FDataChannelObserver* GetDataChannelObserver_SignallingThread(FPixelStreamingPlayerId PlayerId);
 
-	void SetPlayerSessionDataChannel_SignallingThread(FPlayerId PlayerId, rtc::scoped_refptr<webrtc::DataChannelInterface> DataChannel);
-	void PollWebRTCStats_SignallingThread() const;
-	void OnRemoteIceCandidate_SignallingThread(FPlayerId PlayerId, const std::string& SdpMid, int SdpMLineIndex, const std::string& Sdp);
-	void OnAnswer_SignallingThread(FPlayerId PlayerId, FString Sdp);
-	IPixelStreamingAudioSink* GetUnlistenedAudioSink_SignallingThread() const;
-	IPixelStreamingAudioSink* GetAudioSink_SignallingThread(FPlayerId PlayerId) const;
-	void SendLatestQP_SignallingThread(FPlayerId PlayerId, int LatestQP) const;
-	int GetNumPlayers_SignallingThread() const;
-	void SendFreezeFrame_SignallingThread(const TArray64<uint8>& JpegBytes);
-	void SendUnfreezeFrame_SignallingThread();
-	void SendFreezeFrameTo_SignallingThread(FPlayerId PlayerId, const TArray64<uint8>& JpegBytes) const;
-	void SendFileData_SignallingThread(const TArray<uint8>& ByteData, const FString& MimeType, const FString& FileExtension);
-	void SetQualityController_SignallingThread(FPlayerId PlayerId);
-	void DisconnectPlayer_SignallingThread(FPlayerId PlayerId, const FString& Reason);
-	void SendMessageAll_SignallingThread(PixelStreamingProtocol::EToPlayerMsg Type, const FString& Descriptor) const;
-	bool SendMessage_SignallingThread(FPlayerId PlayerId, PixelStreamingProtocol::EToPlayerMsg Type, const FString& Descriptor) const;
-	void SendLatestQPAllPlayers_SignallingThread(int LatestQP) const;
+		private:
+			rtc::Thread* WebRtcSignallingThread;
+			TMap<FPixelStreamingPlayerId, FPlayerSession*> Players;
 
-	webrtc::PeerConnectionInterface* CreatePlayerSession_SignallingThread(FPlayerId PlayerId,
-		rtc::scoped_refptr<webrtc::PeerConnectionFactoryInterface> PeerConnectionFactory,
-		webrtc::PeerConnectionInterface::RTCConfiguration PeerConnectionConfig,
-		FSignallingServerConnection* SignallingServerConnection,
-		int Flags);
-
-	void DeleteAllPlayerSessions_SignallingThread();
-	int DeletePlayerSession_SignallingThread(FPlayerId PlayerId);
-	FPixelStreamingDataChannelObserver* GetDataChannelObserver_SignallingThread(FPlayerId PlayerId);
-
-private:
-	rtc::Thread* WebRtcSignallingThread;
-	TMap<FPlayerId, FPlayerSession*> Players;
-
-	mutable FCriticalSection QualityControllerCS;
-	FPlayerId QualityControllingPlayer = ToPlayerId(FString(TEXT("No quality controlling peer.")));
-};
\ No newline at end of file
+			mutable FCriticalSection QualityControllerCS;
+			FPixelStreamingPlayerId QualityControllingPlayer = ToPlayerId(FString(TEXT("No quality controlling peer.")));
+		};
+	}
+}
\ No newline at end of file
diff --git a/Engine/Plugins/Media/PixelStreaming/Source/PixelStreaming/Private/Utils.h b/Engine/Plugins/Media/PixelStreaming/Source/PixelStreaming/Private/Utils.h
--- a/Engine/Plugins/Media/PixelStreaming/Source/PixelStreaming/Private/Utils.h
+++ b/Engine/Plugins/Media/PixelStreaming/Source/PixelStreaming/Private/Utils.h
@@ -2,314 +2,91 @@
 
 #pragma once
 
-#include "PixelStreamingPrivate.h"
-
-#if PLATFORM_WINDOWS || PLATFORM_XBOXONE
-	#include "Windows/WindowsPlatformMisc.h"
-#elif PLATFORM_LINUX
-	#include "Linux/LinuxPlatformMisc.h"
-#endif
-
 #include "RHI.h"
-#include "Containers/StringConv.h"
-#include "Containers/UnrealString.h"
-#include "Templates/SharedPointer.h"
-#include "Dom/JsonObject.h"
-#include "Policies/PrettyJsonPrintPolicy.h"
-#include "Policies/CondensedJsonPrintPolicy.h"
-#include "Serialization/JsonSerializer.h"
-#include "Templates/Atomic.h"
-#include "WebRTCIncludes.h"
-#include <string>
 #include "CommonRenderResources.h"
 #include "ScreenRendering.h"
 #include "RHIStaticStates.h"
 #include "Modules/ModuleManager.h"
-#include "GenericPlatform/GenericPlatformProcess.h"
 #include "PixelStreamingPrivate.h"
-#include "Async/Async.h"
-
-#if PLATFORM_WINDOWS || PLATFORM_XBOXONE
-inline bool IsWindows7Plus()
-{
-	return FPlatformMisc::VerifyWindowsVersion(6, 1);
-}
-
-inline bool IsWindows8Plus()
-{
-	return FPlatformMisc::VerifyWindowsVersion(6, 2);
-}
-#endif
-
-inline FString ToString(const std::string& Str)
-{
-	auto Conv = StringCast<TCHAR>(Str.c_str(), Str.size());
-	FString Res{ Conv.Length(), Conv.Get() };
-	return Res;
-}
-
-inline std::string to_string(const FString& Str)
-{
-	auto Ansi = StringCast<ANSICHAR>(*Str, Str.Len());
-	std::string Res{ Ansi.Get(), static_cast<SIZE_T>(Ansi.Length()) };
-	return Res;
-}
-
-inline FString ToString(const TSharedPtr<FJsonObject>& JsonObj, bool bPretty = true)
-{
-	FString Res;
-	if (bPretty)
-	{
-		auto JsonWriter = TJsonWriterFactory<TCHAR, TPrettyJsonPrintPolicy<TCHAR>>::Create(&Res);
-		FJsonSerializer::Serialize(JsonObj.ToSharedRef(), JsonWriter);
-	}
-	else
-	{
-		auto JsonWriter = TJsonWriterFactory<TCHAR, TCondensedJsonPrintPolicy<TCHAR>>::Create(&Res);
-		FJsonSerializer::Serialize(JsonObj.ToSharedRef(), JsonWriter);
-	}
-	return Res;
-}
-
-inline const TCHAR* ToString(webrtc::PeerConnectionInterface::SignalingState Val)
-{
-	TCHAR const* SignallingStatesStr[] = {
-		TEXT("Stable"),
-		TEXT("HaveLocalOffer"),
-		TEXT("HaveLocalPrAnswer"),
-		TEXT("HaveRemoteOffer"),
-		TEXT("HaveRemotePrAnswer"),
-		TEXT("Closed")
-	};
-
-	return ensureMsgf(0 <= Val && Val <= webrtc::PeerConnectionInterface::kClosed, TEXT("Invalid `webrtc::PeerConnectionInterface::SignalingState` value: %d"), static_cast<uint32>(Val)) ? SignallingStatesStr[Val] : TEXT("Unknown");
-}
-
-inline const TCHAR* ToString(webrtc::PeerConnectionInterface::IceConnectionState Val)
-{
-	TCHAR const* IceConnectionStatsStr[] = {
-		TEXT("IceConnectionNew"),
-		TEXT("IceConnectionChecking"),
-		TEXT("IceConnectionConnected"),
-		TEXT("IceConnectionCompleted"),
-		TEXT("IceConnectionFailed"),
-		TEXT("IceConnectionDisconnected"),
-		TEXT("IceConnectionClosed")
-	};
-
-	return ensureMsgf(
-		0 <= Val && Val < webrtc::PeerConnectionInterface::kIceConnectionMax,
-		TEXT("Invalid `webrtc::PeerConnectionInterface::IceConnectionState` value: %d"),
-		static_cast<uint32>(Val))
-			? IceConnectionStatsStr[Val]
-			: TEXT("Unknown");
-}
-
-inline const TCHAR* ToString(webrtc::PeerConnectionInterface::IceGatheringState Val)
-{
-	TCHAR const* IceGatheringStatsStr[] = {
-		TEXT("IceGatheringNew"),
-		TEXT("IceGatheringGathering"),
-		TEXT("IceGatheringComplete")
-	};
-
-	return ensureMsgf(
-		0 <= Val && Val <= webrtc::PeerConnectionInterface::kIceGatheringComplete,
-		TEXT("Invalid `webrtc::PeerConnectionInterface::IceGatheringState` value: %d"),
-		static_cast<uint32>(Val))
-			? IceGatheringStatsStr[Val]
-			: TEXT("Unknown");
-}
-
-inline const TCHAR* ToString(webrtc::VideoFrameType FrameType)
-{
-	TCHAR const* FrameTypesStr[] = {
-		TEXT("EmptyFrame"),
-		TEXT("AudioFrameSpeech"),
-		TEXT("AudioFrameCN"),
-		TEXT("VideoFrameKey"),
-		TEXT("VideoFrameDelta")
-	};
-	int FrameTypeInt = (int)FrameType;
-	return ensureMsgf(
-		0 <= FrameTypeInt && FrameTypeInt <= (int)webrtc::VideoFrameType::kVideoFrameDelta,
-		TEXT("Invalid `webrtc::FrameType`: %d"), static_cast<uint32>(FrameType))
-			? FrameTypesStr[FrameTypeInt]
-			: TEXT("Unknown");
-}
-
-inline webrtc::SdpVideoFormat CreateH264Format(webrtc::H264::Profile profile, webrtc::H264::Level level)
-{
-	const absl::optional<std::string> profile_string =
-		webrtc::H264::ProfileLevelIdToString(webrtc::H264::ProfileLevelId(profile, level));
-	check(profile_string);
-	return webrtc::SdpVideoFormat(
-		cricket::kH264CodecName,
-		{ { cricket::kH264FmtpProfileLevelId, *profile_string },
-			{ cricket::kH264FmtpLevelAsymmetryAllowed, "1" },
-			{ cricket::kH264FmtpPacketizationMode, "1" } });
-}
-
-// returns milliseconds part of current time, useful for logging
-inline uint32 RtcTimeMs()
-{
-	return rtc::TimeMicros() / 1000 % 1000;
-}
-
-template <uint32 SmoothingPeriod>
-class FSmoothedValue
-{
-public:
-	double Get() const
-	{
-		return Value;
-	}
-
-	void Update(double NewValue)
-	{
-		Value = Value * (SmoothingPeriod - 1) / SmoothingPeriod + NewValue / SmoothingPeriod;
-	}
-
-	void Reset()
-	{
-		Value = 0;
-	}
-
-private:
-	TAtomic<double> Value{ 0 };
-};
-
-inline void CopyTexture(FTexture2DRHIRef SourceTexture, FTexture2DRHIRef DestinationTexture, FGPUFenceRHIRef& CopyFence)
-{
-	FRHICommandListImmediate& RHICmdList = FRHICommandListExecutor::GetImmediateCommandList();
-
-	IRendererModule* RendererModule = &FModuleManager::GetModuleChecked<IRendererModule>("Renderer");
-
-	// #todo-renderpasses there's no explicit resolve here? Do we need one?
-	FRHIRenderPassInfo RPInfo(DestinationTexture, ERenderTargetActions::Load_Store);
-
-	RHICmdList.BeginRenderPass(RPInfo, TEXT("CopyBackbuffer"));
-
-	{
-		RHICmdList.SetViewport(0, 0, 0.0f, DestinationTexture->GetSizeX(), DestinationTexture->GetSizeY(), 1.0f);
-
-		FGraphicsPipelineStateInitializer GraphicsPSOInit;
-		RHICmdList.ApplyCachedRenderTargets(GraphicsPSOInit);
-		GraphicsPSOInit.BlendState = TStaticBlendState<>::GetRHI();
-		GraphicsPSOInit.RasterizerState = TStaticRasterizerState<>::GetRHI();
-		GraphicsPSOInit.DepthStencilState = TStaticDepthStencilState<false, CF_Always>::GetRHI();
-
-		// New engine version...
-		FGlobalShaderMap* ShaderMap = GetGlobalShaderMap(GMaxRHIFeatureLevel);
-		TShaderMapRef<FScreenVS> VertexShader(ShaderMap);
-		TShaderMapRef<FScreenPS> PixelShader(ShaderMap);
-
-		GraphicsPSOInit.BoundShaderState.VertexDeclarationRHI = GFilterVertexDeclaration.VertexDeclarationRHI;
-		GraphicsPSOInit.BoundShaderState.VertexShaderRHI = VertexShader.GetVertexShader();
-		GraphicsPSOInit.BoundShaderState.PixelShaderRHI = PixelShader.GetPixelShader();
-
-		GraphicsPSOInit.PrimitiveType = PT_TriangleList;
-
-		SetGraphicsPipelineState(RHICmdList, GraphicsPSOInit, 0);
 
-		if (DestinationTexture->GetSizeX() != SourceTexture->GetSizeX() || DestinationTexture->GetSizeY() != SourceTexture->GetSizeY())
+namespace UE {
+	namespace PixelStreaming {
+		inline void CopyTexture(FTexture2DRHIRef SourceTexture, FTexture2DRHIRef DestinationTexture, FGPUFenceRHIRef& CopyFence)
 		{
-			PixelShader->SetParameters(RHICmdList, TStaticSamplerState<SF_Bilinear>::GetRHI(), SourceTexture);
-		}
-		else
-		{
-			PixelShader->SetParameters(RHICmdList, TStaticSamplerState<SF_Point>::GetRHI(), SourceTexture);
-		}
+			FRHICommandListImmediate& RHICmdList = FRHICommandListExecutor::GetImmediateCommandList();
 
-		RendererModule->DrawRectangle(RHICmdList, 0, 0, // Dest X, Y
-			DestinationTexture->GetSizeX(),				// Dest Width
-			DestinationTexture->GetSizeY(),				// Dest Height
-			0, 0,										// Source U, V
-			1, 1,										// Source USize, VSize
-			DestinationTexture->GetSizeXY(),			// Target buffer size
-			FIntPoint(1, 1),							// Source texture size
-			VertexShader, EDRF_Default);
-	}
+			IRendererModule* RendererModule = &FModuleManager::GetModuleChecked<IRendererModule>("Renderer");
 
-	RHICmdList.EndRenderPass();
+			// #todo-renderpasses there's no explicit resolve here? Do we need one?
+			FRHIRenderPassInfo RPInfo(DestinationTexture, ERenderTargetActions::Load_Store);
 
-	RHICmdList.WriteGPUFence(CopyFence);
+			RHICmdList.BeginRenderPass(RPInfo, TEXT("CopyBackbuffer"));
 
-	RHICmdList.ImmediateFlush(EImmediateFlushType::WaitForOutstandingTasksOnly);
-}
+			{
+				RHICmdList.SetViewport(0, 0, 0.0f, DestinationTexture->GetSizeX(), DestinationTexture->GetSizeY(), 1.0f);
 
-inline size_t SerializeToBuffer(rtc::CopyOnWriteBuffer& Buffer, size_t Pos, const void* Data, size_t DataSize)
-{
-	FMemory::Memcpy(&Buffer[Pos], reinterpret_cast<const uint8_t*>(Data), DataSize);
-	return Pos + DataSize;
-}
+				FGraphicsPipelineStateInitializer GraphicsPSOInit;
+				RHICmdList.ApplyCachedRenderTargets(GraphicsPSOInit);
+				GraphicsPSOInit.BlendState = TStaticBlendState<>::GetRHI();
+				GraphicsPSOInit.RasterizerState = TStaticRasterizerState<>::GetRHI();
+				GraphicsPSOInit.DepthStencilState = TStaticDepthStencilState<false, CF_Always>::GetRHI();
 
-inline void SendArbitraryData(rtc::scoped_refptr<webrtc::DataChannelInterface> DataChannel, const TArray<uint8>& DataBytes, const uint8 MessageType)
-{
-	if (!DataChannel)
-	{
-		return;
-	}
+				// New engine version...
+				FGlobalShaderMap* ShaderMap = GetGlobalShaderMap(GMaxRHIFeatureLevel);
+				TShaderMapRef<FScreenVS> VertexShader(ShaderMap);
+				TShaderMapRef<FScreenPS> PixelShader(ShaderMap);
 
-	// int32 results in a maximum 4GB file (4,294,967,296 bytes)
-	const int32 DataSize = DataBytes.Num();
+				GraphicsPSOInit.BoundShaderState.VertexDeclarationRHI = GFilterVertexDeclaration.VertexDeclarationRHI;
+				GraphicsPSOInit.BoundShaderState.VertexShaderRHI = VertexShader.GetVertexShader();
+				GraphicsPSOInit.BoundShaderState.PixelShaderRHI = PixelShader.GetPixelShader();
 
-	// Maximum size of a single buffer should be 16KB as this is spec compliant message length for a single data channel transmission
-	const int32 MaxBufferBytes = 16 * 1024;
-	const int32 MessageHeader = sizeof(MessageType) + sizeof(DataSize);
-	const int32 MaxDataBytesPerMsg = MaxBufferBytes - MessageHeader;
+				GraphicsPSOInit.PrimitiveType = PT_TriangleList;
 
-	int32 BytesTransmitted = 0;
+				SetGraphicsPipelineState(RHICmdList, GraphicsPSOInit, 0);
 
-	while (BytesTransmitted < DataSize)
-	{
-		int32 RemainingBytes = DataSize - BytesTransmitted;
-		int32 BytesToTransmit = FGenericPlatformMath::Min(MaxDataBytesPerMsg, RemainingBytes);
+				if (DestinationTexture->GetSizeX() != SourceTexture->GetSizeX() || DestinationTexture->GetSizeY() != SourceTexture->GetSizeY())
+				{
+					PixelShader->SetParameters(RHICmdList, TStaticSamplerState<SF_Bilinear>::GetRHI(), SourceTexture);
+				}
+				else
+				{
+					PixelShader->SetParameters(RHICmdList, TStaticSamplerState<SF_Point>::GetRHI(), SourceTexture);
+				}
 
-		rtc::CopyOnWriteBuffer Buffer(MessageHeader + BytesToTransmit);
+				RendererModule->DrawRectangle(RHICmdList, 0, 0, // Dest X, Y
+					DestinationTexture->GetSizeX(),				// Dest Width
+					DestinationTexture->GetSizeY(),				// Dest Height
+					0, 0,										// Source U, V
+					1, 1,										// Source USize, VSize
+					DestinationTexture->GetSizeXY(),			// Target buffer size
+					FIntPoint(1, 1),							// Source texture size
+					VertexShader, EDRF_Default);
+			}
 
-		size_t Pos = 0;
+			RHICmdList.EndRenderPass();
 
-		// Write message type
-		Pos = SerializeToBuffer(Buffer, Pos, &MessageType, sizeof(MessageType));
-		// Write size of payload
-		Pos = SerializeToBuffer(Buffer, Pos, &DataSize, sizeof(DataSize));
-		// Write the data bytes payload
-		Pos = SerializeToBuffer(Buffer, Pos, DataBytes.GetData() + BytesTransmitted, BytesToTransmit);
+			RHICmdList.WriteGPUFence(CopyFence);
 
-		uint64_t BufferBefore = DataChannel->buffered_amount();
-		while (BufferBefore + BytesToTransmit >= 16 * 1024 * 1024) // 16MB (WebRTC Data Channel buffer size)
-		{
-			FPlatformProcess::Sleep(0.000001f); // sleep 1 microsecond
-			BufferBefore = DataChannel->buffered_amount();
+			RHICmdList.ImmediateFlush(EImmediateFlushType::WaitForOutstandingTasksOnly);
 		}
 
-		if (!DataChannel->Send(webrtc::DataBuffer(Buffer, true)))
+		inline FTexture2DRHIRef CreateTexture(uint32 Width, uint32 Height)
 		{
-			UE_LOG(PixelStreamer, Error, TEXT("Failed to send data channel packet"));
-			return;
+			// Create empty texture
+			FRHIResourceCreateInfo CreateInfo(TEXT("BlankTexture"));
+
+			FTexture2DRHIRef Texture;
+			FString RHIName = GDynamicRHI->GetName();
+
+			if (RHIName == TEXT("Vulkan"))
+			{
+				Texture = GDynamicRHI->RHICreateTexture2D(Width, Height, EPixelFormat::PF_B8G8R8A8, 1, 1, TexCreate_RenderTargetable | TexCreate_External, ERHIAccess::Present, CreateInfo);
+			}
+			else
+			{
+				Texture = GDynamicRHI->RHICreateTexture2D(Width, Height, EPixelFormat::PF_B8G8R8A8, 1, 1, TexCreate_Shared | TexCreate_RenderTargetable, ERHIAccess::CopyDest, CreateInfo);
+			}
+			return Texture;
 		}
-
-		// Increment the number of bytes transmitted
-		BytesTransmitted += BytesToTransmit;
-	}
-}
-
-inline FTexture2DRHIRef CreateTexture(uint32 Width, uint32 Height)
-{
-	// Create empty texture
-	FRHIResourceCreateInfo CreateInfo(TEXT("BlankTexture"));
-
-	FTexture2DRHIRef Texture;
-	FString RHIName = GDynamicRHI->GetName();
-
-	if (RHIName == TEXT("Vulkan"))
-	{
-		Texture = GDynamicRHI->RHICreateTexture2D(Width, Height, EPixelFormat::PF_B8G8R8A8, 1, 1, TexCreate_RenderTargetable | TexCreate_External, ERHIAccess::Present, CreateInfo);
-	}
-	else
-	{
-		Texture = GDynamicRHI->RHICreateTexture2D(Width, Height, EPixelFormat::PF_B8G8R8A8, 1, 1, TexCreate_Shared | TexCreate_RenderTargetable, ERHIAccess::CopyDest, CreateInfo);
 	}
-	return Texture;
 }
diff --git a/Engine/Plugins/Media/PixelStreaming/Source/PixelStreaming/Private/VideoEncoderH264Wrapper.cpp b/Engine/Plugins/Media/PixelStreaming/Source/PixelStreaming/Private/VideoEncoderH264Wrapper.cpp
new file mode 100644
--- /dev/null
+++ b/Engine/Plugins/Media/PixelStreaming/Source/PixelStreaming/Private/VideoEncoderH264Wrapper.cpp
@@ -0,0 +1,117 @@
+// Copyright Epic Games, Inc. All Rights Reserved.
+
+#include "VideoEncoderH264Wrapper.h"
+#include "VulkanRHIPrivate.h"
+#include "Stats.h"
+#include "VideoEncoderFactory.h"
+#include "EncoderFactory.h"
+
+UE::PixelStreaming::FVideoEncoderH264Wrapper::FVideoEncoderH264Wrapper(uint64 EncoderId, TUniquePtr<FEncoderFrameFactory> InFrameFactory, TUniquePtr<AVEncoder::FVideoEncoder> InEncoder)
+	: Id(EncoderId)
+	, FrameFactory(MoveTemp(InFrameFactory))
+	, Encoder(MoveTemp(InEncoder)) 
+{
+	checkf(Encoder, TEXT("Encoder is nullptr."))
+}
+
+UE::PixelStreaming::FVideoEncoderH264Wrapper::~FVideoEncoderH264Wrapper()
+{
+	Encoder->ClearOnEncodedPacket();
+	Encoder->Shutdown();
+}
+
+void UE::PixelStreaming::FVideoEncoderH264Wrapper::Encode(const webrtc::VideoFrame& WebRTCFrame, bool bKeyframe)
+{
+	FTexture2DRHIRef SourceTexture;
+
+	UE::PixelStreaming::FFrameBuffer* FrameBuffer = static_cast<UE::PixelStreaming::FFrameBuffer*>(WebRTCFrame.video_frame_buffer().get());
+	check(FrameBuffer->GetFrameBufferType() == UE::PixelStreaming::FFrameBufferType::Layer);
+	UE::PixelStreaming::FLayerFrameBuffer* LayerFrameBuffer = static_cast<UE::PixelStreaming::FLayerFrameBuffer*>(FrameBuffer);
+	SourceTexture = LayerFrameBuffer->GetFrame();
+
+	if (SourceTexture)
+	{
+		AVEncoder::FVideoEncoder::FLayerConfig EncoderConfig = GetCurrentConfig();
+		AVEncoder::FVideoEncoderInputFrame* EncoderInputFrame = FrameFactory->GetFrameAndSetTexture(EncoderConfig.Width, EncoderConfig.Height, SourceTexture);
+		if (EncoderInputFrame)
+		{
+			EncoderInputFrame->SetTimestampUs(WebRTCFrame.timestamp_us());
+			EncoderInputFrame->SetTimestampRTP(WebRTCFrame.timestamp());
+			EncoderInputFrame->SetFrameID(WebRTCFrame.id());
+
+			AVEncoder::FVideoEncoder::FEncodeOptions Options;
+			Options.bForceKeyFrame = bKeyframe || bForceNextKeyframe;
+			bForceNextKeyframe = false;
+
+			Encoder->Encode(EncoderInputFrame, Options);
+		}
+	}
+}
+
+AVEncoder::FVideoEncoder::FLayerConfig UE::PixelStreaming::FVideoEncoderH264Wrapper::GetCurrentConfig()
+{
+	checkf(Encoder, TEXT("Cannot request layer config when encoder is nullptr."));
+
+	// Asume user wants config for layer zero.
+	return Encoder->GetLayerConfig(0);
+}
+
+void UE::PixelStreaming::FVideoEncoderH264Wrapper::SetConfig(const AVEncoder::FVideoEncoder::FLayerConfig& NewConfig)
+{
+	checkf(Encoder, TEXT("Cannot set layer config when encoder is nullptr."));
+
+	// Assumer user wants to update layer zero.
+	if (NewConfig != GetCurrentConfig())
+	{
+		if (Encoder)
+		{
+			Encoder->UpdateLayerConfig(0, NewConfig);
+		}
+	}
+}
+
+/* ------------------ Static functions below --------------------- */
+
+void UE::PixelStreaming::FVideoEncoderH264Wrapper::OnEncodedPacket(uint64 SourceEncoderId, UE::PixelStreaming::FVideoEncoderFactory* Factory, uint32 InLayerIndex, const AVEncoder::FVideoEncoderInputFrame* InFrame, const AVEncoder::FCodecPacket& InPacket)
+{
+	webrtc::EncodedImage Image;
+
+	webrtc::RTPFragmentationHeader FragHeader;
+	//CreateH264FragmentHeader(InPacket.Data.Get(), InPacket.DataSize, FragHeader);
+
+	std::vector<webrtc::H264::NaluIndex> NALUIndices = webrtc::H264::FindNaluIndices(InPacket.Data.Get(), InPacket.DataSize);
+	FragHeader.VerifyAndAllocateFragmentationHeader(NALUIndices.size());
+	FragHeader.fragmentationVectorSize = static_cast<uint16_t>(NALUIndices.size());
+	for (int i = 0; i != NALUIndices.size(); ++i)
+	{
+		webrtc::H264::NaluIndex const& NALUIndex = NALUIndices[i];
+		FragHeader.fragmentationOffset[i] = NALUIndex.payload_start_offset;
+		FragHeader.fragmentationLength[i] = NALUIndex.payload_size;
+	}
+
+	Image.timing_.packetization_finish_ms = FTimespan::FromSeconds(FPlatformTime::Seconds()).GetTotalMilliseconds();
+	Image.timing_.encode_start_ms = InPacket.Timings.StartTs.GetTotalMilliseconds();
+	Image.timing_.encode_finish_ms = InPacket.Timings.FinishTs.GetTotalMilliseconds();
+	Image.timing_.flags = webrtc::VideoSendTiming::kTriggeredByTimer;
+
+	Image.SetEncodedData(webrtc::EncodedImageBuffer::Create(const_cast<uint8_t*>(InPacket.Data.Get()), InPacket.DataSize));
+	Image._encodedWidth = InFrame->GetWidth();
+	Image._encodedHeight = InFrame->GetHeight();
+	Image._frameType = InPacket.IsKeyFrame ? webrtc::VideoFrameType::kVideoFrameKey : webrtc::VideoFrameType::kVideoFrameDelta;
+	Image.content_type_ = webrtc::VideoContentType::UNSPECIFIED;
+	Image.qp_ = InPacket.VideoQP;
+	Image.SetSpatialIndex(InLayerIndex);
+	Image._completeFrame = true;
+	Image.rotation_ = webrtc::VideoRotation::kVideoRotation_0;
+	Image.SetTimestamp(InFrame->GetTimestampRTP());
+	Image.capture_time_ms_ = InFrame->GetTimestampUs() / 1000.0;
+
+	webrtc::CodecSpecificInfo CodecInfo;
+	CodecInfo.codecType = webrtc::VideoCodecType::kVideoCodecH264;
+	CodecInfo.codecSpecific.H264.packetization_mode = webrtc::H264PacketizationMode::NonInterleaved;
+	CodecInfo.codecSpecific.H264.temporal_idx = webrtc::kNoTemporalIdx;
+	CodecInfo.codecSpecific.H264.idr_frame = InPacket.IsKeyFrame;
+	CodecInfo.codecSpecific.H264.base_layer_sync = false;
+
+	Factory->OnEncodedImage(SourceEncoderId, Image, &CodecInfo, &FragHeader);
+}
diff --git a/Engine/Plugins/Media/PixelStreaming/Source/PixelStreaming/Private/VideoEncoderH264Wrapper.h b/Engine/Plugins/Media/PixelStreaming/Source/PixelStreaming/Private/VideoEncoderH264Wrapper.h
new file mode 100644
--- /dev/null
+++ b/Engine/Plugins/Media/PixelStreaming/Source/PixelStreaming/Private/VideoEncoderH264Wrapper.h
@@ -0,0 +1,38 @@
+// Copyright Epic Games, Inc. All Rights Reserved.
+#pragma once
+
+#include "VideoEncoder.h"
+#include "EncoderFrameFactory.h"
+
+namespace UE {
+	namespace PixelStreaming {
+		class FVideoEncoderFactory;
+
+		class FVideoEncoderH264Wrapper
+		{
+		public:
+			FVideoEncoderH264Wrapper(uint64 EncoderId, TUniquePtr<FEncoderFrameFactory> FrameFactory, TUniquePtr<AVEncoder::FVideoEncoder> Encoder);
+			~FVideoEncoderH264Wrapper();
+
+			uint64 GetId() const { return Id; }
+
+			void SetForceNextKeyframe() { bForceNextKeyframe = true; }
+
+			void Encode(const webrtc::VideoFrame& WebRTCFrame, bool bKeyframe);
+
+			AVEncoder::FVideoEncoder::FLayerConfig GetCurrentConfig();
+			void SetConfig(const AVEncoder::FVideoEncoder::FLayerConfig& NewConfig);
+
+			static void OnEncodedPacket(uint64 SourceEncoderId, 
+				UE::PixelStreaming::FVideoEncoderFactory* Factory, 
+				uint32 InLayerIndex, const AVEncoder::FVideoEncoderInputFrame* InFrame, 
+				const AVEncoder::FCodecPacket& InPacket);
+
+		private:
+			uint64 Id;
+			TUniquePtr<FEncoderFrameFactory> FrameFactory;
+			TUniquePtr<AVEncoder::FVideoEncoder> Encoder;
+			bool bForceNextKeyframe = false;
+		};
+	}
+}
diff --git a/Engine/Plugins/Media/PixelStreaming/Source/PixelStreaming/Private/VideoEncoderRTC.cpp b/Engine/Plugins/Media/PixelStreaming/Source/PixelStreaming/Private/VideoEncoderRTC.cpp
new file mode 100644
--- /dev/null
+++ b/Engine/Plugins/Media/PixelStreaming/Source/PixelStreaming/Private/VideoEncoderRTC.cpp
@@ -0,0 +1,206 @@
+// Copyright Epic Games, Inc. All Rights Reserved.
+
+#include "VideoEncoderRTC.h"
+#include "EncoderFactory.h"
+#include "VideoEncoderFactory.h"
+#include "FrameBuffer.h"
+#include "PlayerSession.h"
+#include "Stats.h"
+#include "UnrealEngine.h"
+#include "Misc/Paths.h"
+#include "Settings.h"
+#include "PixelStreamingPlayerId.h"
+#include "Async/Async.h"
+#include "VideoEncoderH264Wrapper.h"
+#include "HAL/PlatformFileManager.h"
+
+UE::PixelStreaming::FVideoEncoderRTC::FVideoEncoderRTC(UE::PixelStreaming::FVideoEncoderFactory& InFactory)
+	: Factory(InFactory)
+{
+}
+
+UE::PixelStreaming::FVideoEncoderRTC::~FVideoEncoderRTC()
+{
+	Factory.ReleaseVideoEncoder(this);
+}
+
+int UE::PixelStreaming::FVideoEncoderRTC::InitEncode(webrtc::VideoCodec const* codec_settings, VideoEncoder::Settings const& settings)
+{
+	checkf(AVEncoder::FVideoEncoderFactory::Get().IsSetup(), TEXT("FVideoEncoderFactory not setup"));
+
+	// Try to create an encoder if it doesnt already.
+	TOptional<FVideoEncoderFactory::FHardwareEncoderId> EncoderIdOptional = Factory.GetOrCreateHardwareEncoder(codec_settings->width, codec_settings->height, codec_settings->maxBitrate, codec_settings->startBitrate, codec_settings->maxFramerate);
+	
+	if(EncoderIdOptional.IsSet())
+	{
+		HardwareEncoderId = EncoderIdOptional.GetValue();
+		UpdateConfig();
+		WebRtcProposedTargetBitrate = codec_settings->startBitrate;
+		return WEBRTC_VIDEO_CODEC_OK;
+	}
+	else
+	{
+		return WEBRTC_VIDEO_CODEC_ERROR;
+	}
+}
+
+int32 UE::PixelStreaming::FVideoEncoderRTC::RegisterEncodeCompleteCallback(webrtc::EncodedImageCallback* callback)
+{
+	OnEncodedImageCallback = callback;
+	return WEBRTC_VIDEO_CODEC_OK;
+}
+
+int32 UE::PixelStreaming::FVideoEncoderRTC::Release()
+{
+	OnEncodedImageCallback = nullptr;
+	return WEBRTC_VIDEO_CODEC_OK;
+}
+
+int32 UE::PixelStreaming::FVideoEncoderRTC::Encode(webrtc::VideoFrame const& frame, std::vector<webrtc::VideoFrameType> const* frame_types)
+{
+	UE::PixelStreaming::FFrameBuffer* FrameBuffer = static_cast<UE::PixelStreaming::FFrameBuffer*>(frame.video_frame_buffer().get());
+
+	// If initialize frame is passed, this is a dummy frame so WebRTC is happy frames are passing through
+	// This is used so that an encoder can be active as far as WebRTC is concerned by simply have frames transmitted by some other encoder on its behalf.
+	if(FrameBuffer->GetFrameBufferType() == UE::PixelStreaming::FFrameBufferType::Initialize)
+	{
+		return WEBRTC_VIDEO_CODEC_OK;
+	}
+
+	bool bKeyframe = false;
+	if ((frame_types && (*frame_types)[0] == webrtc::VideoFrameType::kVideoFrameKey))
+	{
+		bKeyframe = true;
+	}
+
+	UE::PixelStreaming::FVideoEncoderH264Wrapper* Encoder = Factory.GetHardwareEncoder(HardwareEncoderId);
+	if (Encoder == nullptr)
+	{
+		return WEBRTC_VIDEO_CODEC_ERROR;
+	}
+
+	UpdateConfig();
+
+	Encoder->Encode(frame, bKeyframe);
+
+	return WEBRTC_VIDEO_CODEC_OK;
+}
+
+// Pass rate control parameters from WebRTC to our encoder
+// This is how WebRTC can control the bitrate/framerate of the encoder.
+void UE::PixelStreaming::FVideoEncoderRTC::SetRates(RateControlParameters const& parameters)
+{
+	PendingRateChange.Emplace(parameters);
+}
+
+webrtc::VideoEncoder::EncoderInfo UE::PixelStreaming::FVideoEncoderRTC::GetEncoderInfo() const
+{
+	VideoEncoder::EncoderInfo info;
+	info.supports_native_handle = true;
+	info.is_hardware_accelerated = true;
+	info.has_internal_source = false;
+	info.supports_simulcast = false;
+	info.implementation_name = TCHAR_TO_UTF8(*FString::Printf(TEXT("PIXEL_STREAMING_HW_ENCODER_%s"), GDynamicRHI->GetName()));
+
+	const int LowQP = UE::PixelStreaming::Settings::CVarPixelStreamingWebRTCLowQpThreshold.GetValueOnAnyThread();
+	const int HighQP = UE::PixelStreaming::Settings::CVarPixelStreamingWebRTCHighQpThreshold.GetValueOnAnyThread();
+	info.scaling_settings = VideoEncoder::ScalingSettings(LowQP, HighQP);
+
+	// if true means HW encoder must be perfect and drop frames itself etc
+	info.has_trusted_rate_controller = false;
+
+	return info;
+}
+
+void UE::PixelStreaming::FVideoEncoderRTC::UpdateConfig()
+{
+	if (UE::PixelStreaming::FVideoEncoderH264Wrapper* Encoder = Factory.GetHardwareEncoder(HardwareEncoderId))
+	{
+		AVEncoder::FVideoEncoder::FLayerConfig NewConfig = Encoder->GetCurrentConfig();
+
+		if (PendingRateChange.IsSet())
+		{
+			const RateControlParameters& RateChangeParams = PendingRateChange.GetValue();
+
+			NewConfig.MaxFramerate = RateChangeParams.framerate_fps;
+
+			// We store what WebRTC wants as the bitrate, even if we are overriding it, so we can restore back to it when user stops using CVar.
+			WebRtcProposedTargetBitrate = RateChangeParams.bitrate.get_sum_kbps() * 1000;
+
+			// Clear the rate change request
+			PendingRateChange.Reset();
+		}
+
+		// TODO how to handle res changes?
+
+		NewConfig = CreateEncoderConfigFromCVars(NewConfig);
+
+		Encoder->SetConfig(NewConfig);
+	}
+}
+
+AVEncoder::FVideoEncoder::FLayerConfig UE::PixelStreaming::FVideoEncoderRTC::CreateEncoderConfigFromCVars(AVEncoder::FVideoEncoder::FLayerConfig InEncoderConfig) const
+{
+	// Change encoder settings through CVars
+	const int32 MaxBitrateCVar = UE::PixelStreaming::Settings::CVarPixelStreamingEncoderMaxBitrate.GetValueOnAnyThread();
+	const int32 TargetBitrateCVar = UE::PixelStreaming::Settings::CVarPixelStreamingEncoderTargetBitrate.GetValueOnAnyThread();
+	const int32 MinQPCVar = UE::PixelStreaming::Settings::CVarPixelStreamingEncoderMinQP.GetValueOnAnyThread();
+	const int32 MaxQPCVar = UE::PixelStreaming::Settings::CVarPixelStreamingEncoderMaxQP.GetValueOnAnyThread();
+	const AVEncoder::FVideoEncoder::RateControlMode RateControlCVar = UE::PixelStreaming::Settings::GetRateControlCVar();
+	const AVEncoder::FVideoEncoder::MultipassMode MultiPassCVar = UE::PixelStreaming::Settings::GetMultipassCVar();
+	const bool bFillerDataCVar = UE::PixelStreaming::Settings::CVarPixelStreamingEnableFillerData.GetValueOnAnyThread();
+	const AVEncoder::FVideoEncoder::H264Profile H264Profile = UE::PixelStreaming::Settings::GetH264Profile();
+
+	InEncoderConfig.MaxBitrate = MaxBitrateCVar > -1 ? MaxBitrateCVar : InEncoderConfig.MaxBitrate;
+	InEncoderConfig.TargetBitrate = TargetBitrateCVar > -1 ? TargetBitrateCVar : WebRtcProposedTargetBitrate;
+	InEncoderConfig.QPMin = MinQPCVar;
+	InEncoderConfig.QPMax = MaxQPCVar;
+	InEncoderConfig.RateControlMode = RateControlCVar;
+	InEncoderConfig.MultipassMode = MultiPassCVar;
+	InEncoderConfig.FillData = bFillerDataCVar;
+	InEncoderConfig.H264Profile = H264Profile;
+
+	return InEncoderConfig;
+}
+
+void UE::PixelStreaming::FVideoEncoderRTC::SendEncodedImage(uint64 SourceEncoderId, webrtc::EncodedImage const& encoded_image, webrtc::CodecSpecificInfo const* codec_specific_info, webrtc::RTPFragmentationHeader const* fragmentation)
+{
+	if (SourceEncoderId != HardwareEncoderId)
+	{
+		return;
+	}
+
+	if (FirstKeyframeCountdown > 0)
+	{
+		// ideally we want to make the first frame of new peers a keyframe but we
+		// dont know when webrtc will decide to start sending out frames. this is
+		// the next best option. its a count beause when it was the very next frame
+		// it still didnt seem to work. delaying it a few frames seems to have worked.
+		--FirstKeyframeCountdown;
+		if (FirstKeyframeCountdown == 0)
+		{
+			Factory.ForceKeyFrame();
+		}
+	}
+
+	// Dump H264 frames to file for debugging if CVar is turned on.
+	if (UE::PixelStreaming::Settings::CVarPixelStreamingDebugDumpFrame.GetValueOnAnyThread())
+	{
+		static IFileHandle* FileHandle = nullptr;
+		if (!FileHandle)
+		{
+			IPlatformFile& PlatformFile = FPlatformFileManager::Get().GetPlatformFile();
+			FString TempFilePath = FPaths::CreateTempFilename(*FPaths::ProjectSavedDir(), TEXT("encoded_frame"), TEXT(".h264"));
+			FileHandle = PlatformFile.OpenWrite(*TempFilePath);
+			check(FileHandle);
+		}
+
+		FileHandle->Write(encoded_image.data(), encoded_image.size());
+		FileHandle->Flush();
+	}
+
+	if (OnEncodedImageCallback)
+	{
+		OnEncodedImageCallback->OnEncodedImage(encoded_image, codec_specific_info, fragmentation);
+	}
+}
diff --git a/Engine/Plugins/Media/PixelStreaming/Source/PixelStreaming/Private/VideoEncoderRTC.h b/Engine/Plugins/Media/PixelStreaming/Source/PixelStreaming/Private/VideoEncoderRTC.h
new file mode 100644
--- /dev/null
+++ b/Engine/Plugins/Media/PixelStreaming/Source/PixelStreaming/Private/VideoEncoderRTC.h
@@ -0,0 +1,61 @@
+// Copyright Epic Games, Inc. All Rights Reserved.
+#pragma once
+
+#include "PixelStreamingPlayerId.h"
+#include "HAL/ThreadSafeBool.h"
+#include "WebRTCIncludes.h"
+#include "VideoEncoder.h"
+#include "Templates/SharedPointer.h"
+#include "Misc/Optional.h"
+
+namespace UE {
+	namespace PixelStreaming {
+		class FVideoEncoderFactory;
+
+		// Implementation that is a WebRTC video encoder that allows us tie to our actually underlying non-WebRTC video encoder.
+		class FVideoEncoderRTC : public webrtc::VideoEncoder
+		{
+		public:
+			FVideoEncoderRTC(FVideoEncoderFactory& InFactory);
+			virtual ~FVideoEncoderRTC() override;
+
+			// WebRTC Interface
+			virtual int InitEncode(webrtc::VideoCodec const* codec_settings, webrtc::VideoEncoder::Settings const& settings) override;
+			virtual int32 RegisterEncodeCompleteCallback(webrtc::EncodedImageCallback* callback) override;
+			virtual int32 Release() override;
+			virtual int32 Encode(webrtc::VideoFrame const& frame, std::vector<webrtc::VideoFrameType> const* frame_types) override;
+			virtual void SetRates(RateControlParameters const& parameters) override;
+			virtual webrtc::VideoEncoder::EncoderInfo GetEncoderInfo() const override;
+
+			// Note: These funcs can also be overriden but are not pure virtual
+			// virtual void SetFecControllerOverride(FecControllerOverride* fec_controller_override) override;
+			// virtual void OnPacketLossRateUpdate(float packet_loss_rate) override;
+			// virtual void OnRttUpdate(int64_t rtt_ms) override;
+			// virtual void OnLossNotification(const LossNotification& loss_notification) override;
+			// End WebRTC Interface.
+
+			void SendEncodedImage(uint64 SourceEncoderId, webrtc::EncodedImage const& encoded_image, webrtc::CodecSpecificInfo const* codec_specific_info, webrtc::RTPFragmentationHeader const* fragmentation);
+
+		private:
+			void UpdateConfig();
+			AVEncoder::FVideoEncoder::FLayerConfig CreateEncoderConfigFromCVars(AVEncoder::FVideoEncoder::FLayerConfig BaseEncoderConfig) const;
+
+			FVideoEncoderFactory& Factory;
+
+			uint64 HardwareEncoderId;
+
+			// We store this so we can restore back to it if the user decides to use then stop using the PixelStreaming.Encoder.TargetBitrate CVar.
+			int32 WebRtcProposedTargetBitrate = 5000000;
+
+			webrtc::EncodedImageCallback* OnEncodedImageCallback = nullptr;
+
+			// WebRTC may request a bitrate/framerate change using SetRates(), we only respect this if this encoder is actually encoding
+			// so we use this optional object to store a rate change and act upon it when this encoder does its next call to Encode().
+			TOptional<RateControlParameters> PendingRateChange;
+
+			// Used to send an initial keyframe
+			// see notes in SendEncodedImage implementation
+			int FirstKeyframeCountdown = 2;
+		};
+	}
+}
diff --git a/Engine/Plugins/Media/PixelStreaming/Source/PixelStreaming/Private/VideoSource.cpp b/Engine/Plugins/Media/PixelStreaming/Source/PixelStreaming/Private/VideoSource.cpp
new file mode 100644
--- /dev/null
+++ b/Engine/Plugins/Media/PixelStreaming/Source/PixelStreaming/Private/VideoSource.cpp
@@ -0,0 +1,174 @@
+// Copyright Epic Games, Inc. All Rights Reserved.
+
+#include "VideoSource.h"
+#include "FixedFPSPump.h"
+#include "FrameBuffer.h"
+#include "Stats.h"
+#include "PixelStreamingStatNames.h"
+
+UE::PixelStreaming::FVideoSourceBase::FVideoSourceBase(FPixelStreamingPlayerId InPlayerId)
+	: PlayerId(InPlayerId)
+{
+}
+
+void UE::PixelStreaming::FVideoSourceBase::AddRef() const 
+{ 
+	RefCount.IncRef(); 
+}
+
+rtc::RefCountReleaseStatus UE::PixelStreaming::FVideoSourceBase::Release() const 
+{
+	const rtc::RefCountReleaseStatus Status = RefCount.DecRef();
+	if (Status == rtc::RefCountReleaseStatus::kDroppedLastRef)
+	{
+		delete this;
+	}
+	return Status;
+}
+
+void UE::PixelStreaming::FVideoSourceBase::Initialize()
+{
+	CurrentState = webrtc::MediaSourceInterface::SourceState::kInitializing;
+	UE::PixelStreaming::FFixedFPSPump::Get().RegisterVideoSource(PlayerId, this);
+}
+
+UE::PixelStreaming::FVideoSourceBase::~FVideoSourceBase()
+{
+	CurrentState = webrtc::MediaSourceInterface::SourceState::kEnded;
+	UE::PixelStreaming::FFixedFPSPump::Get().UnregisterVideoSource(PlayerId);
+}
+
+void UE::PixelStreaming::FVideoSourceBase::OnPump(int32 FrameId)
+{
+
+	// `OnPump` prepares a frame to be passed to WebRTC's `OnFrame`.
+	// we dont call `OnFrame` outside this method so we can call AdaptCaptureFrame and possibly
+	// drop a frame if requested.
+
+	CurrentState = webrtc::MediaSourceInterface::SourceState::kLive;
+
+	webrtc::VideoFrame Frame = CreateFrame(FrameId);
+
+	// might want to allow source frames to be scaled here?
+	const FIntPoint Resolution(Frame.width(), Frame.height());
+	if (!AdaptCaptureFrame(Frame.timestamp_us(), Resolution))
+	{
+		return;
+	}
+
+	OnFrame(Frame);
+}
+
+bool UE::PixelStreaming::FVideoSourceBase::AdaptCaptureFrame(const int64 TimestampUs, FIntPoint Resolution)
+{
+	int outWidth, outHeight, cropWidth, cropHeight, cropX, cropY;
+	if (!AdaptFrame(Resolution.X, Resolution.Y, TimestampUs, &outWidth, &outHeight, &cropWidth, &cropHeight, &cropX, &cropY))
+	{
+		return false;
+	}
+
+	return true;
+}
+
+/*
+* ------------------ FPlayerVideSource ------------------
+*/
+
+UE::PixelStreaming::FVideoSourceP2P::FVideoSourceP2P(FPixelStreamingPlayerId InPlayerId, UE::PixelStreaming::IPixelStreamingSessions* InSessions)
+	: UE::PixelStreaming::FVideoSourceBase(InPlayerId)
+	, Sessions(InSessions)
+	, TextureSource(MakeShared<FTextureSourceBackBuffer>(1.0))
+{
+	
+}
+
+bool UE::PixelStreaming::FVideoSourceP2P::IsReadyForPump() const
+{
+	return TextureSource.IsValid() && TextureSource->IsAvailable();
+}
+
+webrtc::VideoFrame UE::PixelStreaming::FVideoSourceP2P::CreateFrame(int32 FrameId)
+{
+	bool bQualityController = Sessions->IsQualityController(PlayerId);
+	TextureSource->SetEnabled(bQualityController);
+
+	const int64 TimestampUs = rtc::TimeMicros();
+
+	// Based on quality control we either send a frame buffer with a texture or an empty frame
+	// buffer to indicate that this frame should be skipped by the encoder as this peer is having its
+	// frames transmitted by some other peer (the quality controlling peer).
+	rtc::scoped_refptr<webrtc::VideoFrameBuffer> FrameBuffer;
+
+	if (bQualityController)
+	{
+		FrameBuffer = new rtc::RefCountedObject<UE::PixelStreaming::FLayerFrameBuffer>(TextureSource);
+	}
+	else
+	{
+		FrameBuffer = new rtc::RefCountedObject<UE::PixelStreaming::FInitializeFrameBuffer>(TextureSource);
+	}
+
+	webrtc::VideoFrame Frame = webrtc::VideoFrame::Builder()
+								   .set_video_frame_buffer(FrameBuffer)
+								   .set_timestamp_us(TimestampUs)
+								   .set_rotation(webrtc::VideoRotation::kVideoRotation_0)
+								   .set_id(FrameId)
+								   .build();
+
+	return Frame;
+}
+
+/*
+* ------------------ FVideoSourceSFU ------------------
+*/
+
+UE::PixelStreaming::FVideoSourceSFU::FVideoSourceSFU()
+	: FVideoSourceBase(SFU_PLAYER_ID)
+{
+	// Make a copy of simulcast settings and sort them based on scaling.
+	using FLayer = UE::PixelStreaming::Settings::FSimulcastParameters::FLayer;
+
+	TArray<FLayer*> SortedLayers;
+	for (FLayer& Layer : UE::PixelStreaming::Settings::SimulcastParameters.Layers)
+	{
+		SortedLayers.Add(&Layer);
+	}
+	SortedLayers.Sort([](const FLayer& LayerA, const FLayer& LayerB) { return LayerA.Scaling > LayerB.Scaling; });
+	
+	for (FLayer* SimulcastLayer : SortedLayers)
+	{
+		const float Scale = 1.0f / SimulcastLayer->Scaling;
+		TSharedPtr<FTextureSourceBackBuffer> TextureSource = MakeShared<FTextureSourceBackBuffer>(Scale);
+		TextureSource->SetEnabled(true);
+		LayerTextures.Add(TextureSource);
+	}
+}
+
+bool UE::PixelStreaming::FVideoSourceSFU::IsReadyForPump() const
+{
+	// Check all texture sources are ready.
+	int NumReady = 0;
+
+	for(int i = 0; i < LayerTextures.Num(); i++)
+	{
+		bool bIsReady = LayerTextures[i].IsValid() && LayerTextures[i]->IsAvailable();
+		NumReady += bIsReady ? 1 : 0;
+	}
+
+	return NumReady == LayerTextures.Num();
+}
+
+webrtc::VideoFrame UE::PixelStreaming::FVideoSourceSFU::CreateFrame(int32 FrameId)
+{
+	rtc::scoped_refptr<webrtc::VideoFrameBuffer> FrameBuffer = new rtc::RefCountedObject<FSimulcastFrameBuffer>(LayerTextures);
+	const int64 TimestampUs = rtc::TimeMicros();
+
+	webrtc::VideoFrame Frame = webrtc::VideoFrame::Builder()
+								.set_video_frame_buffer(FrameBuffer)
+								.set_timestamp_us(TimestampUs)
+								.set_rotation(webrtc::VideoRotation::kVideoRotation_0)
+								.set_id(FrameId)
+								.build();
+
+	return Frame;
+}
\ No newline at end of file
diff --git a/Engine/Plugins/Media/PixelStreaming/Source/PixelStreaming/Private/VideoSource.h b/Engine/Plugins/Media/PixelStreaming/Source/PixelStreaming/Private/VideoSource.h
new file mode 100644
--- /dev/null
+++ b/Engine/Plugins/Media/PixelStreaming/Source/PixelStreaming/Private/VideoSource.h
@@ -0,0 +1,92 @@
+// Copyright Epic Games, Inc. All Rights Reserved.
+
+#pragma once
+
+#include "WebRTCIncludes.h"
+#include "PixelStreamingPlayerId.h"
+#include "FixedFPSPump.h"
+#include "TextureSource.h"
+#include "IPixelStreamingSessions.h"
+
+namespace UE {
+	namespace PixelStreaming {
+		class FVideoSourceFactory;
+
+		/*
+		* Base class for pumped Pixel Streaming video sources. Video sources are WebRTC video sources that populate WebRTC tracks
+		* and pass WebRTC video frames to `OnFrame`, which eventually gets passed to a WebRTC video encoder, encoded, and transmitted.
+		*/
+		class FVideoSourceBase : public rtc::AdaptedVideoTrackSource, public IPumpedVideoSource
+		{
+		public:
+			FVideoSourceBase(FPixelStreamingPlayerId InPlayerId);
+			virtual ~FVideoSourceBase();
+			FPixelStreamingPlayerId GetPlayerId() const { return PlayerId; }
+			virtual void Initialize();
+
+			/* Begin rtc::AdaptedVideoTrackSource overrides */
+			virtual webrtc::MediaSourceInterface::SourceState state() const override { return CurrentState; }
+			virtual bool remote() const override { return false; }
+			virtual bool is_screencast() const override { return false;}
+			virtual absl::optional<bool> needs_denoising() const override { return false; }
+			/* End rtc::AdaptedVideoTrackSource overrides */
+
+			/* Begin IPumpedVideoSource */
+			virtual void OnPump(int32 FrameId) override;
+			virtual bool IsReadyForPump() const = 0;
+			virtual void AddRef() const override;
+			virtual rtc::RefCountReleaseStatus Release() const override;
+			virtual bool HasOneRef() const { return RefCount.HasOneRef(); }
+			/* End IPumpedVideoSource */
+
+		protected:
+			virtual bool AdaptCaptureFrame(const int64 TimestampUs, FIntPoint Resolution);
+			virtual webrtc::VideoFrame CreateFrame(int32 FrameId) = 0;
+
+		private:
+			mutable webrtc::webrtc_impl::RefCounter RefCount{0};
+
+		protected:
+			webrtc::MediaSourceInterface::SourceState CurrentState;
+			FPixelStreamingPlayerId PlayerId;
+		};
+
+		/*
+		* A video source for P2P peers.
+		*/
+		class FVideoSourceP2P : public FVideoSourceBase
+		{
+			public:
+				FVideoSourceP2P(FPixelStreamingPlayerId InPlayerId, IPixelStreamingSessions* InSessions);
+
+			protected:
+				IPixelStreamingSessions* Sessions;
+				TSharedPtr<ITextureSource> TextureSource;
+
+			protected:
+				/* Begin FVideoSourceBase */
+				virtual webrtc::VideoFrame CreateFrame(int32 FrameId) override;
+				virtual bool IsReadyForPump() const override;
+				/* End FVideoSourceBase */
+		};
+
+		/*
+		* A video source for the SFU.
+		*/
+		class FVideoSourceSFU : public FVideoSourceBase
+		{
+			public:
+				FVideoSourceSFU();
+
+			protected:
+				TArray<TSharedPtr<ITextureSource>> LayerTextures;
+
+			protected:
+				/* Begin FVideoSourceBase */
+				virtual webrtc::VideoFrame CreateFrame(int32 FrameId) override;
+				virtual bool IsReadyForPump() const override;
+				/* End FVideoSourceBase */
+		};
+
+	}
+}
diff --git a/Engine/Plugins/Media/PixelStreaming/Source/PixelStreaming/Private/WebRTCLogging.cpp b/Engine/Plugins/Media/PixelStreaming/Source/PixelStreaming/Private/WebRTCLogging.cpp
--- a/Engine/Plugins/Media/PixelStreaming/Source/PixelStreaming/Private/WebRTCLogging.cpp
+++ b/Engine/Plugins/Media/PixelStreaming/Source/PixelStreaming/Private/WebRTCLogging.cpp
@@ -2,13 +2,11 @@
 
 #include "WebRTCLogging.h"
 
-#include "Utils.h"
-
 #include "Logging/LogMacros.h"
 #include "Misc/AssertionMacros.h"
 
-DECLARE_LOG_CATEGORY_EXTERN(PixelStreamingWebRTC, Log, All);
-DEFINE_LOG_CATEGORY(PixelStreamingWebRTC);
+DECLARE_LOG_CATEGORY_EXTERN(LogPixelStreamingWebRTC, Log, All);
+DEFINE_LOG_CATEGORY(LogPixelStreamingWebRTC);
 
 /**
 	* Receives logging from WebRTC internals, and writes it to a log file
@@ -19,7 +17,7 @@ class FWebRtcLogsRedirector : public rtc::LogSink
 public:
 	explicit FWebRtcLogsRedirector(rtc::LoggingSeverity Verbosity)
 	{
-		UE_LOG(PixelStreamingWebRTC, Verbose, TEXT("Starting WebRTC logging"));
+		UE_LOG(LogPixelStreamingWebRTC, Verbose, TEXT("Starting WebRTC logging"));
 
 		rtc::LogMessage::AddLogToStream(this, Verbosity);
 
@@ -31,7 +29,7 @@ public:
 
 	~FWebRtcLogsRedirector() override
 	{
-		UE_LOG(PixelStreamingWebRTC, Log, TEXT("Stopping WebRTC logging"));
+		UE_LOG(LogPixelStreamingWebRTC, Log, TEXT("Stopping WebRTC logging"));
 
 		rtc::LogMessage::RemoveLogToStream(this);
 	}
@@ -49,7 +47,7 @@ private:
 			ELogVerbosity::Error
 		};
 
-		if (PixelStreamingWebRTC.IsSuppressed(RtcToUnrealLogCategoryMap[severity]))
+		if (LogPixelStreamingWebRTC.IsSuppressed(RtcToUnrealLogCategoryMap[severity]))
 		{
 			return;
 		}
@@ -71,22 +69,22 @@ private:
 		{
 			case rtc::LS_VERBOSE:
 			{
-				UE_LOG(PixelStreamingWebRTC, Verbose, TEXT("%s"), *Msg);
+				UE_LOG(LogPixelStreamingWebRTC, Verbose, TEXT("%s"), *Msg);
 				break;
 			}
 			case rtc::LS_INFO:
 			{
-				UE_LOG(PixelStreamingWebRTC, Log, TEXT("%s"), *Msg);
+				UE_LOG(LogPixelStreamingWebRTC, Log, TEXT("%s"), *Msg);
 				break;
 			}
 			case rtc::LS_WARNING:
 			{
-				UE_LOG(PixelStreamingWebRTC, Warning, TEXT("%s"), *Msg);
+				UE_LOG(LogPixelStreamingWebRTC, Warning, TEXT("%s"), *Msg);
 				break;
 			}
 			case rtc::LS_ERROR:
 			{
-				UE_LOG(PixelStreamingWebRTC, Error, TEXT("%s"), *Msg);
+				UE_LOG(LogPixelStreamingWebRTC, Error, TEXT("%s"), *Msg);
 				break;
 			}
 		}
@@ -96,7 +94,7 @@ private:
 	void OnLogMessage(const std::string& message) override
 	{
 		//unimplemented();
-		UE_LOG(PixelStreamingWebRTC, Verbose, TEXT("%s"), *FString(message.c_str()));
+		UE_LOG(LogPixelStreamingWebRTC, Verbose, TEXT("%s"), *FString(message.c_str()));
 	}
 };
 
