Merging Pixel Streaming refactor using //UE5/Dev-TensorWorks-to-//UE5/Release-5.0

diff --git a/Engine/Plugins/Media/PixelStreaming/Source/PixelStreaming/Public/EditorPixelStreamingSettings.h b/Engine/Plugins/Media/PixelStreaming/Source/PixelStreaming/Public/EditorPixelStreamingSettings.h
--- a/Engine/Plugins/Media/PixelStreaming/Source/PixelStreaming/Public/EditorPixelStreamingSettings.h
+++ b/Engine/Plugins/Media/PixelStreaming/Source/PixelStreaming/Public/EditorPixelStreamingSettings.h
@@ -20,9 +20,9 @@ public:
 	 * click and interact with UI elements.
 	 */
 	UPROPERTY(config, EditAnywhere, Category = PixelStreaming)
-	FSoftClassPath PixelStreamerDefaultCursorClassName;
+	FSoftClassPath DefaultCursorClassName;
 	UPROPERTY(config, EditAnywhere, Category = PixelStreaming)
-	FSoftClassPath PixelStreamerTextEditBeamCursorClassName;
+	FSoftClassPath TextEditBeamCursorClassName;
 
 	/**
 	 * Pixel Streaming can have a server-side cursor (where the cursor itself
@@ -31,7 +31,7 @@ public:
 	 * cursor invisible.
 	 */
 	UPROPERTY(config, EditAnywhere, Category = PixelStreaming)
-	FSoftClassPath PixelStreamerHiddenCursorClassName;
+	FSoftClassPath HiddenCursorClassName;
 
 	/**
 	 * Pixel Streaming may be running on a machine which has no physical mouse
@@ -39,7 +39,7 @@ public:
 	 * fake the presence of a mouse.
 	 */
 	UPROPERTY(config, EditAnywhere, Category = PixelStreaming)
-	bool bPixelStreamerMouseAlwaysAttached = true;
+	bool bMouseAlwaysAttached = true;
 
 	// Begin UDeveloperSettings Interface
 	virtual FName GetCategoryName() const override;
diff --git a/Engine/Plugins/Media/PixelStreaming/Source/PixelStreaming/Public/IPixelStreamingModule.h b/Engine/Plugins/Media/PixelStreaming/Source/PixelStreaming/Public/IPixelStreamingModule.h
--- a/Engine/Plugins/Media/PixelStreaming/Source/PixelStreaming/Public/IPixelStreamingModule.h
+++ b/Engine/Plugins/Media/PixelStreaming/Source/PixelStreaming/Public/IPixelStreamingModule.h
@@ -7,13 +7,14 @@
 #include "Modules/ModuleManager.h"
 #include "IInputDeviceModule.h"
 #include "Templates/SharedPointer.h"
-#include "PlayerId.h"
+#include "PixelStreamingPlayerId.h"
 #include "IPixelStreamingAudioSink.h"
 #include "Templates/SharedPointer.h"
-#include "IPixelStreamingStatsConsumer.h"
+
+#include "IInputDevice.h"
 
 class UTexture2D;
-class UPixelStreamerInputComponent;
+class UPixelStreamingInput;
 
 /**
 * The public interface to this module
@@ -64,7 +65,7 @@ public:
 	 * is that of the underlying shared pointer.
 	 * @return A reference to the input device.
 	 */
-	virtual class FInputDevice& GetInputDevice() = 0;
+	virtual IInputDevice& GetInputDevice() = 0;
 
 	/**
 	 * Add any player config JSON to the given object which relates to
@@ -107,7 +108,7 @@ public:
 	/**
 	 * Get the audio sink associated with a specific peer/player.
 	 */
-	virtual IPixelStreamingAudioSink* GetPeerAudioSink(FPlayerId PlayerId) = 0;
+	virtual IPixelStreamingAudioSink* GetPeerAudioSink(FPixelStreamingPlayerId PlayerId) = 0;
 
 	/**
 	 * Get an audio sink that has no peers/players listening to it.
@@ -118,29 +119,17 @@ public:
 	 * Tell the input device about a new pixel streaming input component.
 	 * @param InInputComponent - The new pixel streaming input component.
 	 */
-	virtual void AddInputComponent(UPixelStreamerInputComponent* InInputComponent) = 0;
+	virtual void AddInputComponent(UPixelStreamingInput* InInputComponent) = 0;
 
 	/*
 	 * Tell the input device that a pixel streaming input component is no longer
 	 * relevant.
 	 * @param InInputComponent - The pixel streaming input component which is no longer relevant.
 	 */
-	virtual void RemoveInputComponent(UPixelStreamerInputComponent* InInputComponent) = 0;
+	virtual void RemoveInputComponent(UPixelStreamingInput* InInputComponent) = 0;
 
 	/*
 	 * Get the input components currently attached to Pixel Streaming.
 	 */
-	virtual const TArray<UPixelStreamerInputComponent*> GetInputComponents() = 0;
-
-	/*
-	 * Add a callback that gets fired any time a stat is updated during Pixel Streaming.
-	 * @param Callback - The callback that is fired when the stat changes.
-	 */
-	virtual void AddAnyStatChangedCallback(TWeakPtr<IPixelStreamingStatsConsumer> Callback) = 0;
-
-	/*
-	 * Remove the callback for any stat changed.
-	 * @param Callback - The callback to remove.
-	 */
-	virtual void RemoveAnyStatChangedCallback(TWeakPtr<IPixelStreamingStatsConsumer> Callback) = 0;
+	virtual const TArray<UPixelStreamingInput*> GetInputComponents() = 0;
 };
diff --git a/Engine/Plugins/Media/PixelStreaming/Source/PixelStreaming/Public/IPixelStreamingStatsConsumer.h b/Engine/Plugins/Media/PixelStreaming/Source/PixelStreaming/Public/IPixelStreamingStatsConsumer.h
--- a/Engine/Plugins/Media/PixelStreaming/Source/PixelStreaming/Public/IPixelStreamingStatsConsumer.h
+++ b/Engine/Plugins/Media/PixelStreaming/Source/PixelStreaming/Public/IPixelStreamingStatsConsumer.h
@@ -3,7 +3,7 @@
 #pragma once
 
 #include "CoreMinimal.h"
-#include "PlayerId.h"
+#include "PixelStreamingPlayerId.h"
 
 // Interface for consuming stats from Pixel Streaming, simply implement and do as you wish with the stats.
 // Stats will be pumped to consumers no faster than game thread tick rate, but commonly, 1s for many stats.
@@ -12,5 +12,5 @@ class PIXELSTREAMING_API IPixelStreamingStatsConsumer
 public:
 	IPixelStreamingStatsConsumer() {}
 	virtual ~IPixelStreamingStatsConsumer() {}
-	virtual void ConsumeStat(FPlayerId PlayerId, FName StatName, float StatValue) = 0;
+	virtual void ConsumeStat(FPixelStreamingPlayerId PlayerId, FName StatName, float StatValue) = 0;
 };
\ No newline at end of file
diff --git a/Engine/Plugins/Media/PixelStreaming/Source/PixelStreaming/Public/PixelStreamingAudioComponent.h b/Engine/Plugins/Media/PixelStreaming/Source/PixelStreaming/Public/PixelStreamingAudioComponent.h
--- a/Engine/Plugins/Media/PixelStreaming/Source/PixelStreaming/Public/PixelStreamingAudioComponent.h
+++ b/Engine/Plugins/Media/PixelStreaming/Source/PixelStreaming/Public/PixelStreamingAudioComponent.h
@@ -3,7 +3,7 @@
 #pragma once
 
 #include "CoreMinimal.h"
-#include "PlayerId.h"
+#include "PixelStreamingPlayerId.h"
 #include "Components/SynthComponent.h"
 #include "IPixelStreamingAudioConsumer.h"
 #include "IPixelStreamingAudioSink.h"
@@ -13,7 +13,7 @@
  * Allows in-engine playback of incoming WebRTC audio from a particular Pixel Streaming player/peer using their mic in the browser.
  * Note: Each audio component associates itself with a particular Pixel Streaming player/peer (using the the Pixel Streaming player id).
  */
-UCLASS(Blueprintable, ClassGroup = (PixelStreamer), meta = (BlueprintSpawnableComponent))
+UCLASS(Blueprintable, ClassGroup = (PixelStreaming), meta = (BlueprintSpawnableComponent))
 class PIXELSTREAMING_API UPixelStreamingAudioComponent : public USynthComponent, public IPixelStreamingAudioConsumer
 {
 	GENERATED_BODY()
diff --git a/Engine/Plugins/Media/PixelStreaming/Source/PixelStreaming/Public/PixelStreamingDelegates.h b/Engine/Plugins/Media/PixelStreaming/Source/PixelStreaming/Public/PixelStreamingDelegates.h
new file mode 100644
--- /dev/null
+++ b/Engine/Plugins/Media/PixelStreaming/Source/PixelStreaming/Public/PixelStreamingDelegates.h
@@ -0,0 +1,123 @@
+// Copyright Epic Games, Inc. All Rights Reserved.
+
+#pragma once
+
+#include "PixelStreamingPlayerId.h"
+#include "Containers/UnrealString.h"
+#include "Delegates/DelegateSignatureImpl.inl"
+#include "PixelStreamingDelegates.generated.h"
+
+
+namespace webrtc
+{
+	class DataChannelInterface;
+}
+
+UCLASS()
+class PIXELSTREAMING_API UPixelStreamingDelegates : public UObject
+{
+	GENERATED_BODY()
+
+public:
+	/**
+	 * A connection to the signalling server was made.
+	 */
+	// BP Delegate
+	DECLARE_DYNAMIC_MULTICAST_DELEGATE(FConnectedToSignallingServer);
+	UPROPERTY(BlueprintAssignable, Category = "Pixel Streaming Delegates")
+	FConnectedToSignallingServer OnConnecedToSignallingServer;
+	// C++ Delegate
+	DECLARE_MULTICAST_DELEGATE(FConnectedToSignallingServerNative);
+	FConnectedToSignallingServerNative OnConnecedToSignallingServerNative;
+
+
+	/**
+	 * A connection to the signalling server was lost.
+	 */
+	// BP Delegate
+	DECLARE_DYNAMIC_MULTICAST_DELEGATE(FDisconnectedFromSignallingServer);
+	UPROPERTY(BlueprintAssignable, Category = "Pixel Streaming Delegates")
+	FDisconnectedFromSignallingServer OnDisconnectedFromSignallingServer;
+	// C++ Delegate
+	DECLARE_MULTICAST_DELEGATE(FDisconnectedFromSignallingServerNative);
+	FDisconnectedFromSignallingServerNative OnDisconnectedFromSignallingServerNative;
+
+
+	/**
+	 * A new connection has been made to the session.
+	 */
+	// BP Delegate
+	DECLARE_DYNAMIC_MULTICAST_DELEGATE_TwoParams(FNewConnection, FString, PlayerId, bool, QualityController);
+	UPROPERTY(BlueprintAssignable, Category = "Pixel Streaming Delegates")
+	FNewConnection OnNewConnection;
+	// C++ Delegate
+	DECLARE_MULTICAST_DELEGATE_TwoParams(FNewConnectionNative, FString, bool);
+	FNewConnectionNative OnNewConnectionNative;
+
+
+	/**
+	 * A connection to a player was lost.
+	 */
+	// BP Delegate
+	DECLARE_DYNAMIC_MULTICAST_DELEGATE_TwoParams(FClosedConnection, FString, PlayerId, bool, WasQualityController);
+	UPROPERTY(BlueprintAssignable, Category = "Pixel Streaming Delegates")
+	FClosedConnection OnClosedConnection;
+	// C++ Delegate
+	DECLARE_MULTICAST_DELEGATE_TwoParams(FClosedConnectionNative, FPixelStreamingPlayerId, bool);
+	FClosedConnectionNative OnClosedConnectionNative;
+
+
+	/**
+	 * All connections have closed and nobody is viewing or interacting with
+	 * the app. This is an opportunity to reset the app.
+	 */
+	// BP Delegate
+	DECLARE_DYNAMIC_MULTICAST_DELEGATE(FAllConnectionsClosed);
+	UPROPERTY(BlueprintAssignable, Category = "Pixel Streaming Delegates")
+	FAllConnectionsClosed OnAllConnectionsClosed;
+	// C++ Delegate
+	DECLARE_MULTICAST_DELEGATE(FAllConnectionsClosedNative);
+	FAllConnectionsClosedNative OnAllConnectionsClosedNative;
+
+
+	/**
+	 * The quality controlling peer has changed
+	 */
+	// C++ Delegate
+	DECLARE_MULTICAST_DELEGATE_OneParam(FQualityControllerChanged, FPixelStreamingPlayerId)
+	FQualityControllerChanged OnQualityControllerChangedNative;
+
+
+	/**
+	 * A new data channel has been opened
+	 */
+	// C++ Delegate
+	DECLARE_MULTICAST_DELEGATE_TwoParams(FDataChannelOpen, FPixelStreamingPlayerId, webrtc::DataChannelInterface*)
+	FDataChannelOpen OnDataChannelOpenNative;
+	
+
+	/**
+	 * A pixel streaming stat has changed
+	 */
+	// C++ Delegate
+	DECLARE_MULTICAST_DELEGATE_ThreeParams(FPixelStreamerStatChanged, FPixelStreamingPlayerId, FName, float)
+	FPixelStreamerStatChanged OnPixelStreamerStatChangedNative;
+
+	/**
+	 * Create the singleton.
+	 */
+	static UPixelStreamingDelegates* CreateInstance();
+
+	static UPixelStreamingDelegates* GetPixelStreamingDelegates()
+	{
+		if (Singleton == nullptr)
+		{
+			return CreateInstance();
+		}
+		return Singleton;
+	}
+
+private:
+	// The singleton object.
+	static UPixelStreamingDelegates* Singleton;
+};
diff --git a/Engine/Plugins/Media/PixelStreaming/Source/PixelStreaming/Public/PixelStreamingStatNames.h b/Engine/Plugins/Media/PixelStreaming/Source/PixelStreaming/Public/PixelStreamingStatNames.h
--- a/Engine/Plugins/Media/PixelStreaming/Source/PixelStreaming/Public/PixelStreamingStatNames.h
+++ b/Engine/Plugins/Media/PixelStreaming/Source/PixelStreaming/Public/PixelStreamingStatNames.h
@@ -39,14 +39,15 @@ namespace PixelStreamingStatNames
 	const FName TotalEncodeTime			  = FName(TEXT("totalEncodeTime"));
 	const FName TotalPacketSendDelay	  = FName(TEXT("totalPacketSendDelay"));
 	const FName FramesEncoded			  = FName(TEXT("framesEncoded"));
+	const FName QualityController		  = FName(TEXT("qualityController"));
 
 	// Calculated stats
-	const FName FramesSentPerSecond	   = FName(TEXT("transmit fps"));
-	const FName Bitrate				   = FName(TEXT("bitrate (kb/s)"));
+	const FName FramesSentPerSecond	   = FName(TEXT("transmitFps"));
+	const FName Bitrate				   = FName(TEXT("bitrate"));
 	const FName MeanQPPerSecond		   = FName(TEXT("qp"));
-	const FName MeanEncodeTime		   = FName(TEXT("encode time (ms)"));
-	const FName EncodedFramesPerSecond = FName(TEXT("encode fps"));
-	const FName MeanSendDelay		   = FName(TEXT("capture to send (ms)"));
-	const FName SourceFps			   = FName(TEXT("capture fps"));
+	const FName MeanEncodeTime		   = FName(TEXT("encodeTime"));
+	const FName EncodedFramesPerSecond = FName(TEXT("encodeFps"));
+	const FName MeanSendDelay		   = FName(TEXT("captureToSend"));
+	const FName SourceFps			   = FName(TEXT("captureFps"));
 
 } // namespace PixelStreamingStatNames
\ No newline at end of file
