Fix so that malformed peerConnectionOptions do not crash the application.

diff --git a/Engine/Plugins/Media/PixelStreaming/Source/PixelStreaming/Private/ThreadSafePlayerSessions.cpp b/Engine/Plugins/Media/PixelStreaming/Source/PixelStreaming/Private/ThreadSafePlayerSessions.cpp
--- a/Engine/Plugins/Media/PixelStreaming/Source/PixelStreaming/Private/ThreadSafePlayerSessions.cpp
+++ b/Engine/Plugins/Media/PixelStreaming/Source/PixelStreaming/Private/ThreadSafePlayerSessions.cpp
@@ -548,11 +548,15 @@ webrtc::PeerConnectionInterface* UE::PixelStreaming::FThreadSafePlayerSessions::
 	UE_LOG(LogPixelStreaming, Log, TEXT("Creating player session for PlayerId=%s"), *PlayerId);
 
 	// this is called from WebRTC signalling thread, the only thread where `Players` map is modified, so no need to lock it
-	bool bMakeQualityController = Players.Num() == 0; // first player controls quality by default
+	bool bMakeQualityController = PlayerId != SFU_PLAYER_ID && Players.Num() == 0; // first player controls quality by default
 	UE::PixelStreaming::FPlayerSession* Session = new UE::PixelStreaming::FPlayerSession(this, SignallingServerConnection, PlayerId);
 
 	rtc::scoped_refptr<webrtc::PeerConnectionInterface> PeerConnection = PeerConnectionFactory->CreatePeerConnection(PeerConnectionConfig, webrtc::PeerConnectionDependencies{ Session });
-	check(PeerConnection);
+	if(!PeerConnection)
+	{
+		UE_LOG(LogPixelStreaming, Error, TEXT("Failed to created PeerConnection. This may indicate you passed malformed peerConnectionOptions."));
+		return nullptr;
+	}
 
 	// Setup suggested bitrate settings on the Peer Connection based on our CVars
 	webrtc::BitrateSettings BitrateSettings;
