Fix clean shutdown with Pixel Streaming.

- Explictly null SFUPeerConnectionFactory
- Add GPUFencePoller to be owned by FStreamer
- Add FFixedFPSPump to be owned by FStreamer
- Move registration of VideoSource to use a raw pointer as we were already scopelocking it, so reference counting seemed redundant.

diff --git a/Engine/Plugins/Media/PixelStreaming/Source/PixelStreaming/Private/FixedFPSPump.cpp b/Engine/Plugins/Media/PixelStreaming/Source/PixelStreaming/Private/FixedFPSPump.cpp
--- a/Engine/Plugins/Media/PixelStreaming/Source/PixelStreaming/Private/FixedFPSPump.cpp
+++ b/Engine/Plugins/Media/PixelStreaming/Source/PixelStreaming/Private/FixedFPSPump.cpp
@@ -7,23 +7,35 @@
 #include "FrameBuffer.h"
 #include "PlayerSession.h"
 
-UE::PixelStreaming::FFixedFPSPump& UE::PixelStreaming::FFixedFPSPump::Get()
+/*
+* ---------- FFixedFPSPump ----------------
+*/
+
+UE::PixelStreaming::FFixedFPSPump* UE::PixelStreaming::FFixedFPSPump::Instance = nullptr;
+
+UE::PixelStreaming::FFixedFPSPump* UE::PixelStreaming::FFixedFPSPump::Get()
 {
-	static UE::PixelStreaming::FFixedFPSPump Pump = UE::PixelStreaming::FFixedFPSPump();
-	return Pump;
+	checkf(Instance, TEXT("You should not try to Get() and instance of the pump before it has been constructed somewhere yet."));
+	return UE::PixelStreaming::FFixedFPSPump::Instance;
 }
 
 UE::PixelStreaming::FFixedFPSPump::FFixedFPSPump()
 	: NextPumpEvent(FPlatformProcess::GetSynchEventFromPool(false))
 {
 	PumpThread = MakeUnique<FThread>(TEXT("PumpThread"), [this]() { PumpLoop(); });
+	UE::PixelStreaming::FFixedFPSPump::Instance = this;
 }
 
 UE::PixelStreaming::FFixedFPSPump::~FFixedFPSPump()
+{
+	Shutdown();
+	PumpThread->Join();
+}
+
+void UE::PixelStreaming::FFixedFPSPump::Shutdown()
 {
 	bThreadRunning = false;
 	NextPumpEvent->Trigger();
-	PumpThread->Join();
 }
 
 void UE::PixelStreaming::FFixedFPSPump::UnregisterVideoSource(FPixelStreamingPlayerId PlayerId)
@@ -33,8 +45,9 @@ void UE::PixelStreaming::FFixedFPSPump::UnregisterVideoSource(FPixelStreamingPla
 	NextPumpEvent->Trigger();
 }
 
-void UE::PixelStreaming::FFixedFPSPump::RegisterVideoSource(FPixelStreamingPlayerId PlayerId, rtc::scoped_refptr<UE::PixelStreaming::IPumpedVideoSource> Source)
+void UE::PixelStreaming::FFixedFPSPump::RegisterVideoSource(FPixelStreamingPlayerId PlayerId, UE::PixelStreaming::IPumpedVideoSource* Source)
 {
+	checkf(Source, TEXT("Cannot register a nullptr VideoSource."));
 	FScopeLock Guard(&SourcesGuard);
 	VideoSources.Add(PlayerId, Source);
 	NextPumpEvent->Trigger();
@@ -58,18 +71,12 @@ void UE::PixelStreaming::FFixedFPSPump::PumpLoop()
 			const int32 FrameId = NextFrameId++;
 
 			// Pump each video source
-			TMap<FPixelStreamingPlayerId, rtc::scoped_refptr<UE::PixelStreaming::IPumpedVideoSource>>::TIterator Iter = VideoSources.CreateIterator();
+			TMap<FPixelStreamingPlayerId, IPumpedVideoSource*>::TIterator Iter = VideoSources.CreateIterator();
 			for(; Iter; ++Iter)
 			{
-				rtc::scoped_refptr<UE::PixelStreaming::IPumpedVideoSource>& VideoSource = Iter.Value();
 
-				// If the pump is the last thing holding this video source, remove the video source.
-				if(VideoSource->HasOneRef())
-				{
-					Iter.RemoveCurrent();
-					continue;
-				}
-				
+				IPumpedVideoSource* VideoSource = Iter.Value();
+
 				if(VideoSource->IsReadyForPump())
 				{
 					VideoSource->OnPump(FrameId);
diff --git a/Engine/Plugins/Media/PixelStreaming/Source/PixelStreaming/Private/FixedFPSPump.h b/Engine/Plugins/Media/PixelStreaming/Source/PixelStreaming/Private/FixedFPSPump.h
--- a/Engine/Plugins/Media/PixelStreaming/Source/PixelStreaming/Private/FixedFPSPump.h
+++ b/Engine/Plugins/Media/PixelStreaming/Source/PixelStreaming/Private/FixedFPSPump.h
@@ -8,6 +8,9 @@
 namespace UE {
 	namespace PixelStreaming {
 
+		/*
+		* Interface for video sources that get pumped through `OnPump`.
+		*/
 		class IPumpedVideoSource : public rtc::RefCountInterface
 		{
 			public:
@@ -29,23 +32,20 @@ namespace UE {
 		class FFixedFPSPump
 		{
 
-		private:
-			// Private constructor as there should only ever be one instance of this class and the intended design is to use the static Get() method.
+		public:
 			FFixedFPSPump();
 			~FFixedFPSPump();
-			void PumpLoop();
-
-		public:
-			
-			// Single accessor.
-			static FFixedFPSPump& Get();
-
-			void RegisterVideoSource(FPixelStreamingPlayerId PlayerId, rtc::scoped_refptr<IPumpedVideoSource> Source);
+			void Shutdown();
+			void RegisterVideoSource(FPixelStreamingPlayerId PlayerId, IPumpedVideoSource* Source);
 			void UnregisterVideoSource(FPixelStreamingPlayerId PlayerId);
+			static FFixedFPSPump* Get();
+
+		private:
+			void PumpLoop();
 
 		private:
 			FCriticalSection SourcesGuard;
-			TMap<FPixelStreamingPlayerId, rtc::scoped_refptr<IPumpedVideoSource>> VideoSources;
+			TMap<FPixelStreamingPlayerId, IPumpedVideoSource*> VideoSources;
 
 
 			TUniquePtr<FThread> PumpThread;
@@ -54,6 +54,7 @@ namespace UE {
 			FEvent* NextPumpEvent;
 			int32 NextFrameId = 0;
 
+			static FFixedFPSPump* Instance;
 		};
 	}
 }
diff --git a/Engine/Plugins/Media/PixelStreaming/Source/PixelStreaming/Private/GPUFencePoller.cpp b/Engine/Plugins/Media/PixelStreaming/Source/PixelStreaming/Private/GPUFencePoller.cpp
--- a/Engine/Plugins/Media/PixelStreaming/Source/PixelStreaming/Private/GPUFencePoller.cpp
+++ b/Engine/Plugins/Media/PixelStreaming/Source/PixelStreaming/Private/GPUFencePoller.cpp
@@ -5,24 +5,32 @@
 * ------------- UE::PixelStreaming::FGPUFencePoller ---------------------
 */
 
-UE::PixelStreaming::FGPUFencePoller& UE::PixelStreaming::FGPUFencePoller::Get()
+UE::PixelStreaming::FGPUFencePoller* UE::PixelStreaming::FGPUFencePoller::Instance = nullptr;
+
+UE::PixelStreaming::FGPUFencePoller* UE::PixelStreaming::FGPUFencePoller::Get()
 {
-	static UE::PixelStreaming::FGPUFencePoller Poller;
-	return Poller;
+	checkf(Instance, TEXT("You should not try to Get() and instance of the poller before it has been constructed somewhere yet."));
+	return UE::PixelStreaming::FGPUFencePoller::Instance;
 }
 
 UE::PixelStreaming::FGPUFencePoller::FGPUFencePoller()
 	: Runnable()
 	, Thread(FRunnableThread::Create(&Runnable, TEXT("Pixel Streaming GPUFencePoller")))
 {
+	UE::PixelStreaming::FGPUFencePoller::Instance = this;
 }
 
 UE::PixelStreaming::FGPUFencePoller::~FGPUFencePoller()
 {
-	Runnable.Stop();
+	Shutdown();
 	Thread->Kill(true);
 }
 
+void UE::PixelStreaming::FGPUFencePoller::Shutdown()
+{
+	Runnable.Stop();
+}
+
 void UE::PixelStreaming::FGPUFencePoller::AddJob(FGPUFenceRHIRef Fence, TSharedRef<bool, ESPMode::ThreadSafe> bKeepPolling, TFunction<void()> FenceDoneCallback)
 {
 	checkf(Runnable.IsRunning(), TEXT("Poller runnable must be running to submit jobs to it."));
@@ -38,6 +46,7 @@ UE::PixelStreaming::FGPUFencePoller::FPollJob::FPollJob(FGPUFenceRHIRef InFence,
 	, bKeepPolling(bInKeepPolling)
 	, FenceDoneCallback(InFenceDoneCallback)
 {
+	
 }
 
 UE::PixelStreaming::FGPUFencePoller::FPollJob::FPollJob()
@@ -126,9 +135,11 @@ uint32 UE::PixelStreaming::FGPUFencePoller::FPollerRunnable::Run()
 void UE::PixelStreaming::FGPUFencePoller::FPollerRunnable::Stop()
 {
 	bIsRunning = false;
+	JobAddedEvent->Trigger();
 }
 
 void UE::PixelStreaming::FGPUFencePoller::FPollerRunnable::Exit()
 {
 	bIsRunning = false;
+	JobAddedEvent->Trigger();
 }
diff --git a/Engine/Plugins/Media/PixelStreaming/Source/PixelStreaming/Private/GPUFencePoller.h b/Engine/Plugins/Media/PixelStreaming/Source/PixelStreaming/Private/GPUFencePoller.h
--- a/Engine/Plugins/Media/PixelStreaming/Source/PixelStreaming/Private/GPUFencePoller.h
+++ b/Engine/Plugins/Media/PixelStreaming/Source/PixelStreaming/Private/GPUFencePoller.h
@@ -21,15 +21,11 @@ namespace UE {
         class FGPUFencePoller
         {
             public:
-                static FGPUFencePoller& Get();
-
-            private:
-                /* Private constructor because the intended usage is through the static Get method. */
                 FGPUFencePoller();
-
-            public:
+                void Shutdown();
                 virtual ~FGPUFencePoller();
                 void AddJob(FGPUFenceRHIRef Fence, TSharedRef<bool, ESPMode::ThreadSafe> bKeepPolling, TFunction<void()> FenceDoneCallback);
+                static FGPUFencePoller* Get();
 
             private:
                 struct FPollJob
@@ -63,6 +59,8 @@ namespace UE {
             private:
                 FPollerRunnable Runnable;
                 FRunnableThread* Thread;
+
+                static FGPUFencePoller* Instance;
         };
     }
 }
diff --git a/Engine/Plugins/Media/PixelStreaming/Source/PixelStreaming/Private/Streamer.cpp b/Engine/Plugins/Media/PixelStreaming/Source/PixelStreaming/Private/Streamer.cpp
--- a/Engine/Plugins/Media/PixelStreaming/Source/PixelStreaming/Private/Streamer.cpp
+++ b/Engine/Plugins/Media/PixelStreaming/Source/PixelStreaming/Private/Streamer.cpp
@@ -51,6 +51,7 @@ UE::PixelStreaming::FStreamer::~FStreamer()
 {
 	DeleteAllPlayerSessions();
 	P2PPeerConnectionFactory = nullptr;
+	SFUPeerConnectionFactory = nullptr;
 	WebRtcSignallingThread->Stop();
 	rtc::CleanupSSL();
 }
diff --git a/Engine/Plugins/Media/PixelStreaming/Source/PixelStreaming/Private/Streamer.h b/Engine/Plugins/Media/PixelStreaming/Source/PixelStreaming/Private/Streamer.h
--- a/Engine/Plugins/Media/PixelStreaming/Source/PixelStreaming/Private/Streamer.h
+++ b/Engine/Plugins/Media/PixelStreaming/Source/PixelStreaming/Private/Streamer.h
@@ -9,13 +9,14 @@
 #include "HAL/Thread.h"
 #include "IPixelStreamingSessions.h"
 #include "ThreadSafePlayerSessions.h"
+#include "FixedFPSPump.h"
+#include "GPUFencePoller.h"
 
 namespace UE {
 	namespace PixelStreaming {
 		class FPlayerSession;
 		class FVideoEncoderFactory;
 		class FSimulcastEncoderFactory;
-		class FStreamer;
 		class FSetSessionDescriptionObserver;
 
 		class FStreamer : public FSignallingServerConnectionObserver, public IPixelStreamingSessions
@@ -36,6 +37,8 @@ namespace UE {
 			void SetStreamingStarted(bool bStarted) { bStreamingStarted = bStarted; }
 			FSignallingServerConnection* GetSignallingServerConnection() const { return SignallingServerConnection.Get(); }
 			FVideoEncoderFactory* GetP2PVideoEncoderFactory() const { return P2PVideoEncoderFactory; }
+			FFixedFPSPump& GetPumpThread() { return PumpThread; }
+			FGPUFencePoller& GetFencePollerThread() { return FencePollerThread; }
 
 			// data coming from the engine
 			void OnFrameBufferReady(const FTexture2DRHIRef& FrameBuffer);
@@ -122,6 +125,9 @@ namespace UE {
 			FThreadSafePlayerSessions PlayerSessions;
 			FStats Stats;
 
+			FFixedFPSPump PumpThread;
+			FGPUFencePoller FencePollerThread;
+
 			TUniquePtr<webrtc::SessionDescriptionInterface> SFULocalDescription;
 			TUniquePtr<webrtc::SessionDescriptionInterface> SFURemoteDescription;
 		};
diff --git a/Engine/Plugins/Media/PixelStreaming/Source/PixelStreaming/Private/TextureSource.cpp b/Engine/Plugins/Media/PixelStreaming/Source/PixelStreaming/Private/TextureSource.cpp
--- a/Engine/Plugins/Media/PixelStreaming/Source/PixelStreaming/Private/TextureSource.cpp
+++ b/Engine/Plugins/Media/PixelStreaming/Source/PixelStreaming/Private/TextureSource.cpp
@@ -78,7 +78,7 @@ void UE::PixelStreaming::FTextureSourceBackBuffer::OnBackBufferReady_RenderThrea
 
 		UE::PixelStreaming::CopyTexture(FrameBuffer, WriteBuffer.Texture, WriteBuffer.Fence);
 
-		UE::PixelStreaming::FGPUFencePoller::Get().AddJob(WriteBuffer.Fence, bEnabled, [this, &WriteBuffer]() {
+		UE::PixelStreaming::FGPUFencePoller::Get()->AddJob(WriteBuffer.Fence, bEnabled, [this, &WriteBuffer]() {
 			// This lambda is called only once the GPUFence is done
 			{
 				FScopeLock Lock(&CriticalSection);
diff --git a/Engine/Plugins/Media/PixelStreaming/Source/PixelStreaming/Private/VideoSource.cpp b/Engine/Plugins/Media/PixelStreaming/Source/PixelStreaming/Private/VideoSource.cpp
--- a/Engine/Plugins/Media/PixelStreaming/Source/PixelStreaming/Private/VideoSource.cpp
+++ b/Engine/Plugins/Media/PixelStreaming/Source/PixelStreaming/Private/VideoSource.cpp
@@ -21,6 +21,7 @@ rtc::RefCountReleaseStatus UE::PixelStreaming::FVideoSourceBase::Release() const
 	const rtc::RefCountReleaseStatus Status = RefCount.DecRef();
 	if (Status == rtc::RefCountReleaseStatus::kDroppedLastRef)
 	{
+		UE::PixelStreaming::FFixedFPSPump::Get()->UnregisterVideoSource(PlayerId);
 		delete this;
 	}
 	return Status;
@@ -29,13 +30,13 @@ rtc::RefCountReleaseStatus UE::PixelStreaming::FVideoSourceBase::Release() const
 void UE::PixelStreaming::FVideoSourceBase::Initialize()
 {
 	CurrentState = webrtc::MediaSourceInterface::SourceState::kInitializing;
-	UE::PixelStreaming::FFixedFPSPump::Get().RegisterVideoSource(PlayerId, this);
+
+	UE::PixelStreaming::FFixedFPSPump::Get()->RegisterVideoSource(PlayerId, this);
 }
 
 UE::PixelStreaming::FVideoSourceBase::~FVideoSourceBase()
 {
 	CurrentState = webrtc::MediaSourceInterface::SourceState::kEnded;
-	UE::PixelStreaming::FFixedFPSPump::Get().UnregisterVideoSource(PlayerId);
 }
 
 void UE::PixelStreaming::FVideoSourceBase::OnPump(int32 FrameId)
