Add PIXELSTREAMING_API to PixelStreamingInputComponent.h so it has symbols exported and usable in the public API.

diff --git a/Engine/Plugins/Media/PixelStreaming/Source/PixelStreaming/Public/PixelStreamingInputComponent.h b/Engine/Plugins/Media/PixelStreaming/Source/PixelStreaming/Public/PixelStreamingInputComponent.h
--- a/Engine/Plugins/Media/PixelStreaming/Source/PixelStreaming/Public/PixelStreamingInputComponent.h
+++ b/Engine/Plugins/Media/PixelStreaming/Source/PixelStreaming/Public/PixelStreamingInputComponent.h
@@ -12,7 +12,7 @@
  * Responses back to the source of the UI interactions may also be sent.
  */
 UCLASS(Blueprintable, ClassGroup = (PixelStreaming), meta = (BlueprintSpawnableComponent))
-class UPixelStreamingInput : public UActorComponent
+class PIXELSTREAMING_API UPixelStreamingInput : public UActorComponent
 {
 	GENERATED_BODY()
 
