Fix regression where there was no audio transmitted, this was caused by introduction of simulcast.

diff --git a/Engine/Plugins/Media/PixelStreaming/Source/PixelStreaming/Private/Streamer.cpp b/Engine/Plugins/Media/PixelStreaming/Source/PixelStreaming/Private/Streamer.cpp
--- a/Engine/Plugins/Media/PixelStreaming/Source/PixelStreaming/Private/Streamer.cpp
+++ b/Engine/Plugins/Media/PixelStreaming/Source/PixelStreaming/Private/Streamer.cpp
@@ -39,7 +39,6 @@ UE::PixelStreaming::FStreamer::FStreamer(const FString& InSignallingServerUrl, c
 	StartWebRtcSignallingThread();
 	ConnectToSignallingServer();
 
-
 	if (UPixelStreamingDelegates* Delegates = UPixelStreamingDelegates::GetPixelStreamingDelegates())
 	{
 		Delegates->OnClosedConnectionNative.AddRaw(this, &UE::PixelStreaming::FStreamer::PostPlayerDeleted);
@@ -56,18 +55,8 @@ UE::PixelStreaming::FStreamer::~FStreamer()
 	rtc::CleanupSSL();
 }
 
-void UE::PixelStreaming::FStreamer::StartWebRtcSignallingThread()
+rtc::scoped_refptr<webrtc::AudioDeviceModule> UE::PixelStreaming::FStreamer::CreateAudioDeviceModule() const
 {
-	// Initialisation of WebRTC. Note: That interfacing with WebRTC should generally happen in WebRTC signalling thread.
-
-	// Create our own WebRTC thread for signalling
-	WebRtcSignallingThread->SetName("WebRtcSignallingThread", nullptr);
-	WebRtcSignallingThread->Start();
-
-	rtc::InitializeSSL();
-
-	PeerConnectionConfig = {};
-
 	bool bUseLegacyAudioDeviceModule = UE::PixelStreaming::Settings::CVarPixelStreamingWebRTCUseLegacyAudioDevice.GetValueOnAnyThread();
 	rtc::scoped_refptr<webrtc::AudioDeviceModule> AudioDeviceModule;
 	if (bUseLegacyAudioDeviceModule)
@@ -78,17 +67,31 @@ void UE::PixelStreaming::FStreamer::StartWebRtcSignallingThread()
 	{
 		AudioDeviceModule = new rtc::RefCountedObject<UE::PixelStreaming::FAudioDeviceModule>();
 	}
+	return AudioDeviceModule;
+}
+
+void UE::PixelStreaming::FStreamer::StartWebRtcSignallingThread()
+{
+	// Initialisation of WebRTC. Note: That interfacing with WebRTC should generally happen in WebRTC signalling thread.
+
+	// Create our own WebRTC thread for signalling
+	WebRtcSignallingThread->SetName("WebRtcSignallingThread", nullptr);
+	WebRtcSignallingThread->Start();
+
+	rtc::InitializeSSL();
+
+	PeerConnectionConfig = {};
 
 	/* ---------- P2P Peer Connection Factory ---------- */
 
 	std::unique_ptr<UE::PixelStreaming::FVideoEncoderFactory> P2PPeerConnectionFactoryPtr = std::make_unique<UE::PixelStreaming::FVideoEncoderFactory>();
-	P2PVideoEncoderFactory = P2PPeerConnectionFactoryPtr.get(); 
+	P2PVideoEncoderFactory = P2PPeerConnectionFactoryPtr.get();
 
 	P2PPeerConnectionFactory = webrtc::CreatePeerConnectionFactory(
 		nullptr,					  // network_thread
 		nullptr,					  // worker_thread
 		WebRtcSignallingThread.Get(), // signal_thread
-		AudioDeviceModule,			  // audio device manager
+		CreateAudioDeviceModule(),	  // audio device module
 		webrtc::CreateAudioEncoderFactory<webrtc::AudioEncoderOpus>(),
 		webrtc::CreateAudioDecoderFactory<webrtc::AudioDecoderOpus>(),
 		std::move(P2PPeerConnectionFactoryPtr),
@@ -101,13 +104,13 @@ void UE::PixelStreaming::FStreamer::StartWebRtcSignallingThread()
 	/* ---------- SFU Peer Connection Factory ---------- */
 
 	std::unique_ptr<UE::PixelStreaming::FSimulcastEncoderFactory> SFUPeerConnectionFactoryPtr = std::make_unique<UE::PixelStreaming::FSimulcastEncoderFactory>();
-	SFUVideoEncoderFactory = SFUPeerConnectionFactoryPtr.get(); 
+	SFUVideoEncoderFactory = SFUPeerConnectionFactoryPtr.get();
 
 	SFUPeerConnectionFactory = webrtc::CreatePeerConnectionFactory(
 		nullptr,					  // network_thread
 		nullptr,					  // worker_thread
 		WebRtcSignallingThread.Get(), // signal_thread
-		AudioDeviceModule,			  // audio device manager
+		CreateAudioDeviceModule(),	  // audio device module
 		webrtc::CreateAudioEncoderFactory<webrtc::AudioEncoderOpus>(),
 		webrtc::CreateAudioDecoderFactory<webrtc::AudioDecoderOpus>(),
 		std::move(SFUPeerConnectionFactoryPtr),
@@ -116,7 +119,6 @@ void UE::PixelStreaming::FStreamer::StartWebRtcSignallingThread()
 		SetupAudioProcessingModule()); // audio_processing
 
 	check(SFUPeerConnectionFactory.get() != nullptr);
-
 }
 
 webrtc::AudioProcessing* UE::PixelStreaming::FStreamer::SetupAudioProcessingModule()
@@ -183,10 +185,8 @@ webrtc::PeerConnectionInterface* UE::PixelStreaming::FStreamer::CreateSession(FP
 	PeerConnectionConfig.enable_simulcast_stats = true;
 
 	bool bIsSFU = (Flags & UE::PixelStreaming::Protocol::EPlayerFlags::PSPFlag_IsSFU) != UE::PixelStreaming::Protocol::EPlayerFlags::PSPFlag_None;
-	
-	rtc::scoped_refptr<webrtc::PeerConnectionFactoryInterface> PCFactory = bIsSFU ?
-		SFUPeerConnectionFactory :
-		P2PPeerConnectionFactory;
+
+	rtc::scoped_refptr<webrtc::PeerConnectionFactoryInterface> PCFactory = bIsSFU ? SFUPeerConnectionFactory : P2PPeerConnectionFactory;
 
 	webrtc::PeerConnectionInterface* PeerConnection = PlayerSessions.CreatePlayerSession(
 		PlayerId,
@@ -324,8 +324,8 @@ void UE::PixelStreaming::FStreamer::SendAnswer(FPixelStreamingPlayerId PlayerId,
 
 	auto OnSetRemoteDescriptionSuccess = [this, PlayerId, PeerConnection, CreateAnswerObserver]() {
 		// Note: these offer to receive are superseded now we are use transceivers to setup our peer connection media
-		int offer_to_receive_video = webrtc::PeerConnectionInterface::RTCOfferAnswerOptions::kUndefined; 
-		int offer_to_receive_audio = webrtc::PeerConnectionInterface::RTCOfferAnswerOptions::kUndefined; 
+		int offer_to_receive_video = webrtc::PeerConnectionInterface::RTCOfferAnswerOptions::kUndefined;
+		int offer_to_receive_audio = webrtc::PeerConnectionInterface::RTCOfferAnswerOptions::kUndefined;
 		bool voice_activity_detection = false;
 		bool ice_restart = true;
 		bool use_rtp_mux = true;
@@ -362,14 +362,14 @@ void UE::PixelStreaming::FStreamer::MungeRemoteSDP(cricket::SessionDescription*
 {
 	// Munge SDP of remote description to inject min, max, start bitrates
 	std::vector<cricket::ContentInfo>& ContentInfos = RemoteDescription->contents();
-	for(cricket::ContentInfo& Content : ContentInfos)
+	for (cricket::ContentInfo& Content : ContentInfos)
 	{
 		cricket::MediaContentDescription* MediaDescription = Content.media_description();
-		if(MediaDescription->type() == cricket::MediaType::MEDIA_TYPE_VIDEO)
+		if (MediaDescription->type() == cricket::MediaType::MEDIA_TYPE_VIDEO)
 		{
 			cricket::VideoContentDescription* VideoDescription = MediaDescription->as_video();
 			std::vector<cricket::VideoCodec> CodecsCopy = VideoDescription->codecs();
-			for(cricket::VideoCodec& Codec : CodecsCopy)
+			for (cricket::VideoCodec& Codec : CodecsCopy)
 			{
 				// Note: These params are passed as kilobits, so divide by 1000.
 				Codec.SetParam(cricket::kCodecParamMinBitrate, UE::PixelStreaming::Settings::CVarPixelStreamingWebRTCMinBitrate.GetValueOnAnyThread() / 1000);
@@ -377,7 +377,7 @@ void UE::PixelStreaming::FStreamer::MungeRemoteSDP(cricket::SessionDescription*
 				Codec.SetParam(cricket::kCodecParamMaxBitrate, UE::PixelStreaming::Settings::CVarPixelStreamingWebRTCMaxBitrate.GetValueOnAnyThread() / 1000);
 			}
 			VideoDescription->set_codecs(CodecsCopy);
-		}	
+		}
 	}
 }
 
@@ -574,14 +574,12 @@ void UE::PixelStreaming::FStreamer::SetupVideoTrack(FPixelStreamingPlayerId Play
 {
 
 	bool bIsSFU = (Flags & UE::PixelStreaming::Protocol::EPlayerFlags::PSPFlag_IsSFU) != UE::PixelStreaming::Protocol::EPlayerFlags::PSPFlag_None;
-	
-	rtc::scoped_refptr<webrtc::PeerConnectionFactoryInterface> PCFactory = bIsSFU ?
-		SFUPeerConnectionFactory :
-		P2PPeerConnectionFactory;
+
+	rtc::scoped_refptr<webrtc::PeerConnectionFactoryInterface> PCFactory = bIsSFU ? SFUPeerConnectionFactory : P2PPeerConnectionFactory;
 
 	// Create a video source for this player
 	rtc::scoped_refptr<FVideoSourceBase> VideoSource;
-	if(bIsSFU)
+	if (bIsSFU)
 	{
 		VideoSource = new FVideoSourceSFU();
 	}
@@ -624,7 +622,7 @@ void UE::PixelStreaming::FStreamer::SetupVideoTrack(FPixelStreamingPlayerId Play
 	}
 
 	// If there is no existing video transceiver, add one.
-	if(!bHasVideoTransceiver)
+	if (!bHasVideoTransceiver)
 	{
 		webrtc::RtpTransceiverInit TransceiverOptions;
 		TransceiverOptions.stream_ids = { TCHAR_TO_UTF8(*InVideoStreamId) };
@@ -640,10 +638,8 @@ void UE::PixelStreaming::FStreamer::SetupAudioTrack(webrtc::PeerConnectionInterf
 {
 
 	bool bIsSFU = (Flags & UE::PixelStreaming::Protocol::EPlayerFlags::PSPFlag_IsSFU) != UE::PixelStreaming::Protocol::EPlayerFlags::PSPFlag_None;
-	
-	rtc::scoped_refptr<webrtc::PeerConnectionFactoryInterface> PCFactory = bIsSFU ?
-		SFUPeerConnectionFactory :
-		P2PPeerConnectionFactory;
+
+	rtc::scoped_refptr<webrtc::PeerConnectionFactoryInterface> PCFactory = bIsSFU ? SFUPeerConnectionFactory : P2PPeerConnectionFactory;
 
 	bool bTransmitUEAudio = !UE::PixelStreaming::Settings::CVarPixelStreamingWebRTCDisableTransmitAudio.GetValueOnAnyThread();
 
@@ -671,7 +667,7 @@ void UE::PixelStreaming::FStreamer::SetupAudioTrack(webrtc::PeerConnectionInterf
 	}
 
 	// Add the audio track to the audio transceiver's sender if we are transmitting audio
-	if(!bTransmitUEAudio)
+	if (!bTransmitUEAudio)
 	{
 		return;
 	}
@@ -689,7 +685,7 @@ void UE::PixelStreaming::FStreamer::SetupAudioTrack(webrtc::PeerConnectionInterf
 		}
 	}
 
-	if(!bHasAudioTransceiver)
+	if (!bHasAudioTransceiver)
 	{
 		// Add the track
 		webrtc::RTCErrorOr<rtc::scoped_refptr<webrtc::RtpSenderInterface>> Result = PeerConnection->AddTrack(AudioTrack, { TCHAR_TO_UTF8(*InAudioStreamId) });
@@ -699,7 +695,6 @@ void UE::PixelStreaming::FStreamer::SetupAudioTrack(webrtc::PeerConnectionInterf
 			UE_LOG(LogPixelStreaming, Error, TEXT("Failed to add audio track to PeerConnection. Msg=%s"), TCHAR_TO_UTF8(Result.error().message()));
 		}
 	}
-
 }
 
 void UE::PixelStreaming::FStreamer::ModifyTransceivers(std::vector<rtc::scoped_refptr<webrtc::RtpTransceiverInterface>> Transceivers, int Flags)
@@ -735,17 +730,16 @@ void UE::PixelStreaming::FStreamer::ModifyTransceivers(std::vector<rtc::scoped_r
 			}
 
 			Transceiver->SetDirection(AudioTransceiverDirection);
-			Sender->SetStreams( { TCHAR_TO_UTF8(*GetAudioStreamID()) } );
+			Sender->SetStreams({ TCHAR_TO_UTF8(*GetAudioStreamID()) });
 		}
 		else if (Transceiver->media_type() == cricket::MediaType::MEDIA_TYPE_VIDEO)
 		{
 			Transceiver->SetDirection(webrtc::RtpTransceiverDirection::kSendOnly);
-			Sender->SetStreams( { TCHAR_TO_UTF8(*GetVideoStreamID()) } );
+			Sender->SetStreams({ TCHAR_TO_UTF8(*GetVideoStreamID()) });
 
 			webrtc::RtpParameters RtpParams = Sender->GetParameters();
 			RtpParams.encodings = CreateRTPEncodingParams(Flags);
 			Sender->SetParameters(RtpParams);
-
 		}
 	}
 }
diff --git a/Engine/Plugins/Media/PixelStreaming/Source/PixelStreaming/Private/Streamer.h b/Engine/Plugins/Media/PixelStreaming/Source/PixelStreaming/Private/Streamer.h
--- a/Engine/Plugins/Media/PixelStreaming/Source/PixelStreaming/Private/Streamer.h
+++ b/Engine/Plugins/Media/PixelStreaming/Source/PixelStreaming/Private/Streamer.h
@@ -12,8 +12,10 @@
 #include "FixedFPSPump.h"
 #include "GPUFencePoller.h"
 
-namespace UE {
-	namespace PixelStreaming {
+namespace UE
+{
+	namespace PixelStreaming
+	{
 		class FPlayerSession;
 		class FVideoEncoderFactory;
 		class FSimulcastEncoderFactory;
@@ -90,6 +92,7 @@ namespace UE {
 			void SetLocalDescription(webrtc::PeerConnectionInterface* PeerConnection, FSetSessionDescriptionObserver* Observer, webrtc::SessionDescriptionInterface* SDP);
 			FString GetAudioStreamID() const;
 			FString GetVideoStreamID() const;
+			rtc::scoped_refptr<webrtc::AudioDeviceModule> CreateAudioDeviceModule() const;
 
 		private:
 			FString SignallingServerUrl;
@@ -131,5 +134,5 @@ namespace UE {
 			TUniquePtr<webrtc::SessionDescriptionInterface> SFULocalDescription;
 			TUniquePtr<webrtc::SessionDescriptionInterface> SFURemoteDescription;
 		};
-	}
-}
+	} // namespace PixelStreaming
+} // namespace UE
