Fix for intermittent crash in PixelStreamingStats due to missing IsEngineExitRequested().

diff --git a/Engine/Plugins/Media/PixelStreaming/Source/PixelStreaming/Private/Stats.cpp b/Engine/Plugins/Media/PixelStreaming/Source/PixelStreaming/Private/Stats.cpp
--- a/Engine/Plugins/Media/PixelStreaming/Source/PixelStreaming/Private/Stats.cpp
+++ b/Engine/Plugins/Media/PixelStreaming/Source/PixelStreaming/Private/Stats.cpp
@@ -104,6 +104,11 @@ void UE::PixelStreaming::FStats::FireStatChanged_GameThread(FPixelStreamingPlaye
 {
 	checkf(IsInGameThread(), TEXT("This method must be called from the game thread"));
 
+	if (IsEngineExitRequested())
+	{
+		return;
+	}
+
 	if (UPixelStreamingDelegates* Delegates = UPixelStreamingDelegates::GetPixelStreamingDelegates())
 	{
 		Delegates->OnPixelStreamerStatChangedNative.Broadcast(PlayerId, StatName, StatValue);
@@ -112,6 +117,13 @@ void UE::PixelStreaming::FStats::FireStatChanged_GameThread(FPixelStreamingPlaye
 
 void UE::PixelStreaming::FStats::RemovePeerStat_GameThread(FPixelStreamingPlayerId PlayerId)
 {
+	checkf(IsInGameThread(), TEXT("This method must be called from the game thread"));
+
+	if (IsEngineExitRequested())
+	{
+		return;
+	}
+
 	PeerStats.Remove(PlayerId);
 
 	if (PlayerId == SFU_PLAYER_ID)
@@ -233,6 +245,11 @@ void UE::PixelStreaming::FStats::AddOnPeerStatChangedCallback_GameThread(FPixelS
 {
 	checkf(IsInGameThread(), TEXT("This method was not called from the game thread."));
 
+	if (IsEngineExitRequested())
+	{
+		return;
+	}
+
 	UE::PixelStreaming::FPeerStats* SinglePeerStats = PeerStats.Find(PlayerId);
 	if (!SinglePeerStats)
 	{
