Fix for when texture source gets destroyed very quickly before a frame has ticked yet.

diff --git a/Engine/Plugins/Media/PixelStreaming/Source/PixelStreaming/Private/TextureSource.cpp b/Engine/Plugins/Media/PixelStreaming/Source/PixelStreaming/Private/TextureSource.cpp
--- a/Engine/Plugins/Media/PixelStreaming/Source/PixelStreaming/Private/TextureSource.cpp
+++ b/Engine/Plugins/Media/PixelStreaming/Source/PixelStreaming/Private/TextureSource.cpp
@@ -7,8 +7,6 @@
 #include "GPUFencePoller.h"
 #include "Framework/Application/SlateApplication.h"
 
-
-
 /*
 * -------- UE::PixelStreaming::FTextureSourceBackBuffer ----------------
 */
@@ -16,8 +14,16 @@
 UE::PixelStreaming::FTextureSourceBackBuffer::FTextureSourceBackBuffer(float InFrameScale)
 	: FrameScale(InFrameScale)
 {
+	// Explictly make clear we are adding another ref to this shared bool for the purposes of using in the below lambda
+	TSharedRef<bool, ESPMode::ThreadSafe> bEnabledClone = bEnabled;
+
 	// The backbuffer delegate can only be accessed using GameThread
-	AsyncTask(ENamedThreads::GameThread, [this]() {
+	AsyncTask(ENamedThreads::GameThread, [this, bEnabledClone]() {
+		/*Early exit if `this` died before game thread ran.*/
+		if (bEnabled.IsUnique())
+		{
+			return;
+		}
 		OnBackbuffer = &FSlateApplication::Get().GetRenderer()->OnBackBufferReadyToPresent();
 		BackbufferDelegateHandle = OnBackbuffer->AddSP(this, &UE::PixelStreaming::FTextureSourceBackBuffer::OnBackBufferReady_RenderThread);
 	});
@@ -28,25 +34,24 @@ UE::PixelStreaming::FTextureSourceBackBuffer::FTextureSourceBackBuffer()
 
 UE::PixelStreaming::FTextureSourceBackBuffer::~FTextureSourceBackBuffer()
 {
-	if(OnBackbuffer)
+	if (OnBackbuffer)
 	{
 		OnBackbuffer->Remove(BackbufferDelegateHandle);
 	}
-	
+
 	SetEnabled(false);
 }
 
-void UE::PixelStreaming::FTextureSourceBackBuffer::SetEnabled(bool bInEnabled) 
-{ 
-	*bEnabled = bInEnabled; 
+void UE::PixelStreaming::FTextureSourceBackBuffer::SetEnabled(bool bInEnabled)
+{
+	*bEnabled = bInEnabled;
 
-	// This source has been disabled, so set `bInitialized` to false so `OnBackBufferReady_RenderThread` 
+	// This source has been disabled, so set `bInitialized` to false so `OnBackBufferReady_RenderThread`
 	// will make new textures next time it is called.
-	if(bInitialized && bInEnabled == false)
+	if (bInitialized && bInEnabled == false)
 	{
 		bInitialized = false;
 	}
-	
 }
 
 void UE::PixelStreaming::FTextureSourceBackBuffer::OnBackBufferReady_RenderThread(SWindow& SlateWindow, const FTexture2DRHIRef& FrameBuffer)
@@ -62,7 +67,7 @@ void UE::PixelStreaming::FTextureSourceBackBuffer::OnBackBufferReady_RenderThrea
 	}
 
 	auto& WriteBuffer = bWriteParity ? WriteBuffers[0] : WriteBuffers[1];
-	bWriteParity	  = !bWriteParity;
+	bWriteParity = !bWriteParity;
 
 	// for safety we just make sure that the buffer is not currently waiting for a copy
 	if (WriteBuffer.bAvailable)
@@ -91,7 +96,7 @@ void UE::PixelStreaming::FTextureSourceBackBuffer::OnBackBufferReady_RenderThrea
 
 			// For debugging timing information about the copy operation
 			// Turning it on all the time is a bit too much log spam if logging stats
-			uint64 PostWaitingOnCopy	= FPlatformTime::Cycles64();
+			uint64 PostWaitingOnCopy = FPlatformTime::Cycles64();
 			UE::PixelStreaming::FStats* Stats = UE::PixelStreaming::FStats::Get();
 			if (Stats)
 			{
@@ -99,7 +104,6 @@ void UE::PixelStreaming::FTextureSourceBackBuffer::OnBackBufferReady_RenderThrea
 				Stats->StoreApplicationStat(UE::PixelStreaming::FStatData(FName(*FString::Printf(TEXT("Layer (x%.2f) Capture time (ms)"), FrameScale)), CaptureLatencyMs, 2, true));
 			}
 		});
-
 	}
 }
 
@@ -116,19 +120,19 @@ FTexture2DRHIRef UE::PixelStreaming::FTextureSourceBackBuffer::GetTexture()
 
 void UE::PixelStreaming::FTextureSourceBackBuffer::Initialize(int Width, int Height)
 {
-	SourceWidth	 = Width;
+	SourceWidth = Width;
 	SourceHeight = Height;
 
 	for (auto& Buffer : WriteBuffers)
 	{
-		Buffer.Texture	  = UE::PixelStreaming::CreateTexture(SourceWidth, SourceHeight);
-		Buffer.Fence	  = GDynamicRHI->RHICreateGPUFence(TEXT("VideoCapturerCopyFence"));
+		Buffer.Texture = UE::PixelStreaming::CreateTexture(SourceWidth, SourceHeight);
+		Buffer.Fence = GDynamicRHI->RHICreateGPUFence(TEXT("VideoCapturerCopyFence"));
 		Buffer.bAvailable = true;
 	}
 	bWriteParity = true;
 
-	TempBuffer	 = UE::PixelStreaming::CreateTexture(SourceWidth, SourceHeight);
-	ReadBuffer	 = UE::PixelStreaming::CreateTexture(SourceWidth, SourceHeight);
+	TempBuffer = UE::PixelStreaming::CreateTexture(SourceWidth, SourceHeight);
+	ReadBuffer = UE::PixelStreaming::CreateTexture(SourceWidth, SourceHeight);
 	bIsTempDirty = false;
 
 	bInitialized = true;
