Fix a typo resulting in a crash in TextureSource.cpp

diff --git a/Engine/Plugins/Media/PixelStreaming/Source/PixelStreaming/Private/TextureSource.cpp b/Engine/Plugins/Media/PixelStreaming/Source/PixelStreaming/Private/TextureSource.cpp
--- a/Engine/Plugins/Media/PixelStreaming/Source/PixelStreaming/Private/TextureSource.cpp
+++ b/Engine/Plugins/Media/PixelStreaming/Source/PixelStreaming/Private/TextureSource.cpp
@@ -20,7 +20,7 @@ UE::PixelStreaming::FTextureSourceBackBuffer::FTextureSourceBackBuffer(float InF
 	// The backbuffer delegate can only be accessed using GameThread
 	AsyncTask(ENamedThreads::GameThread, [this, bEnabledClone]() {
 		/*Early exit if `this` died before game thread ran.*/
-		if (bEnabled.IsUnique())
+		if (bEnabledClone.IsUnique())
 		{
 			return;
 		}
