Remove PixelStreamer from delegates to better suit the naming conventions for the rest of the plugin

diff --git a/Engine/Plugins/Media/PixelStreaming/Source/PixelStreaming/Private/Settings.cpp b/Engine/Plugins/Media/PixelStreaming/Source/PixelStreaming/Private/Settings.cpp
--- a/Engine/Plugins/Media/PixelStreaming/Source/PixelStreaming/Private/Settings.cpp
+++ b/Engine/Plugins/Media/PixelStreaming/Source/PixelStreaming/Private/Settings.cpp
@@ -407,7 +407,7 @@ namespace UE::PixelStreaming::Settings
 		UPixelStreamingDelegates* Delegates = UPixelStreamingDelegates::GetPixelStreamingDelegates();
 		if (bLogStats && Delegates)
 		{
-			Delegates->OnPixelStreamerStatChangedNative.AddStatic(&UE::PixelStreaming::Settings::ConsumeStat);
+			Delegates->OnStatChangedNative.AddStatic(&UE::PixelStreaming::Settings::ConsumeStat);
 		}
 	}
 
diff --git a/Engine/Plugins/Media/PixelStreaming/Source/PixelStreaming/Private/Stats.cpp b/Engine/Plugins/Media/PixelStreaming/Source/PixelStreaming/Private/Stats.cpp
--- a/Engine/Plugins/Media/PixelStreaming/Source/PixelStreaming/Private/Stats.cpp
+++ b/Engine/Plugins/Media/PixelStreaming/Source/PixelStreaming/Private/Stats.cpp
@@ -111,7 +111,7 @@ void UE::PixelStreaming::FStats::FireStatChanged_GameThread(FPixelStreamingPlaye
 
 	if (UPixelStreamingDelegates* Delegates = UPixelStreamingDelegates::GetPixelStreamingDelegates())
 	{
-		Delegates->OnPixelStreamerStatChangedNative.Broadcast(PlayerId, StatName, StatValue);
+		Delegates->OnStatChangedNative.Broadcast(PlayerId, StatName, StatValue);
 	}
 }
 
diff --git a/Engine/Plugins/Media/PixelStreaming/Source/PixelStreaming/Public/PixelStreamingDelegates.h b/Engine/Plugins/Media/PixelStreaming/Source/PixelStreaming/Public/PixelStreamingDelegates.h
--- a/Engine/Plugins/Media/PixelStreaming/Source/PixelStreaming/Public/PixelStreamingDelegates.h
+++ b/Engine/Plugins/Media/PixelStreaming/Source/PixelStreaming/Public/PixelStreamingDelegates.h
@@ -100,8 +100,8 @@ public:
 	 * A pixel streaming stat has changed
 	 */
 	// C++ Delegate
-	DECLARE_MULTICAST_DELEGATE_ThreeParams(FPixelStreamerStatChanged, FPixelStreamingPlayerId, FName, float)
-	FPixelStreamerStatChanged OnPixelStreamerStatChangedNative;
+	DECLARE_MULTICAST_DELEGATE_ThreeParams(FStatChanged, FPixelStreamingPlayerId, FName, float)
+	FStatChanged OnStatChangedNative;
 
 	/**
 	 * Create the singleton.
