Remove the following unused console variable from Pixel Streaming:
* PixelStreamingUseBackBufferCaptureSize
* PixelStreamingCaptureSize
* PixelStreamingWebRTCDisableResolutionChange
* PixelStreamingMaxNumBackBuffers

diff --git a/Engine/Plugins/Media/PixelStreaming/Source/PixelStreaming/Private/EncoderFrameFactory.cpp b/Engine/Plugins/Media/PixelStreaming/Source/PixelStreaming/Private/EncoderFrameFactory.cpp
--- a/Engine/Plugins/Media/PixelStreaming/Source/PixelStreaming/Private/EncoderFrameFactory.cpp
+++ b/Engine/Plugins/Media/PixelStreaming/Source/PixelStreaming/Private/EncoderFrameFactory.cpp
@@ -99,7 +99,6 @@ TSharedPtr<AVEncoder::FVideoEncoderInput> UE::PixelStreaming::FEncoderFrameFacto
 	if (!EncoderInput.IsValid())
 	{
 		EncoderInput = CreateVideoEncoderInput(InWidth, InHeight);
-		EncoderInput->SetMaxNumBuffers((uint32)UE::PixelStreaming::Settings::CVarPixelStreamingMaxNumBackBuffers.GetValueOnAnyThread());
 	}
 
 	return EncoderInput;
diff --git a/Engine/Plugins/Media/PixelStreaming/Source/PixelStreaming/Private/Settings.cpp b/Engine/Plugins/Media/PixelStreaming/Source/PixelStreaming/Private/Settings.cpp
--- a/Engine/Plugins/Media/PixelStreaming/Source/PixelStreaming/Private/Settings.cpp
+++ b/Engine/Plugins/Media/PixelStreaming/Source/PixelStreaming/Private/Settings.cpp
@@ -109,26 +109,6 @@ namespace UE::PixelStreaming::Settings
 		ECVF_Default);
 	// End Encoder CVars
 
-	// Begin Capturer CVars
-	TAutoConsoleVariable<bool> CVarPixelStreamingUseBackBufferCaptureSize(
-		TEXT("PixelStreaming.Capturer.UseBackBufferSize"),
-		true,
-		TEXT("Whether to use back buffer size or custom size"),
-		ECVF_Default);
-
-	TAutoConsoleVariable<FString> CVarPixelStreamingCaptureSize(
-		TEXT("PixelStreaming.Capturer.CaptureSize"),
-		TEXT("1920x1080"),
-		TEXT("Capture size in format widthxheight. Recommended to UseBackBufferSize instead."),
-		ECVF_Default);
-
-	TAutoConsoleVariable<int32> CVarPixelStreamingMaxNumBackBuffers(
-		TEXT("PixelStreaming.Capturer.MaxNumBackBuffers"),
-		3,
-		TEXT("Maximum number of back buffers to use. A value of 0 will not limit the number of back buffers."),
-		ECVF_Default);
-	// End Capturer CVars
-
 	// Begin WebRTC CVars
 	TAutoConsoleVariable<FString> CVarPixelStreamingDegradationPreference(
 		TEXT("PixelStreaming.WebRTC.DegradationPreference"),
@@ -197,12 +177,6 @@ namespace UE::PixelStreaming::Settings
 		TEXT("Whether put audio and video in the same stream (which will make WebRTC try to sync them)."),
 		ECVF_Default);
 
-	TAutoConsoleVariable<bool> CVarPixelStreamingWebRTCDisableResolutionChange(
-		TEXT("PixelStreaming.WebRTC.DisableResolutionChange"),
-		true,
-		TEXT("Disables whether WebRTC can trigger a resolution a change of the stream (regardless of it wants to)."),
-		ECVF_Default);
-
 	TAutoConsoleVariable<bool> CVarPixelStreamingWebRTCDisableStats(
 		TEXT("PixelStreaming.WebRTC.DisableStats"),
 		false,
@@ -438,8 +412,6 @@ namespace UE::PixelStreaming::Settings
 		CommandLineParseValue(TEXT("PixelStreamingEncoderRateControl="), UE::PixelStreaming::Settings::CVarPixelStreamingEncoderRateControl);
 		CommandLineParseValue(TEXT("PixelStreamingEncoderMultipass="), UE::PixelStreaming::Settings::CVarPixelStreamingEncoderMultipass);
 		CommandLineParseValue(TEXT("PixelStreamingH264Profile="), UE::PixelStreaming::Settings::CVarPixelStreamingH264Profile);
-		CommandLineParseValue(TEXT("PixelStreamingCaptureSize="), UE::PixelStreaming::Settings::CVarPixelStreamingCaptureSize);
-		CommandLineParseValue(TEXT("PixelStreamingMaxNumBackBuffers="), UE::PixelStreaming::Settings::CVarPixelStreamingMaxNumBackBuffers);
 		CommandLineParseValue(TEXT("PixelStreamingDegradationPreference="), UE::PixelStreaming::Settings::CVarPixelStreamingDegradationPreference);
 		CommandLineParseValue(TEXT("PixelStreamingWebRTCDegradationPreference="), UE::PixelStreaming::Settings::CVarPixelStreamingDegradationPreference);
 		CommandLineParseValue(TEXT("PixelStreamingWebRTCFps="), UE::PixelStreaming::Settings::CVarPixelStreamingWebRTCFps);
@@ -458,12 +430,10 @@ namespace UE::PixelStreaming::Settings
 
 		CommandLineParseOption(TEXT("PixelStreamingDebugDumpFrame"), UE::PixelStreaming::Settings::CVarPixelStreamingDebugDumpFrame);
 		CommandLineParseOption(TEXT("PixelStreamingEnableFillerData"), UE::PixelStreaming::Settings::CVarPixelStreamingEnableFillerData);
-		CommandLineParseOption(TEXT("PixelStreamingUseBackBufferSize"), UE::PixelStreaming::Settings::CVarPixelStreamingUseBackBufferCaptureSize);
 		CommandLineParseOption(TEXT("PixelStreamingWebRTCDisableStats"), UE::PixelStreaming::Settings::CVarPixelStreamingWebRTCDisableStats);
 		CommandLineParseOption(TEXT("PixelStreamingWebRTCDisableReceiveAudio"), UE::PixelStreaming::Settings::CVarPixelStreamingWebRTCDisableReceiveAudio);
 		CommandLineParseOption(TEXT("PixelStreamingWebRTCDisableTransmitAudio"), UE::PixelStreaming::Settings::CVarPixelStreamingWebRTCDisableTransmitAudio);
 		CommandLineParseOption(TEXT("PixelStreamingWebRTCDisableAudioSync"), UE::PixelStreaming::Settings::CVarPixelStreamingWebRTCDisableAudioSync);
-		CommandLineParseOption(TEXT("PixelStreamingWebRTCDisableResolutionChange"), UE::PixelStreaming::Settings::CVarPixelStreamingWebRTCDisableResolutionChange);
 		CommandLineParseOption(TEXT("PixelStreamingSendPlayerIdAsInteger"), UE::PixelStreaming::Settings::CVarSendPlayerIdAsInteger);
 		CommandLineParseOption(TEXT("PixelStreamingWebRTCUseLegacyAudioDevice"), UE::PixelStreaming::Settings::CVarPixelStreamingWebRTCUseLegacyAudioDevice);
 		CommandLineParseOption(TEXT("PixelStreamingDisableLatencyTester"), UE::PixelStreaming::Settings::CVarPixelStreamingDisableLatencyTester);
diff --git a/Engine/Plugins/Media/PixelStreaming/Source/PixelStreaming/Private/Settings.h b/Engine/Plugins/Media/PixelStreaming/Source/PixelStreaming/Private/Settings.h
--- a/Engine/Plugins/Media/PixelStreaming/Source/PixelStreaming/Private/Settings.h
+++ b/Engine/Plugins/Media/PixelStreaming/Source/PixelStreaming/Private/Settings.h
@@ -30,12 +30,6 @@ namespace UE {
 			extern TAutoConsoleVariable<int32> CVarPixelStreamingEncoderKeyframeInterval;
 			// End Encoder CVars
 
-			// Begin Capturer CVars
-			extern TAutoConsoleVariable<bool> CVarPixelStreamingUseBackBufferCaptureSize;
-			extern TAutoConsoleVariable<FString> CVarPixelStreamingCaptureSize;
-			extern TAutoConsoleVariable<int32> CVarPixelStreamingMaxNumBackBuffers;
-			// End Capturer CVars
-
 			// Begin WebRTC CVars
 			extern TAutoConsoleVariable<FString> CVarPixelStreamingDegradationPreference;
 			extern TAutoConsoleVariable<int32> CVarPixelStreamingWebRTCFps;
@@ -48,7 +42,6 @@ namespace UE {
 			extern TAutoConsoleVariable<bool> CVarPixelStreamingWebRTCDisableTransmitAudio;
 			extern TAutoConsoleVariable<bool> CVarPixelStreamingWebRTCDisableAudioSync;
 			extern TAutoConsoleVariable<bool> CVarPixelStreamingWebRTCUseLegacyAudioDevice;
-			extern TAutoConsoleVariable<bool> CVarPixelStreamingWebRTCDisableResolutionChange;
 			extern TAutoConsoleVariable<bool> CVarPixelStreamingWebRTCDisableStats;
 			// End WebRTC CVars
 
