Added experimental support for VP8 and VP9 software encoding.
VP8 runs at about 30fps and VP9 runs at about 15fps currently.
Encoder codec selection can now be made with -PixelStreamingEncoderCodec where supported values are "H264", "VP8", and "VP9".
The default encoder remains hardware accelerated H264.

diff --git a/Engine/Plugins/Media/PixelStreaming/Source/PixelStreaming/Private/EncoderFactory.cpp b/Engine/Plugins/Media/PixelStreaming/Source/PixelStreaming/Private/EncoderFactory.cpp
--- a/Engine/Plugins/Media/PixelStreaming/Source/PixelStreaming/Private/EncoderFactory.cpp
+++ b/Engine/Plugins/Media/PixelStreaming/Source/PixelStreaming/Private/EncoderFactory.cpp
@@ -5,6 +5,7 @@
 #include "Settings.h"
 #include "absl/strings/match.h"
 #include "modules/video_coding/codecs/vp8/include/vp8.h"
+#include "modules/video_coding/codecs/vp9/include/vp9.h"
 #include "Misc/ScopeLock.h"
 #include "VideoEncoderH264Wrapper.h"
 #include "SimulcastEncoderAdapter.h"
@@ -70,18 +71,18 @@ webrtc::SdpVideoFormat UE::PixelStreaming::FVideoEncoderFactory::CreateH264Forma
 
 std::vector<webrtc::SdpVideoFormat> UE::PixelStreaming::FVideoEncoderFactory::GetSupportedFormats() const
 {
-	const bool bForceVP8 = UE::PixelStreaming::Settings::IsForceVP8();
-
 	std::vector<webrtc::SdpVideoFormat> video_formats;
-	if (bForceVP8)
-	{
-		video_formats.push_back(webrtc::SdpVideoFormat(cricket::kVp8CodecName));
-	}
-	else
+
+	switch (UE::PixelStreaming::Settings::GetSelectedCodec())
 	{
-		video_formats.push_back(UE::PixelStreaming::FVideoEncoderFactory::CreateH264Format(webrtc::H264::kProfileConstrainedBaseline, webrtc::H264::kLevel3_1));
+		case UE::PixelStreaming::Settings::ECodec::VP8:
+			video_formats.push_back(webrtc::SdpVideoFormat(cricket::kVp8CodecName));
+		case UE::PixelStreaming::Settings::ECodec::VP9:
+			video_formats.push_back(webrtc::SdpVideoFormat(cricket::kVp9CodecName));
+		case UE::PixelStreaming::Settings::ECodec::H264:
+		default:
+			video_formats.push_back(UE::PixelStreaming::FVideoEncoderFactory::CreateH264Format(webrtc::H264::kProfileConstrainedBaseline, webrtc::H264::kLevel3_1));
 	}
-
 	return video_formats;
 }
 
@@ -96,7 +97,13 @@ UE::PixelStreaming::FVideoEncoderFactory::CodecInfo UE::PixelStreaming::FVideoEn
 std::unique_ptr<webrtc::VideoEncoder> UE::PixelStreaming::FVideoEncoderFactory::CreateVideoEncoder(const webrtc::SdpVideoFormat& format)
 {
 	if (absl::EqualsIgnoreCase(format.name, cricket::kVp8CodecName))
+	{
 		return webrtc::VP8Encoder::Create();
+	}
+	else if (absl::EqualsIgnoreCase(format.name, cricket::kVp9CodecName))
+	{
+		return webrtc::VP9Encoder::Create();
+	}
 	else
 	{
 		// Lock during encoder creation
@@ -148,7 +155,7 @@ TOptional<UE::PixelStreaming::FVideoEncoderFactory::FHardwareEncoderId> UE::Pixe
 {
 	FHardwareEncoderId EncoderId = HashResolution(Width, Height);
 	UE::PixelStreaming::FVideoEncoderH264Wrapper* Existing = GetHardwareEncoder(EncoderId);
-	
+
 	TOptional<FHardwareEncoderId> Result = EncoderId;
 
 	if (Existing == nullptr)
@@ -168,7 +175,7 @@ TOptional<UE::PixelStreaming::FVideoEncoderFactory::FHardwareEncoderId> UE::Pixe
 		// Make the actual AVEncoder encoder.
 		const TArray<AVEncoder::FVideoEncoderInfo>& Available = AVEncoder::FVideoEncoderFactory::Get().GetAvailable();
 		TUniquePtr<AVEncoder::FVideoEncoder> Encoder = AVEncoder::FVideoEncoderFactory::Get().Create(Available[0].ID, FrameFactory->GetOrCreateVideoEncoderInput(EncoderConfig.Width, EncoderConfig.Height), EncoderConfig);
-		if(Encoder.IsValid())
+		if (Encoder.IsValid())
 		{
 			Encoder->SetOnEncodedPacket([EncoderId, this](uint32 InLayerIndex, const AVEncoder::FVideoEncoderInputFrame* InFrame, const AVEncoder::FCodecPacket& InPacket) {
 				// Note: this is a static method call.
diff --git a/Engine/Plugins/Media/PixelStreaming/Source/PixelStreaming/Private/FrameBuffer.cpp b/Engine/Plugins/Media/PixelStreaming/Source/PixelStreaming/Private/FrameBuffer.cpp
--- a/Engine/Plugins/Media/PixelStreaming/Source/PixelStreaming/Private/FrameBuffer.cpp
+++ b/Engine/Plugins/Media/PixelStreaming/Source/PixelStreaming/Private/FrameBuffer.cpp
@@ -2,6 +2,7 @@
 
 #include "FrameBuffer.h"
 #include "TextureSource.h"
+#include "Utils.h"
 
 /*
 * ----------------- FInitializeFrameBuffer -----------------
@@ -53,13 +54,13 @@ TSharedPtr<UE::PixelStreaming::ITextureSource> UE::PixelStreaming::FSimulcastFra
 int UE::PixelStreaming::FSimulcastFrameBuffer::width() const
 {
 	checkf(TextureSources.Num() > 0, TEXT("Must be at least one texture source to get the width from."));
-	return FMath::Max(TextureSources[0]->GetSourceWidth(), TextureSources[TextureSources.Num()-1]->GetSourceWidth());
+	return FMath::Max(TextureSources[0]->GetSourceWidth(), TextureSources[TextureSources.Num() - 1]->GetSourceWidth());
 }
 
 int UE::PixelStreaming::FSimulcastFrameBuffer::height() const
 {
 	checkf(TextureSources.Num() > 0, TEXT("Must be at least one texture source to get the height from."));
-	return FMath::Max(TextureSources[0]->GetSourceHeight(), TextureSources[TextureSources.Num()-1]->GetSourceHeight());
+	return FMath::Max(TextureSources[0]->GetSourceHeight(), TextureSources[TextureSources.Num() - 1]->GetSourceHeight());
 }
 
 /*
@@ -89,3 +90,95 @@ int UE::PixelStreaming::FLayerFrameBuffer::height() const
 {
 	return TextureSource->GetSourceHeight();
 }
+
+/*
+* ----------------- FFrameBufferI420 -----------------
+*/
+
+UE::PixelStreaming::FFrameBufferI420::FFrameBufferI420(TSharedPtr<UE::PixelStreaming::ITextureSource> InTextureSource)
+	: TextureSource(InTextureSource)
+{
+}
+
+UE::PixelStreaming::FFrameBufferI420::~FFrameBufferI420()
+{
+}
+
+FTexture2DRHIRef UE::PixelStreaming::FFrameBufferI420::GetFrame() const
+{
+	return TextureSource->GetTexture();
+}
+
+int UE::PixelStreaming::FFrameBufferI420::width() const
+{
+	return TextureSource->GetSourceWidth();
+}
+
+int UE::PixelStreaming::FFrameBufferI420::height() const
+{
+	return TextureSource->GetSourceHeight();
+}
+
+const webrtc::I420BufferInterface* UE::PixelStreaming::FFrameBufferI420::GetI420() const
+{
+	return Buffer;
+}
+
+/*
+* NOTE: Only used for non-hardware encoders (e.g. VP8) - H264 does NOT hit this function.
+*/
+rtc::scoped_refptr<webrtc::I420BufferInterface> UE::PixelStreaming::FFrameBufferI420::ToI420()
+{
+	ITextureSource* TextureSourcePtr = TextureSource.Get();
+	checkf(TextureSourcePtr && TextureSourcePtr->GetName() == FString(TEXT("FTextureSourceBackBufferToCPU")), TEXT("To use VPX encoder texture source must be FTextureSourceBackBufferToCPU"));
+
+	FTextureSourceBackBufferToCPU* TextureSourceCPU = static_cast<FTextureSourceBackBufferToCPU*>(TextureSourcePtr);
+	TRefCountPtr<FRawPixelsTexture> Current = TextureSourceCPU->GetCurrent();
+
+	uint32 TextureWidth = Current->TextureRef->GetSizeX();
+	uint32 TextureHeight = Current->TextureRef->GetSizeY();
+
+	uint32 i = 0;
+	uint32 NumPixels = TextureWidth * TextureHeight;
+	uint32 ui = NumPixels;
+	uint32 vi = NumPixels + NumPixels / 4;
+	uint32 s = 0;
+
+// data is coming in as B,G,R,A
+#define sR ((uint32)(rgbIn[s + 2]))
+#define sG ((uint32)(rgbIn[s + 1]))
+#define sB ((uint32)(rgbIn[s + 0]))
+
+	Buffer = webrtc::I420Buffer::Create(TextureWidth, TextureHeight, TextureWidth, TextureWidth / 2, TextureWidth / 2);
+
+	// If we have no raw pixels then early exit this swizzle operation
+	if (Current->RawPixels.Num() == 0)
+	{
+		webrtc::I420Buffer::SetBlack(Buffer);
+		return Buffer;
+	}
+
+	uint8* yuv420p = Buffer->MutableDataY();
+	const uint8* rgbIn = (const uint8*)Current->RawPixels.GetData();
+	for (uint32 j = 0; j < TextureHeight; j++)
+	{
+		for (uint32 k = 0; k < TextureWidth; k++)
+		{
+			yuv420p[i] = (uint8)((66 * sR + 129 * sG + 25 * sB + 128) >> 8) + 16;
+
+			if (0 == j % 2 && 0 == k % 2)
+			{
+				yuv420p[ui++] = (uint8)((-38 * sR - 74 * sG + 112 * sB + 128) >> 8) + 128;
+				yuv420p[vi++] = (uint8)((112 * sR - 94 * sG - 18 * sB + 128) >> 8) + 128;
+			}
+			i++;
+			s += 4;
+		}
+	}
+
+#undef sR
+#undef sG
+#undef sB
+
+	return Buffer;
+}
diff --git a/Engine/Plugins/Media/PixelStreaming/Source/PixelStreaming/Private/FrameBuffer.h b/Engine/Plugins/Media/PixelStreaming/Source/PixelStreaming/Private/FrameBuffer.h
--- a/Engine/Plugins/Media/PixelStreaming/Source/PixelStreaming/Private/FrameBuffer.h
+++ b/Engine/Plugins/Media/PixelStreaming/Source/PixelStreaming/Private/FrameBuffer.h
@@ -6,8 +6,10 @@
 #include "WebRTCIncludes.h"
 #include "TextureSource.h"
 
-namespace UE {
-	namespace PixelStreaming {
+namespace UE
+{
+	namespace PixelStreaming
+	{
 
 		enum FFrameBufferType
 		{
@@ -110,5 +112,29 @@ namespace UE {
 		private:
 			TSharedPtr<ITextureSource> TextureSource;
 		};
-	}
-}
+
+		/*
+		* ----------------- FFrameBufferI420 -----------------
+		* Holds a texture source that is capable of being converted to I420 for encoding.
+		*/
+		class FFrameBufferI420 : public FFrameBuffer
+		{
+		public:
+			FFrameBufferI420(TSharedPtr<ITextureSource> InTextureSource);
+			virtual ~FFrameBufferI420();
+
+			virtual FFrameBufferType GetFrameBufferType() const { return Layer; }
+			FTexture2DRHIRef GetFrame() const;
+
+			virtual int width() const override;
+			virtual int height() const override;
+			virtual rtc::scoped_refptr<webrtc::I420BufferInterface> ToI420() override;
+			virtual const webrtc::I420BufferInterface* GetI420() const override;
+
+		private:
+			TSharedPtr<ITextureSource> TextureSource;
+			rtc::scoped_refptr<webrtc::I420Buffer> Buffer;
+		};
+
+	} // namespace PixelStreaming
+} // namespace UE
diff --git a/Engine/Plugins/Media/PixelStreaming/Source/PixelStreaming/Private/Settings.cpp b/Engine/Plugins/Media/PixelStreaming/Source/PixelStreaming/Private/Settings.cpp
--- a/Engine/Plugins/Media/PixelStreaming/Source/PixelStreaming/Private/Settings.cpp
+++ b/Engine/Plugins/Media/PixelStreaming/Source/PixelStreaming/Private/Settings.cpp
@@ -107,6 +107,12 @@ namespace UE::PixelStreaming::Settings
 		300,
 		TEXT("How many frames before a key frame is sent. Default: 300. Values <=0 will disable sending of periodic key frames. Note: NVENC does not support changing this after encoding has started."),
 		ECVF_Default);
+
+	TAutoConsoleVariable<FString> CVarPixelStreamingEncoderCodec(
+		TEXT("PixelStreaming.Encoder.Codec"),
+		TEXT("H264"),
+		TEXT("PixelStreaming encoder codec. Supported values are `H264`, `VP8`, `VP9`"),
+		ECVF_Default);
 	// End Encoder CVars
 
 	// Begin WebRTC CVars
@@ -367,6 +373,36 @@ namespace UE::PixelStreaming::Settings
 		}
 	}
 
+	/*
+	* Selected Codec.
+	*/
+	ECodec GetSelectedCodec()
+	{
+		const FString CodecStr = CVarPixelStreamingEncoderCodec.GetValueOnAnyThread();
+		if (CodecStr == TEXT("H264"))
+		{
+			return ECodec::H264;
+		}
+		else if (CodecStr == TEXT("VP8"))
+		{
+			return ECodec::VP8;
+		}
+		else if (CodecStr == TEXT("VP9"))
+		{
+			return ECodec::VP9;
+		}
+		else
+		{
+			return ECodec::H264;
+		}
+	}
+
+	bool IsCodecVPX()
+	{
+		ECodec SelectedCodec = GetSelectedCodec();
+		return SelectedCodec == ECodec::VP8 || SelectedCodec == ECodec::VP9;
+	}
+
 	/*
 	* Stats logger - as turned on/off by CVarPixelStreamingLogStats
 	*/
@@ -411,6 +447,7 @@ namespace UE::PixelStreaming::Settings
 		CommandLineParseValue(TEXT("PixelStreamingEncoderMaxQP="), UE::PixelStreaming::Settings::CVarPixelStreamingEncoderMaxQP);
 		CommandLineParseValue(TEXT("PixelStreamingEncoderRateControl="), UE::PixelStreaming::Settings::CVarPixelStreamingEncoderRateControl);
 		CommandLineParseValue(TEXT("PixelStreamingEncoderMultipass="), UE::PixelStreaming::Settings::CVarPixelStreamingEncoderMultipass);
+		CommandLineParseValue(TEXT("PixelStreamingEncoderCodec="), UE::PixelStreaming::Settings::CVarPixelStreamingEncoderCodec);
 		CommandLineParseValue(TEXT("PixelStreamingH264Profile="), UE::PixelStreaming::Settings::CVarPixelStreamingH264Profile);
 		CommandLineParseValue(TEXT("PixelStreamingDegradationPreference="), UE::PixelStreaming::Settings::CVarPixelStreamingDegradationPreference);
 		CommandLineParseValue(TEXT("PixelStreamingWebRTCDegradationPreference="), UE::PixelStreaming::Settings::CVarPixelStreamingDegradationPreference);
diff --git a/Engine/Plugins/Media/PixelStreaming/Source/PixelStreaming/Private/Settings.h b/Engine/Plugins/Media/PixelStreaming/Source/PixelStreaming/Private/Settings.h
--- a/Engine/Plugins/Media/PixelStreaming/Source/PixelStreaming/Private/Settings.h
+++ b/Engine/Plugins/Media/PixelStreaming/Source/PixelStreaming/Private/Settings.h
@@ -10,8 +10,10 @@
 #include "VideoEncoder.h"
 #include "WebRTCIncludes.h"
 
-namespace UE {
-	namespace PixelStreaming {
+namespace UE
+{
+	namespace PixelStreaming
+	{
 		// Console variables (CVars)
 		namespace Settings
 		{
@@ -28,6 +30,7 @@ namespace UE {
 			extern TAutoConsoleVariable<FString> CVarPixelStreamingEncoderMultipass;
 			extern TAutoConsoleVariable<FString> CVarPixelStreamingH264Profile;
 			extern TAutoConsoleVariable<int32> CVarPixelStreamingEncoderKeyframeInterval;
+			extern TAutoConsoleVariable<FString> CVarPixelStreamingEncoderCodec;
 			// End Encoder CVars
 
 			// Begin WebRTC CVars
@@ -56,7 +59,16 @@ namespace UE {
 			extern TArray<FKey> FilteredKeys;
 			// Ends Pixel Streaming Plugin CVars
 
+			enum ECodec
+			{
+				H264,
+				VP8,
+				VP9
+			};
+
 			// Begin utility functions etc.
+			bool IsCodecVPX();
+			ECodec GetSelectedCodec();
 			AVEncoder::FVideoEncoder::RateControlMode GetRateControlCVar();
 			AVEncoder::FVideoEncoder::MultipassMode GetMultipassCVar();
 			webrtc::DegradationPreference GetDegradationPreference();
@@ -84,11 +96,6 @@ namespace UE {
 				return FParse::Param(FCommandLine::Get(), TEXT("PixelStreamingHideCursor"));
 			}
 
-			inline bool IsForceVP8()
-			{
-				return FParse::Param(FCommandLine::Get(), TEXT("PSForceVP8"));
-			}
-
 			inline bool GetSignallingServerIP(FString& OutSignallingServerIP)
 			{
 				return FParse::Value(FCommandLine::Get(), TEXT("PixelStreamingIP="), OutSignallingServerIP);
@@ -111,6 +118,6 @@ namespace UE {
 
 			// End Command line args
 
-		} // namespace UE::PixelStreaming::Settings
-	}
-}
+		} // namespace Settings
+	}	  // namespace PixelStreaming
+} // namespace UE
diff --git a/Engine/Plugins/Media/PixelStreaming/Source/PixelStreaming/Private/TextureSource.cpp b/Engine/Plugins/Media/PixelStreaming/Source/PixelStreaming/Private/TextureSource.cpp
--- a/Engine/Plugins/Media/PixelStreaming/Source/PixelStreaming/Private/TextureSource.cpp
+++ b/Engine/Plugins/Media/PixelStreaming/Source/PixelStreaming/Private/TextureSource.cpp
@@ -1,139 +1,69 @@
 // Copyright Epic Games, Inc. All Rights Reserved.
 
 #include "TextureSource.h"
-#include "Stats.h"
 #include "Utils.h"
-#include "Async/Async.h"
-#include "GPUFencePoller.h"
-#include "Framework/Application/SlateApplication.h"
+
 
 /*
 * -------- UE::PixelStreaming::FTextureSourceBackBuffer ----------------
 */
 
-UE::PixelStreaming::FTextureSourceBackBuffer::FTextureSourceBackBuffer(float InFrameScale)
-	: FrameScale(InFrameScale)
+UE::PixelStreaming::FTextureSourceBackBuffer::FTextureSourceBackBuffer() 
+	: TTextureSourceBackBufferBase()
 {
-	// Explictly make clear we are adding another ref to this shared bool for the purposes of using in the below lambda
-	TSharedRef<bool, ESPMode::ThreadSafe> bEnabledClone = bEnabled;
-
-	// The backbuffer delegate can only be accessed using GameThread
-	AsyncTask(ENamedThreads::GameThread, [this, bEnabledClone]() {
-		/*Early exit if `this` died before game thread ran.*/
-		if (bEnabledClone.IsUnique())
-		{
-			return;
-		}
-		OnBackbuffer = &FSlateApplication::Get().GetRenderer()->OnBackBufferReadyToPresent();
-		BackbufferDelegateHandle = OnBackbuffer->AddSP(this, &UE::PixelStreaming::FTextureSourceBackBuffer::OnBackBufferReady_RenderThread);
-	});
-}
 
-UE::PixelStreaming::FTextureSourceBackBuffer::FTextureSourceBackBuffer()
-	: UE::PixelStreaming::FTextureSourceBackBuffer(1.0) {}
+};
 
-UE::PixelStreaming::FTextureSourceBackBuffer::~FTextureSourceBackBuffer()
+UE::PixelStreaming::FTextureSourceBackBuffer::FTextureSourceBackBuffer(float InScale) 
+	: TTextureSourceBackBufferBase(InScale)
 {
-	if (OnBackbuffer)
-	{
-		OnBackbuffer->Remove(BackbufferDelegateHandle);
-	}
 
-	SetEnabled(false);
+};
+
+void UE::PixelStreaming::FTextureSourceBackBuffer::CopyTexture(const FTexture2DRHIRef& SourceTexture, TRefCountPtr<FRHITexture2D> DestTexture, FGPUFenceRHIRef& CopyFence)
+{
+	UE::PixelStreaming::CopyTexture(SourceTexture, DestTexture, CopyFence);
 }
 
-void UE::PixelStreaming::FTextureSourceBackBuffer::SetEnabled(bool bInEnabled)
+TRefCountPtr<FRHITexture2D> UE::PixelStreaming::FTextureSourceBackBuffer::CreateTexture(int Width, int Height)
 {
-	*bEnabled = bInEnabled;
-
-	// This source has been disabled, so set `bInitialized` to false so `OnBackBufferReady_RenderThread`
-	// will make new textures next time it is called.
-	if (bInitialized && bInEnabled == false)
-	{
-		bInitialized = false;
-	}
+	return UE::PixelStreaming::CreateTexture(Width, Height);
 }
 
-void UE::PixelStreaming::FTextureSourceBackBuffer::OnBackBufferReady_RenderThread(SWindow& SlateWindow, const FTexture2DRHIRef& FrameBuffer)
+FTexture2DRHIRef UE::PixelStreaming::FTextureSourceBackBuffer::ToTextureRef(TRefCountPtr<FRHITexture2D> Texture)
 {
-	if (!bInitialized)
-	{
-		Initialize(FrameBuffer->GetSizeXY().X * FrameScale, FrameBuffer->GetSizeXY().Y * FrameScale);
-	}
-
-	if (!IsEnabled())
-	{
-		return;
-	}
-
-	auto& WriteBuffer = bWriteParity ? WriteBuffers[0] : WriteBuffers[1];
-	bWriteParity = !bWriteParity;
-
-	// for safety we just make sure that the buffer is not currently waiting for a copy
-	if (WriteBuffer.bAvailable)
-	{
-		WriteBuffer.bAvailable = false;
-
-		FRHICommandListImmediate& RHICmdList = FRHICommandListExecutor::GetImmediateCommandList();
-		WriteBuffer.Fence->Clear();
-
-		RHICmdList.EnqueueLambda([&WriteBuffer](FRHICommandListImmediate& RHICmdList) {
-			WriteBuffer.PreWaitingOnCopy = FPlatformTime::Cycles64();
-		});
-
-		UE::PixelStreaming::CopyTexture(FrameBuffer, WriteBuffer.Texture, WriteBuffer.Fence);
-
-		UE::PixelStreaming::FGPUFencePoller::Get()->AddJob(WriteBuffer.Fence, bEnabled, [this, &WriteBuffer]() {
-			// This lambda is called only once the GPUFence is done
-			{
-				FScopeLock Lock(&CriticalSection);
-				TempBuffer.Swap(WriteBuffer.Texture);
-				WriteBuffer.Fence->Clear();
-				WriteBuffer.bAvailable = true;
-
-				bIsTempDirty = true;
-			}
-
-			// For debugging timing information about the copy operation
-			// Turning it on all the time is a bit too much log spam if logging stats
-			uint64 PostWaitingOnCopy = FPlatformTime::Cycles64();
-			UE::PixelStreaming::FStats* Stats = UE::PixelStreaming::FStats::Get();
-			if (Stats)
-			{
-				double CaptureLatencyMs = FPlatformTime::ToMilliseconds64(PostWaitingOnCopy - WriteBuffer.PreWaitingOnCopy);
-				Stats->StoreApplicationStat(UE::PixelStreaming::FStatData(FName(*FString::Printf(TEXT("Layer (x%.2f) Capture time (ms)"), FrameScale)), CaptureLatencyMs, 2, true));
-			}
-		});
-	}
+	return Texture;
 }
 
-FTexture2DRHIRef UE::PixelStreaming::FTextureSourceBackBuffer::GetTexture()
+/*
+* -------- UE::PixelStreaming::FTextureSourceBackBufferToCPU ----------------
+*/
+
+UE::PixelStreaming::FTextureSourceBackBufferToCPU::FTextureSourceBackBufferToCPU() 
+	: TTextureSourceBackBufferBase()
 {
-	if (bIsTempDirty)
-	{
-		FScopeLock Lock(&CriticalSection);
-		ReadBuffer.Swap(TempBuffer);
-		bIsTempDirty = false;
-	}
-	return ReadBuffer;
+
+};
+
+UE::PixelStreaming::FTextureSourceBackBufferToCPU::FTextureSourceBackBufferToCPU(float InScale) 
+	: TTextureSourceBackBufferBase(InScale)
+{
+
+};
+
+void UE::PixelStreaming::FTextureSourceBackBufferToCPU::CopyTexture(const FTexture2DRHIRef& SourceTexture, TRefCountPtr<UE::PixelStreaming::FRawPixelsTexture> DestTexture, FGPUFenceRHIRef& CopyFence)
+{
+	UE::PixelStreaming::CopyTexture(SourceTexture, DestTexture->TextureRef, CopyFence);
+	UE::PixelStreaming::ReadTextureToCPU(FRHICommandListExecutor::GetImmediateCommandList(), DestTexture->TextureRef, DestTexture->RawPixels);
 }
 
-void UE::PixelStreaming::FTextureSourceBackBuffer::Initialize(int Width, int Height)
+TRefCountPtr<UE::PixelStreaming::FRawPixelsTexture> UE::PixelStreaming::FTextureSourceBackBufferToCPU::CreateTexture(int Width, int Height)
 {
-	SourceWidth = Width;
-	SourceHeight = Height;
-
-	for (auto& Buffer : WriteBuffers)
-	{
-		Buffer.Texture = UE::PixelStreaming::CreateTexture(SourceWidth, SourceHeight);
-		Buffer.Fence = GDynamicRHI->RHICreateGPUFence(TEXT("VideoCapturerCopyFence"));
-		Buffer.bAvailable = true;
-	}
-	bWriteParity = true;
-
-	TempBuffer = UE::PixelStreaming::CreateTexture(SourceWidth, SourceHeight);
-	ReadBuffer = UE::PixelStreaming::CreateTexture(SourceWidth, SourceHeight);
-	bIsTempDirty = false;
-
-	bInitialized = true;
+	FRawPixelsTexture* Tex = new FRawPixelsTexture(UE::PixelStreaming::CreateTexture(Width, Height));
+	return TRefCountPtr<FRawPixelsTexture>(Tex);
 }
+
+FTexture2DRHIRef UE::PixelStreaming::FTextureSourceBackBufferToCPU::ToTextureRef(TRefCountPtr<UE::PixelStreaming::FRawPixelsTexture> Texture)
+{
+	return Texture->TextureRef;
+}
\ No newline at end of file
diff --git a/Engine/Plugins/Media/PixelStreaming/Source/PixelStreaming/Private/TextureSource.h b/Engine/Plugins/Media/PixelStreaming/Source/PixelStreaming/Private/TextureSource.h
--- a/Engine/Plugins/Media/PixelStreaming/Source/PixelStreaming/Private/TextureSource.h
+++ b/Engine/Plugins/Media/PixelStreaming/Source/PixelStreaming/Private/TextureSource.h
@@ -1,57 +1,209 @@
 // Copyright Epic Games, Inc. All Rights Reserved.
 #pragma once
 
+#include "Templates/SharedPointer.h"
+#include "Templates/RefCounting.h"
+#include "Async/Async.h"
 #include "RHI.h"
 #include "HAL/ThreadSafeBool.h"
 #include "Rendering/SlateRenderer.h"
-#include "Templates/SharedPointer.h"
-
-namespace UE {
-	namespace PixelStreaming {
+#include "GPUFencePoller.h"
+#include "Framework/Application/SlateApplication.h"
+#include "Stats.h"
 
+namespace UE
+{
+	namespace PixelStreaming
+	{
 		/*
 		* Interface for all texture sources in Pixel Streaming.
 		* These texture sources are used to populate video sources for video tracks.
 		*/
 		class ITextureSource
 		{
-			public:
-				ITextureSource() = default;
-				virtual ~ITextureSource() = default;
-				virtual bool IsAvailable() const = 0;
-				virtual bool IsEnabled() const = 0;
-				virtual void SetEnabled(bool bInEnabled) = 0;
-				virtual int GetSourceHeight() const = 0;
-				virtual int GetSourceWidth() const = 0;
-				virtual FTexture2DRHIRef GetTexture() = 0;
+		public:
+			ITextureSource() = default;
+			virtual ~ITextureSource() = default;
+			virtual bool IsAvailable() const = 0;
+			virtual bool IsEnabled() const = 0;
+			virtual void SetEnabled(bool bInEnabled) = 0;
+			virtual int GetSourceHeight() const = 0;
+			virtual int GetSourceWidth() const = 0;
+			virtual FTexture2DRHIRef GetTexture() = 0;
+			virtual FString& GetName() const = 0;
 		};
 
-
 		/*
-		* Source of textures coming from the UE backbuffer.
-		* This class has the additional functionality of scaling textures from the backbuffer.
+		* Struct with `FTexture2DRHIRef` and the pixels of that texture read into a `TArray<FColor>`.
 		*/
-		class FTextureSourceBackBuffer : public ITextureSource, public TSharedFromThis<FTextureSourceBackBuffer, ESPMode::ThreadSafe>
+		struct FRawPixelsTexture : public FRefCountBase
+		{
+			FTexture2DRHIRef TextureRef;
+			TArray<FColor> RawPixels;
+			FRawPixelsTexture(FTexture2DRHIRef TexRef)
+				: TextureRef(TexRef){};
+		};
+
+		/*
+        * Base class for TextureSources that get their textures from the UE backbuffer.
+		* Textures are copied from the backbuffer using a triplebuffering mechanism so that texture read access is always thread safe while writes are occurring.
+		* If no texture has been written since the last read then the same texture will be read again.
+        * This class also has the additional functionality of scaling textures from the backbuffer.
+        * Note: This is a template class that uses CRTP - see: https://en.wikipedia.org/wiki/Curiously_recurring_template_pattern
+        * For derived types of this class they must contain the following static methods:
+        * 
+        * static void CopyTexture(const FTexture2DRHIRef& SourceTexture, TRefCountPtr<TextureType> DestTexture, FGPUFenceRHIRef& CopyFence)
+        * static TRefCountPtr<TextureType> CreateTexture(int Width, int Height)
+        * static FTexture2DRHIRef ToTextureRef(TRefCountPtr<TextureType> Texture)
+		* static FString& GetName()
+        */
+		template <class TextureType, class DerivedType>
+		class TTextureSourceBackBufferBase : public ITextureSource, public TSharedFromThis<TTextureSourceBackBufferBase<TextureType, DerivedType>, ESPMode::ThreadSafe>
 		{
 		public:
-			FTextureSourceBackBuffer();
-			FTextureSourceBackBuffer(float InScale);
-			~FTextureSourceBackBuffer();
-			
+			TTextureSourceBackBufferBase(float InFrameScale)
+				: FrameScale(InFrameScale)
+			{
+				// Explictly make clear we are adding another ref to this shared bool for the purposes of using in the below lambda
+				TSharedRef<bool, ESPMode::ThreadSafe> bEnabledClone = bEnabled;
+
+				// The backbuffer delegate can only be accessed using GameThread
+				AsyncTask(ENamedThreads::GameThread, [this, bEnabledClone]() {
+					/*Early exit if `this` died before game thread ran.*/
+					if (bEnabledClone.IsUnique())
+					{
+						return;
+					}
+					OnBackbuffer = &FSlateApplication::Get().GetRenderer()->OnBackBufferReadyToPresent();
+					BackbufferDelegateHandle = OnBackbuffer->AddSP(this, &TTextureSourceBackBufferBase::OnBackBufferReady_RenderThread);
+				});
+			}
+
+			TTextureSourceBackBufferBase()
+				: TTextureSourceBackBufferBase(1.0)
+			{
+			}
+
+			virtual ~TTextureSourceBackBufferBase()
+			{
+				if (OnBackbuffer)
+				{
+					OnBackbuffer->Remove(BackbufferDelegateHandle);
+				}
+
+				SetEnabled(false);
+			}
+
+			TRefCountPtr<TextureType> GetCurrent()
+			{
+				if (bIsTempDirty)
+				{
+					FScopeLock Lock(&CriticalSection);
+					ReadBuffer.Swap(TempBuffer);
+					bIsTempDirty = false;
+				}
+				return ReadBuffer;
+			}
+
 			/* Begin ITextureSource interface */
-			void SetEnabled(bool bInEnabled) override;
+			virtual void SetEnabled(bool bInEnabled) override
+			{
+				*bEnabled = bInEnabled;
+				// This source has been disabled, so set `bInitialized` to false so `OnBackBufferReady_RenderThread`
+				// will make new textures next time it is called.
+				if (bInitialized && bInEnabled == false)
+				{
+					bInitialized = false;
+				}
+			}
+
+			virtual FTexture2DRHIRef GetTexture() override
+			{
+				return DerivedType::ToTextureRef(GetCurrent());
+			}
+
 			bool IsEnabled() const override { return *bEnabled; }
 			bool IsAvailable() const override { return bInitialized; }
 			int GetSourceWidth() const override { return SourceWidth; }
 			int GetSourceHeight() const override { return SourceHeight; }
-			FTexture2DRHIRef GetTexture() override;
+			FString& GetName() const override { return DerivedType::GetNameImpl(); };
 			/* End ITextureSource interface */
 
 			float GetFrameScaling() const { return FrameScale; }
 
 		private:
-			void OnBackBufferReady_RenderThread(SWindow& SlateWindow, const FTexture2DRHIRef& FrameBuffer);
-			void Initialize(int Width, int Height);
+			void OnBackBufferReady_RenderThread(SWindow& SlateWindow, const FTexture2DRHIRef& FrameBuffer)
+			{
+				if (!bInitialized)
+				{
+					Initialize(FrameBuffer->GetSizeXY().X * FrameScale, FrameBuffer->GetSizeXY().Y * FrameScale);
+				}
+
+				if (!IsEnabled())
+				{
+					return;
+				}
+
+				auto& WriteBuffer = bWriteParity ? WriteBuffers[0] : WriteBuffers[1];
+				bWriteParity = !bWriteParity;
+
+				// for safety we just make sure that the buffer is not currently waiting for a copy
+				if (WriteBuffer.bAvailable)
+				{
+					WriteBuffer.bAvailable = false;
+
+					FRHICommandListImmediate& RHICmdList = FRHICommandListExecutor::GetImmediateCommandList();
+					WriteBuffer.Fence->Clear();
+
+					RHICmdList.EnqueueLambda([&WriteBuffer](FRHICommandListImmediate& RHICmdList) {
+						WriteBuffer.PreWaitingOnCopy = FPlatformTime::Cycles64();
+					});
+
+					DerivedType::CopyTexture(FrameBuffer, WriteBuffer.Texture, WriteBuffer.Fence);
+
+					UE::PixelStreaming::FGPUFencePoller::Get()->AddJob(WriteBuffer.Fence, bEnabled, [this, &WriteBuffer]() {
+						// This lambda is called only once the GPUFence is done
+						{
+							FScopeLock Lock(&CriticalSection);
+							TempBuffer.Swap(WriteBuffer.Texture);
+							WriteBuffer.Fence->Clear();
+							WriteBuffer.bAvailable = true;
+
+							bIsTempDirty = true;
+						}
+
+						// For debugging timing information about the copy operation
+						// Turning it on all the time is a bit too much log spam if logging stats
+						uint64 PostWaitingOnCopy = FPlatformTime::Cycles64();
+						UE::PixelStreaming::FStats* Stats = UE::PixelStreaming::FStats::Get();
+						if (Stats)
+						{
+							double CaptureLatencyMs = FPlatformTime::ToMilliseconds64(PostWaitingOnCopy - WriteBuffer.PreWaitingOnCopy);
+							Stats->StoreApplicationStat(UE::PixelStreaming::FStatData(FName(*FString::Printf(TEXT("Layer (x%.2f) Capture time (ms)"), FrameScale)), CaptureLatencyMs, 2, true));
+						}
+					});
+				}
+			}
+
+			void Initialize(int Width, int Height)
+			{
+				SourceWidth = Width;
+				SourceHeight = Height;
+
+				for (auto& Buffer : WriteBuffers)
+				{
+					Buffer.Texture = DerivedType::CreateTexture(SourceWidth, SourceHeight);
+					Buffer.Fence = GDynamicRHI->RHICreateGPUFence(TEXT("VideoCapturerCopyFence"));
+					Buffer.bAvailable = true;
+				}
+				bWriteParity = true;
+
+				TempBuffer = DerivedType::CreateTexture(SourceWidth, SourceHeight);
+				ReadBuffer = DerivedType::CreateTexture(SourceWidth, SourceHeight);
+				bIsTempDirty = false;
+
+				bInitialized = true;
+			}
 
 		private:
 			const float FrameScale;
@@ -65,25 +217,65 @@ namespace UE {
 
 			struct FCaptureFrame
 			{
-				FTexture2DRHIRef Texture;
+				TRefCountPtr<TextureType> Texture;
 				FGPUFenceRHIRef Fence;
 				bool bAvailable = true;
 				uint64 PreWaitingOnCopy;
 			};
 
 			/*
-			* Triple buffer setup with queued write buffers (since we have to wait for RHI copy).
-			* 1 Read buffer (read the captured texture)
-			* 1 Temp buffer (for swapping what is read and written)
-			* 2 Write buffers (2 write buffers because UE can render two frames before presenting sometimes)
-			*/
+            * Triple buffer setup with queued write buffers (since we have to wait for RHI copy).
+            * 1 Read buffer (read the captured texture)
+            * 1 Temp buffer (for swapping what is read and written)
+            * 2 Write buffers (2 write buffers because UE can render two frames before presenting sometimes)
+            */
 			FCriticalSection CriticalSection;
 			bool bWriteParity = true;
 			FCaptureFrame WriteBuffers[2];
-			FTexture2DRHIRef TempBuffer;
-			FTexture2DRHIRef ReadBuffer;
+			TRefCountPtr<TextureType> TempBuffer;
+			TRefCountPtr<TextureType> ReadBuffer;
 			FThreadSafeBool bIsTempDirty;
 		};
 
-	}
-}
+		/*
+		* Captures the backbuffer into a `FTexture2DRHIRef` in whatever pixel format the backbuffer is already using.
+		*/
+		class FTextureSourceBackBuffer : public TTextureSourceBackBufferBase<FRHITexture2D, FTextureSourceBackBuffer>
+		{
+		public:
+			FTextureSourceBackBuffer();
+			FTextureSourceBackBuffer(float InScale);
+
+			/* Begin TTextureSourceBackBufferBase "template interface" */
+			static void CopyTexture(const FTexture2DRHIRef& SourceTexture, TRefCountPtr<FRHITexture2D> DestTexture, FGPUFenceRHIRef& CopyFence);
+			static TRefCountPtr<FRHITexture2D> CreateTexture(int Width, int Height);
+			static FTexture2DRHIRef ToTextureRef(TRefCountPtr<FRHITexture2D> Texture);
+			static FString& GetNameImpl() { return Name; }
+			/* End TTextureSourceBackBufferBase "template interface" */
+		private:
+			inline static FString Name = FString(TEXT("FTextureSourceBackBuffer"));
+		};
+
+		/*
+		* Captures the backbuffer into a `FTexture2DRHIRef` in whatever pixel format the backbuffer is already using AND
+		* reads that texture to CPU memory (which is slow). 
+		* Note: Our intent is to remove this in favour of doing all swizzling on the GPU (CPU swizzling is the main purpose of reading the texture on the CPU currently).
+		*/
+		class FTextureSourceBackBufferToCPU : public TTextureSourceBackBufferBase<FRawPixelsTexture, FTextureSourceBackBufferToCPU>
+		{
+		public:
+			FTextureSourceBackBufferToCPU();
+			FTextureSourceBackBufferToCPU(float InScale);
+
+			/* Begin TTextureSourceBackBufferBase "template interface" */
+			static void CopyTexture(const FTexture2DRHIRef& SourceTexture, TRefCountPtr<FRawPixelsTexture> DestTexture, FGPUFenceRHIRef& CopyFence);
+			static TRefCountPtr<FRawPixelsTexture> CreateTexture(int Width, int Height);
+			static FTexture2DRHIRef ToTextureRef(TRefCountPtr<FRawPixelsTexture> Texture);
+			static FString& GetNameImpl() { return Name; }
+			/* End TTextureSourceBackBufferBase "template interface" */
+		private:
+			inline static FString Name = FString(TEXT("FTextureSourceBackBufferToCPU"));
+		};
+
+	} // namespace PixelStreaming
+} // namespace UE
diff --git a/Engine/Plugins/Media/PixelStreaming/Source/PixelStreaming/Private/Utils.h b/Engine/Plugins/Media/PixelStreaming/Source/PixelStreaming/Private/Utils.h
--- a/Engine/Plugins/Media/PixelStreaming/Source/PixelStreaming/Private/Utils.h
+++ b/Engine/Plugins/Media/PixelStreaming/Source/PixelStreaming/Private/Utils.h
@@ -9,8 +9,10 @@
 #include "Modules/ModuleManager.h"
 #include "PixelStreamingPrivate.h"
 
-namespace UE {
-	namespace PixelStreaming {
+namespace UE
+{
+	namespace PixelStreaming
+	{
 		inline void CopyTexture(FTexture2DRHIRef SourceTexture, FTexture2DRHIRef DestinationTexture, FGPUFenceRHIRef& CopyFence)
 		{
 			FRHICommandListImmediate& RHICmdList = FRHICommandListExecutor::GetImmediateCommandList();
@@ -88,5 +90,12 @@ namespace UE {
 			}
 			return Texture;
 		}
-	}
-}
+
+		inline void ReadTextureToCPU(FRHICommandListImmediate& RHICmdList, FTexture2DRHIRef& TextureRef, TArray<FColor>& OutPixels)
+		{
+			FIntRect Rect(0, 0, TextureRef->GetSizeX(), TextureRef->GetSizeY());
+			RHICmdList.ReadSurfaceData(TextureRef, Rect, OutPixels, FReadSurfaceDataFlags());
+		}
+
+	} // namespace PixelStreaming
+} // namespace UE
diff --git a/Engine/Plugins/Media/PixelStreaming/Source/PixelStreaming/Private/VideoSource.cpp b/Engine/Plugins/Media/PixelStreaming/Source/PixelStreaming/Private/VideoSource.cpp
--- a/Engine/Plugins/Media/PixelStreaming/Source/PixelStreaming/Private/VideoSource.cpp
+++ b/Engine/Plugins/Media/PixelStreaming/Source/PixelStreaming/Private/VideoSource.cpp
@@ -11,12 +11,12 @@ UE::PixelStreaming::FVideoSourceBase::FVideoSourceBase(FPixelStreamingPlayerId I
 {
 }
 
-void UE::PixelStreaming::FVideoSourceBase::AddRef() const 
-{ 
-	RefCount.IncRef(); 
+void UE::PixelStreaming::FVideoSourceBase::AddRef() const
+{
+	RefCount.IncRef();
 }
 
-rtc::RefCountReleaseStatus UE::PixelStreaming::FVideoSourceBase::Release() const 
+rtc::RefCountReleaseStatus UE::PixelStreaming::FVideoSourceBase::Release() const
 {
 	const rtc::RefCountReleaseStatus Status = RefCount.DecRef();
 	if (Status == rtc::RefCountReleaseStatus::kDroppedLastRef)
@@ -78,9 +78,19 @@ bool UE::PixelStreaming::FVideoSourceBase::AdaptCaptureFrame(const int64 Timesta
 UE::PixelStreaming::FVideoSourceP2P::FVideoSourceP2P(FPixelStreamingPlayerId InPlayerId, UE::PixelStreaming::IPixelStreamingSessions* InSessions)
 	: UE::PixelStreaming::FVideoSourceBase(InPlayerId)
 	, Sessions(InSessions)
-	, TextureSource(MakeShared<FTextureSourceBackBuffer>(1.0))
 {
-	
+	if (UE::PixelStreaming::Settings::IsCodecVPX())
+	{
+		TextureSource = MakeShared<FTextureSourceBackBufferToCPU>(1.0);
+	}
+	else
+	{
+		TextureSource = MakeShared<FTextureSourceBackBuffer>(1.0);
+	}
+
+	// We store the codec during construction as querying it everytime seem overly wasteful
+	// especially considering we don't actually support switching codec mid-stream.
+	Codec = Settings::GetSelectedCodec();
 }
 
 bool UE::PixelStreaming::FVideoSourceP2P::IsReadyForPump() const
@@ -88,7 +98,7 @@ bool UE::PixelStreaming::FVideoSourceP2P::IsReadyForPump() const
 	return TextureSource.IsValid() && TextureSource->IsAvailable();
 }
 
-webrtc::VideoFrame UE::PixelStreaming::FVideoSourceP2P::CreateFrame(int32 FrameId)
+webrtc::VideoFrame UE::PixelStreaming::FVideoSourceP2P::CreateFrameH264(int32 FrameId)
 {
 	bool bQualityController = Sessions->IsQualityController(PlayerId);
 	TextureSource->SetEnabled(bQualityController);
@@ -119,6 +129,39 @@ webrtc::VideoFrame UE::PixelStreaming::FVideoSourceP2P::CreateFrame(int32 FrameI
 	return Frame;
 }
 
+webrtc::VideoFrame UE::PixelStreaming::FVideoSourceP2P::CreateFrameVPX(int32 FrameId)
+{
+	// We always send the the `FLayerFrameBuffer` as our usage of VPX has no notion of a "quality controller"
+	// for hacky encoder sharing - unlike what we do for H264.
+
+	TextureSource->SetEnabled(true);
+	const int64 TimestampUs = rtc::TimeMicros();
+
+	rtc::scoped_refptr<webrtc::VideoFrameBuffer> FrameBuffer = new rtc::RefCountedObject<UE::PixelStreaming::FFrameBufferI420>(TextureSource);
+
+	webrtc::VideoFrame Frame = webrtc::VideoFrame::Builder()
+								   .set_video_frame_buffer(FrameBuffer)
+								   .set_timestamp_us(TimestampUs)
+								   .set_rotation(webrtc::VideoRotation::kVideoRotation_0)
+								   .set_id(FrameId)
+								   .build();
+
+	return Frame;
+}
+
+webrtc::VideoFrame UE::PixelStreaming::FVideoSourceP2P::CreateFrame(int32 FrameId)
+{
+	switch (Codec)
+	{
+		case UE::PixelStreaming::Settings::ECodec::VP8:
+		case UE::PixelStreaming::Settings::ECodec::VP9:
+			return CreateFrameVPX(FrameId);
+		case UE::PixelStreaming::Settings::ECodec::H264:
+		default:
+			return CreateFrameH264(FrameId);
+	}
+}
+
 /*
 * ------------------ FVideoSourceSFU ------------------
 */
@@ -135,7 +178,7 @@ UE::PixelStreaming::FVideoSourceSFU::FVideoSourceSFU()
 		SortedLayers.Add(&Layer);
 	}
 	SortedLayers.Sort([](const FLayer& LayerA, const FLayer& LayerB) { return LayerA.Scaling > LayerB.Scaling; });
-	
+
 	for (FLayer* SimulcastLayer : SortedLayers)
 	{
 		const float Scale = 1.0f / SimulcastLayer->Scaling;
@@ -150,7 +193,7 @@ bool UE::PixelStreaming::FVideoSourceSFU::IsReadyForPump() const
 	// Check all texture sources are ready.
 	int NumReady = 0;
 
-	for(int i = 0; i < LayerTextures.Num(); i++)
+	for (int i = 0; i < LayerTextures.Num(); i++)
 	{
 		bool bIsReady = LayerTextures[i].IsValid() && LayerTextures[i]->IsAvailable();
 		NumReady += bIsReady ? 1 : 0;
@@ -165,11 +208,11 @@ webrtc::VideoFrame UE::PixelStreaming::FVideoSourceSFU::CreateFrame(int32 FrameI
 	const int64 TimestampUs = rtc::TimeMicros();
 
 	webrtc::VideoFrame Frame = webrtc::VideoFrame::Builder()
-								.set_video_frame_buffer(FrameBuffer)
-								.set_timestamp_us(TimestampUs)
-								.set_rotation(webrtc::VideoRotation::kVideoRotation_0)
-								.set_id(FrameId)
-								.build();
+								   .set_video_frame_buffer(FrameBuffer)
+								   .set_timestamp_us(TimestampUs)
+								   .set_rotation(webrtc::VideoRotation::kVideoRotation_0)
+								   .set_id(FrameId)
+								   .build();
 
 	return Frame;
 }
\ No newline at end of file
diff --git a/Engine/Plugins/Media/PixelStreaming/Source/PixelStreaming/Private/VideoSource.h b/Engine/Plugins/Media/PixelStreaming/Source/PixelStreaming/Private/VideoSource.h
--- a/Engine/Plugins/Media/PixelStreaming/Source/PixelStreaming/Private/VideoSource.h
+++ b/Engine/Plugins/Media/PixelStreaming/Source/PixelStreaming/Private/VideoSource.h
@@ -8,8 +8,10 @@
 #include "TextureSource.h"
 #include "IPixelStreamingSessions.h"
 
-namespace UE {
-	namespace PixelStreaming {
+namespace UE
+{
+	namespace PixelStreaming
+	{
 		class FVideoSourceFactory;
 
 		/*
@@ -27,7 +29,7 @@ namespace UE {
 			/* Begin rtc::AdaptedVideoTrackSource overrides */
 			virtual webrtc::MediaSourceInterface::SourceState state() const override { return CurrentState; }
 			virtual bool remote() const override { return false; }
-			virtual bool is_screencast() const override { return false;}
+			virtual bool is_screencast() const override { return false; }
 			virtual absl::optional<bool> needs_denoising() const override { return false; }
 			/* End rtc::AdaptedVideoTrackSource overrides */
 
@@ -44,7 +46,7 @@ namespace UE {
 			virtual webrtc::VideoFrame CreateFrame(int32 FrameId) = 0;
 
 		private:
-			mutable webrtc::webrtc_impl::RefCounter RefCount{0};
+			mutable webrtc::webrtc_impl::RefCounter RefCount{ 0 };
 
 		protected:
 			webrtc::MediaSourceInterface::SourceState CurrentState;
@@ -56,18 +58,22 @@ namespace UE {
 		*/
 		class FVideoSourceP2P : public FVideoSourceBase
 		{
-			public:
-				FVideoSourceP2P(FPixelStreamingPlayerId InPlayerId, IPixelStreamingSessions* InSessions);
-
-			protected:
-				IPixelStreamingSessions* Sessions;
-				TSharedPtr<ITextureSource> TextureSource;
-
-			protected:
-				/* Begin FVideoSourceBase */
-				virtual webrtc::VideoFrame CreateFrame(int32 FrameId) override;
-				virtual bool IsReadyForPump() const override;
-				/* End FVideoSourceBase */
+		public:
+			FVideoSourceP2P(FPixelStreamingPlayerId InPlayerId, IPixelStreamingSessions* InSessions);
+
+		protected:
+			IPixelStreamingSessions* Sessions;
+			TSharedPtr<ITextureSource> TextureSource;
+			Settings::ECodec Codec;
+
+		protected:
+			/* Begin FVideoSourceBase */
+			virtual webrtc::VideoFrame CreateFrame(int32 FrameId) override;
+			virtual bool IsReadyForPump() const override;
+			/* End FVideoSourceBase */
+
+			virtual webrtc::VideoFrame CreateFrameH264(int32 FrameId);
+			virtual webrtc::VideoFrame CreateFrameVPX(int32 FrameId);
 		};
 
 		/*
@@ -75,18 +81,18 @@ namespace UE {
 		*/
 		class FVideoSourceSFU : public FVideoSourceBase
 		{
-			public:
-				FVideoSourceSFU();
+		public:
+			FVideoSourceSFU();
 
-			protected:
-				TArray<TSharedPtr<ITextureSource>> LayerTextures;
+		protected:
+			TArray<TSharedPtr<ITextureSource>> LayerTextures;
 
-			protected:
-				/* Begin FVideoSourceBase */
-				virtual webrtc::VideoFrame CreateFrame(int32 FrameId) override;
-				virtual bool IsReadyForPump() const override;
-				/* End FVideoSourceBase */
+		protected:
+			/* Begin FVideoSourceBase */
+			virtual webrtc::VideoFrame CreateFrame(int32 FrameId) override;
+			virtual bool IsReadyForPump() const override;
+			/* End FVideoSourceBase */
 		};
 
-	}
-}
+	} // namespace PixelStreaming
+} // namespace UE
