Fix static analysis issue with virtual function call in destructor.

diff --git a/Engine/Plugins/Media/PixelStreaming/Source/PixelStreaming/Private/TextureSource.h b/Engine/Plugins/Media/PixelStreaming/Source/PixelStreaming/Private/TextureSource.h
--- a/Engine/Plugins/Media/PixelStreaming/Source/PixelStreaming/Private/TextureSource.h
+++ b/Engine/Plugins/Media/PixelStreaming/Source/PixelStreaming/Private/TextureSource.h
@@ -91,7 +91,7 @@ namespace UE
 					OnBackbuffer->Remove(BackbufferDelegateHandle);
 				}
 
-				SetEnabled(false);
+				*bEnabled = false;
 			}
 
 			TRefCountPtr<TextureType> GetCurrent()
