Fix compiler warnings in Static Analysis

diff --git a/Engine/Plugins/Media/PixelStreaming/Source/PixelStreaming/Private/SignallingServerConnection.h b/Engine/Plugins/Media/PixelStreaming/Source/PixelStreaming/Private/SignallingServerConnection.h
--- a/Engine/Plugins/Media/PixelStreaming/Source/PixelStreaming/Private/SignallingServerConnection.h
+++ b/Engine/Plugins/Media/PixelStreaming/Source/PixelStreaming/Private/SignallingServerConnection.h
@@ -40,7 +40,7 @@ namespace UE::PixelStreaming
 	{
 	public:
 		explicit FSignallingServerConnection(FSignallingServerConnectionObserver& Observer, FString InStreamerId);
-		virtual ~FSignallingServerConnection();
+		~FSignallingServerConnection();
 
 		void Connect(const FString& Url);
 		void Disconnect();
