Changes to Pixel Streaming plugin to expose more API publicly.

diff --git a/Engine/Plugins/Media/PixelStreaming/Source/PixelStreaming/PixelStreaming.Build.cs b/Engine/Plugins/Media/PixelStreaming/Source/PixelStreaming/PixelStreaming.Build.cs
--- a/Engine/Plugins/Media/PixelStreaming/Source/PixelStreaming/PixelStreaming.Build.cs
+++ b/Engine/Plugins/Media/PixelStreaming/Source/PixelStreaming/PixelStreaming.Build.cs
@@ -47,7 +47,8 @@ namespace UnrealBuildTool.Rules
 			// This is so for game projects using our public headers don't have to include extra modules they might not know about.
 			PublicDependencyModuleNames.AddRange(new string[]
 			{
-				"InputDevice"
+				"InputDevice",
+				"WebRTC"
 			});
 
 			// NOTE: General rule is not to access the private folder of another module
diff --git a/Engine/Plugins/Media/PixelStreaming/Source/PixelStreaming/Private/ClientRTC.h b/Engine/Plugins/Media/PixelStreaming/Source/PixelStreaming/Private/ClientRTC.h
--- a/Engine/Plugins/Media/PixelStreaming/Source/PixelStreaming/Private/ClientRTC.h
+++ b/Engine/Plugins/Media/PixelStreaming/Source/PixelStreaming/Private/ClientRTC.h
@@ -3,7 +3,7 @@
 
 #include "WebRTCIncludes.h"
 #include "SignallingServerConnection.h"
-#include "ProtocolDefs.h"
+#include "PixelStreamingProtocolDefs.h"
 
 namespace UE::PixelStreaming
 {
@@ -46,7 +46,7 @@ namespace UE::PixelStreaming
 		virtual void OnSessionDescription(webrtc::SdpType Type, const FString& Sdp) override;
 		virtual void OnRemoteIceCandidate(const FString& SdpMid, int SdpMLineIndex, const FString& Sdp) override;
 		virtual void OnPeerDataChannels(int32 SendStreamId, int32 RecvStreamId) override;
-		virtual void OnPlayerCount(uint32 Count) override { }
+		virtual void OnPlayerCount(uint32 Count) override {}
 
 		//
 		// webrtc::PeerConnectionObserver implementation.
diff --git a/Engine/Plugins/Media/PixelStreaming/Source/PixelStreaming/Private/DataChannelObserver.cpp b/Engine/Plugins/Media/PixelStreaming/Source/PixelStreaming/Private/DataChannelObserver.cpp
--- a/Engine/Plugins/Media/PixelStreaming/Source/PixelStreaming/Private/DataChannelObserver.cpp
+++ b/Engine/Plugins/Media/PixelStreaming/Source/PixelStreaming/Private/DataChannelObserver.cpp
@@ -58,7 +58,7 @@ namespace UE::PixelStreaming
 				{
 					Delegates->OnDataChannelClosedNative.Broadcast(PlayerId, DataChannel.get());
 				}
-				PlayerSessions->ForSession(PlayerId, [](auto Session) {
+				PlayerSessions->ForSession(PlayerId, [](TSharedPtr<IPlayerSession> Session) {
 					Session->OnDataChannelClosed();
 				});
 				break;
@@ -89,7 +89,7 @@ namespace UE::PixelStreaming
 			const size_t StringLen = (Buffer.data.size() - 1) / sizeof(TCHAR);
 			const TCHAR* StringPtr = reinterpret_cast<const TCHAR*>(Buffer.data.data() + 1);
 			const FString EchoMessage(StringLen, StringPtr);
-			PlayerSessions->ForSession(PlayerId, [&](auto Session) {
+			PlayerSessions->ForSession(PlayerId, [&](TSharedPtr<IPlayerSession> Session) {
 				Session->SendMessage(Protocol::EToPlayerMsg::TestEcho, EchoMessage);
 			});
 		}
@@ -141,7 +141,7 @@ namespace UE::PixelStreaming
 					TransmissionTimeMs);
 			}
 
-			PlayerSessions->ForSession(PlayerId, [&](auto Session) {
+			PlayerSessions->ForSession(PlayerId, [&](TSharedPtr<IPlayerSession> Session) {
 				Session->SendMessage(Protocol::EToPlayerMsg::LatencyTest, ReportToTransmitJSON);
 			});
 		});
@@ -188,7 +188,7 @@ namespace UE::PixelStreaming
 
 		const FString FullPayload = FString::Printf(TEXT("{ \"PixelStreaming\": %s, \"Encoder\": %s, \"WebRTC\": %s }"), *PixelStreamingPayload, *EncoderPayload, *WebRTCPayload);
 
-		PlayerSessions->ForSession(PlayerId, [&](auto Session) {
+		PlayerSessions->ForSession(PlayerId, [&](TSharedPtr<IPlayerSession> Session) {
 			if (!Session->SendMessage(Protocol::EToPlayerMsg::InitialSettings, FullPayload))
 			{
 				UE_LOG(LogPixelStreaming, Log, TEXT("Failed to send initial Pixel Streaming settings."));
diff --git a/Engine/Plugins/Media/PixelStreaming/Source/PixelStreaming/Private/EncoderFactory.cpp b/Engine/Plugins/Media/PixelStreaming/Source/PixelStreaming/Private/EncoderFactory.cpp
--- a/Engine/Plugins/Media/PixelStreaming/Source/PixelStreaming/Private/EncoderFactory.cpp
+++ b/Engine/Plugins/Media/PixelStreaming/Source/PixelStreaming/Private/EncoderFactory.cpp
@@ -38,6 +38,7 @@ std::vector<webrtc::SdpVideoFormat> UE::PixelStreaming::FVideoEncoderFactory::Ge
 		case UE::PixelStreaming::Settings::ECodec::H264:
 		default:
 			video_formats.push_back(UE::PixelStreaming::CreateH264Format(webrtc::H264::kProfileConstrainedBaseline, webrtc::H264::kLevel3_1));
+			video_formats.push_back(UE::PixelStreaming::CreateH264Format(webrtc::H264::kProfileBaseline, webrtc::H264::kLevel3_1));
 	}
 	return video_formats;
 }
@@ -146,17 +147,16 @@ UE::PixelStreaming::FVideoEncoderH264Wrapper* UE::PixelStreaming::FVideoEncoderF
 * ------------- UE::PixelStreaming::FSimulcastEncoderFactory ---------------
 */
 
-UE::PixelStreaming::FSimulcastEncoderFactory::FSimulcastEncoderFactory() 
+UE::PixelStreaming::FSimulcastEncoderFactory::FSimulcastEncoderFactory()
 	: PrimaryEncoderFactory(MakeUnique<UE::PixelStreaming::FVideoEncoderFactory>())
 {
 	// Make a copy of simulcast settings and sort them based on scaling.
-	for(FEncoderFactoryId i = 0; i < Settings::SimulcastParameters.Layers.Num(); i++)
+	for (FEncoderFactoryId i = 0; i < Settings::SimulcastParameters.Layers.Num(); i++)
 	{
 		GetOrCreateEncoderFactory(i);
 	}
 }
 
-
 UE::PixelStreaming::FSimulcastEncoderFactory::~FSimulcastEncoderFactory()
 {
 }
@@ -166,7 +166,6 @@ std::unique_ptr<webrtc::VideoEncoder> UE::PixelStreaming::FSimulcastEncoderFacto
 	return std::make_unique<UE::PixelStreaming::FSimulcastEncoderAdapter>(*this, format);
 }
 
-
 UE::PixelStreaming::FVideoEncoderFactory* UE::PixelStreaming::FSimulcastEncoderFactory::GetEncoderFactory(FEncoderFactoryId Id)
 {
 	FScopeLock Lock(&EncoderFactoriesGuard);
@@ -180,7 +179,7 @@ UE::PixelStreaming::FVideoEncoderFactory* UE::PixelStreaming::FSimulcastEncoderF
 UE::PixelStreaming::FVideoEncoderFactory* UE::PixelStreaming::FSimulcastEncoderFactory::GetOrCreateEncoderFactory(FEncoderFactoryId Id)
 {
 	UE::PixelStreaming::FVideoEncoderFactory* Existing = GetEncoderFactory(Id);
-	if(Existing == nullptr)
+	if (Existing == nullptr)
 	{
 		FScopeLock Lock(&EncoderFactoriesGuard);
 		TUniquePtr<FVideoEncoderFactory> EncoderFactory = MakeUnique<UE::PixelStreaming::FVideoEncoderFactory>();
@@ -193,7 +192,6 @@ UE::PixelStreaming::FVideoEncoderFactory* UE::PixelStreaming::FSimulcastEncoderF
 	}
 }
 
-
 std::vector<webrtc::SdpVideoFormat> UE::PixelStreaming::FSimulcastEncoderFactory::GetSupportedFormats() const
 {
 	return PrimaryEncoderFactory->GetSupportedFormats();
diff --git a/Engine/Plugins/Media/PixelStreaming/Source/PixelStreaming/Private/EncoderFrameFactory.h b/Engine/Plugins/Media/PixelStreaming/Source/PixelStreaming/Private/EncoderFrameFactory.h
--- a/Engine/Plugins/Media/PixelStreaming/Source/PixelStreaming/Private/EncoderFrameFactory.h
+++ b/Engine/Plugins/Media/PixelStreaming/Source/PixelStreaming/Private/EncoderFrameFactory.h
@@ -3,10 +3,11 @@
 #pragma once
 
 #include "VideoEncoderInput.h"
-#include "FrameBuffer.h"
 
-namespace UE {
-	namespace PixelStreaming {
+namespace UE
+{
+	namespace PixelStreaming
+	{
 		/*
 		* Factory to create `AVEncoder::FVideoEncoderInputFrame` for use in Pixel Streaming encoders.
 		* This class is responsible for creating/managing the underlying `AVEncoder::FVideoEncoderInput` required to create
@@ -28,10 +29,10 @@ namespace UE {
 			TSharedPtr<AVEncoder::FVideoEncoderInput> CreateVideoEncoderInput() const;
 			void SetTexture(TSharedPtr<AVEncoder::FVideoEncoderInputFrame> Frame, const FTexture2DRHIRef& Texture);
 			void SetTextureCUDAVulkan(TSharedPtr<AVEncoder::FVideoEncoderInputFrame> Frame, const FTexture2DRHIRef& Texture);
-		#if PLATFORM_WINDOWS
+#if PLATFORM_WINDOWS
 			void SetTextureCUDAD3D11(TSharedPtr<AVEncoder::FVideoEncoderInputFrame> Frame, const FTexture2DRHIRef& Texture);
 			void SetTextureCUDAD3D12(TSharedPtr<AVEncoder::FVideoEncoderInputFrame> Frame, const FTexture2DRHIRef& Texture);
-		#endif // PLATFORM_WINDOWS
+#endif // PLATFORM_WINDOWS
 
 		private:
 			uint64 FrameId = 0;
@@ -39,5 +40,5 @@ namespace UE {
 			// Store a mapping between raw textures and the FVideoEncoderInputFrames that wrap them
 			TMap<FTexture2DRHIRef, TSharedPtr<AVEncoder::FVideoEncoderInputFrame>> TextureToFrameMapping;
 		};
-	}
-}
+	} // namespace PixelStreaming
+} // namespace UE
diff --git a/Engine/Plugins/Media/PixelStreaming/Source/PixelStreaming/Private/FixedFPSPump.cpp b/Engine/Plugins/Media/PixelStreaming/Source/PixelStreaming/Private/FixedFPSPump.cpp
--- a/Engine/Plugins/Media/PixelStreaming/Source/PixelStreaming/Private/FixedFPSPump.cpp
+++ b/Engine/Plugins/Media/PixelStreaming/Source/PixelStreaming/Private/FixedFPSPump.cpp
@@ -3,8 +3,8 @@
 #include "FixedFPSPump.h"
 #include "VideoSource.h"
 #include "Settings.h"
-#include "TextureSource.h"
-#include "FrameBuffer.h"
+#include "PixelStreamingTextureSource.h"
+#include "PixelStreamingFrameBuffer.h"
 #include "PlayerSession.h"
 
 /*
@@ -45,7 +45,7 @@ void UE::PixelStreaming::FFixedFPSPump::UnregisterVideoSource(FPixelStreamingPla
 	NextPumpEvent->Trigger();
 }
 
-void UE::PixelStreaming::FFixedFPSPump::RegisterVideoSource(FPixelStreamingPlayerId PlayerId, UE::PixelStreaming::IPumpedVideoSource* Source)
+void UE::PixelStreaming::FFixedFPSPump::RegisterVideoSource(FPixelStreamingPlayerId PlayerId, IPumpedVideoSource* Source)
 {
 	checkf(Source, TEXT("Cannot register a nullptr VideoSource."));
 	FScopeLock Guard(&SourcesGuard);
@@ -72,12 +72,12 @@ void UE::PixelStreaming::FFixedFPSPump::PumpLoop()
 
 			// Pump each video source
 			TMap<FPixelStreamingPlayerId, IPumpedVideoSource*>::TIterator Iter = VideoSources.CreateIterator();
-			for(; Iter; ++Iter)
+			for (; Iter; ++Iter)
 			{
 
 				IPumpedVideoSource* VideoSource = Iter.Value();
 
-				if(VideoSource->IsReadyForPump())
+				if (VideoSource->IsReadyForPump())
 				{
 					VideoSource->OnPump(FrameId);
 				}
@@ -85,12 +85,12 @@ void UE::PixelStreaming::FFixedFPSPump::PumpLoop()
 		}
 
 		// Sleep as long as we need for a constant FPS
-		const uint64 EndCycles	  = FPlatformTime::Cycles64();
-		const double DeltaMs	  = FPlatformTime::ToMilliseconds64(EndCycles - LastCycles);
-		const int32 FPS			  = UE::PixelStreaming::Settings::CVarPixelStreamingWebRTCFps.GetValueOnAnyThread();
+		const uint64 EndCycles = FPlatformTime::Cycles64();
+		const double DeltaMs = FPlatformTime::ToMilliseconds64(EndCycles - LastCycles);
+		const int32 FPS = UE::PixelStreaming::Settings::CVarPixelStreamingWebRTCFps.GetValueOnAnyThread();
 		const double FrameDeltaMs = 1000.0 / FPS;
-		const double SleepMs	  = FrameDeltaMs - DeltaMs;
-		LastCycles				  = EndCycles;
+		const double SleepMs = FrameDeltaMs - DeltaMs;
+		LastCycles = EndCycles;
 
 		if (SleepMs > 0)
 		{
diff --git a/Engine/Plugins/Media/PixelStreaming/Source/PixelStreaming/Private/FixedFPSPump.h b/Engine/Plugins/Media/PixelStreaming/Source/PixelStreaming/Private/FixedFPSPump.h
--- a/Engine/Plugins/Media/PixelStreaming/Source/PixelStreaming/Private/FixedFPSPump.h
+++ b/Engine/Plugins/Media/PixelStreaming/Source/PixelStreaming/Private/FixedFPSPump.h
@@ -4,25 +4,12 @@
 #include "PixelStreamingPlayerId.h"
 #include "HAL/Thread.h"
 #include "WebRTCIncludes.h"
+#include "IPixelStreamingPumpedVideoSource.h"
 
-namespace UE {
-	namespace PixelStreaming {
-
-		/*
-		* Interface for video sources that get pumped through `OnPump`.
-		*/
-		class IPumpedVideoSource : public rtc::RefCountInterface
-		{
-			public:
-				virtual void OnPump(int32 FrameId) = 0;
-				virtual bool IsReadyForPump() const = 0;
-				
-				/* rtc::RefCountInterface */
-				virtual void AddRef() const = 0;
-				virtual rtc::RefCountReleaseStatus Release() const = 0;
-				virtual bool HasOneRef() const = 0;
-		};
-
+namespace UE
+{
+	namespace PixelStreaming
+	{
 		/*
 		* Runs a seperate thread that "pumps" at a fixed FPS interval. WebRTC video sources may add themselves to be "pumped" 
 		* at a fixed interval indepedent of render FPS or game thread FPS. This is useful so that poorly 
@@ -47,7 +34,6 @@ namespace UE {
 			FCriticalSection SourcesGuard;
 			TMap<FPixelStreamingPlayerId, IPumpedVideoSource*> VideoSources;
 
-
 			TUniquePtr<FThread> PumpThread;
 			bool bThreadRunning = true;
 			//FEvent* PlayersChangedEvent;
@@ -56,5 +42,5 @@ namespace UE {
 
 			static FFixedFPSPump* Instance;
 		};
-	}
-}
+	} // namespace PixelStreaming
+} // namespace UE
diff --git a/Engine/Plugins/Media/PixelStreaming/Source/PixelStreaming/Private/FrameBuffer.cpp b/Engine/Plugins/Media/PixelStreaming/Source/PixelStreaming/Private/FrameBuffer.cpp
--- a/Engine/Plugins/Media/PixelStreaming/Source/PixelStreaming/Private/FrameBuffer.cpp
+++ b/Engine/Plugins/Media/PixelStreaming/Source/PixelStreaming/Private/FrameBuffer.cpp
@@ -1,28 +1,28 @@
 // Copyright Epic Games, Inc. All Rights Reserved.
 
-#include "FrameBuffer.h"
-#include "TextureSource.h"
+#include "PixelStreamingFrameBuffer.h"
+#include "PixelStreamingTextureSource.h"
 #include "Utils.h"
 
 /*
 * ----------------- FInitializeFrameBuffer -----------------
 */
 
-UE::PixelStreaming::FInitializeFrameBuffer::FInitializeFrameBuffer(TSharedPtr<ITextureSource> InTextureSource)
+FInitializeFrameBuffer::FInitializeFrameBuffer(TSharedPtr<IPixelStreamingTextureSource> InTextureSource)
 	: TextureSource(InTextureSource)
 {
 }
 
-UE::PixelStreaming::FInitializeFrameBuffer::~FInitializeFrameBuffer()
+FInitializeFrameBuffer::~FInitializeFrameBuffer()
 {
 }
 
-int UE::PixelStreaming::FInitializeFrameBuffer::width() const
+int FInitializeFrameBuffer::width() const
 {
 	return TextureSource->GetSourceWidth();
 }
 
-int UE::PixelStreaming::FInitializeFrameBuffer::height() const
+int FInitializeFrameBuffer::height() const
 {
 	return TextureSource->GetSourceHeight();
 }
@@ -31,33 +31,33 @@ int UE::PixelStreaming::FInitializeFrameBuffer::height() const
 * ----------------- FSimulcastFrameBuffer -----------------
 */
 
-UE::PixelStreaming::FSimulcastFrameBuffer::FSimulcastFrameBuffer(TArray<TSharedPtr<ITextureSource>>& InTextureSources)
+FSimulcastFrameBuffer::FSimulcastFrameBuffer(TArray<TSharedPtr<IPixelStreamingTextureSource>>& InTextureSources)
 	: TextureSources(InTextureSources)
 {
 }
 
-UE::PixelStreaming::FSimulcastFrameBuffer::~FSimulcastFrameBuffer()
+FSimulcastFrameBuffer::~FSimulcastFrameBuffer()
 {
 }
 
-int UE::PixelStreaming::FSimulcastFrameBuffer::GetNumLayers() const
+int FSimulcastFrameBuffer::GetNumLayers() const
 {
 	return TextureSources.Num();
 }
 
-TSharedPtr<UE::PixelStreaming::ITextureSource> UE::PixelStreaming::FSimulcastFrameBuffer::GetLayerFrameSource(int LayerIndex) const
+TSharedPtr<IPixelStreamingTextureSource> FSimulcastFrameBuffer::GetLayerFrameSource(int LayerIndex) const
 {
 	checkf(LayerIndex >= 0 && LayerIndex < TextureSources.Num(), TEXT("Requested layer index was out of bounds."));
 	return TextureSources[LayerIndex];
 }
 
-int UE::PixelStreaming::FSimulcastFrameBuffer::width() const
+int FSimulcastFrameBuffer::width() const
 {
 	checkf(TextureSources.Num() > 0, TEXT("Must be at least one texture source to get the width from."));
 	return FMath::Max(TextureSources[0]->GetSourceWidth(), TextureSources[TextureSources.Num() - 1]->GetSourceWidth());
 }
 
-int UE::PixelStreaming::FSimulcastFrameBuffer::height() const
+int FSimulcastFrameBuffer::height() const
 {
 	checkf(TextureSources.Num() > 0, TEXT("Must be at least one texture source to get the height from."));
 	return FMath::Max(TextureSources[0]->GetSourceHeight(), TextureSources[TextureSources.Num() - 1]->GetSourceHeight());
@@ -67,26 +67,26 @@ int UE::PixelStreaming::FSimulcastFrameBuffer::height() const
 * ----------------- FLayerFrameBuffer -----------------
 */
 
-UE::PixelStreaming::FLayerFrameBuffer::FLayerFrameBuffer(TSharedPtr<UE::PixelStreaming::ITextureSource> InTextureSource)
+FLayerFrameBuffer::FLayerFrameBuffer(TSharedPtr<IPixelStreamingTextureSource> InTextureSource)
 	: TextureSource(InTextureSource)
 {
 }
 
-UE::PixelStreaming::FLayerFrameBuffer::~FLayerFrameBuffer()
+FLayerFrameBuffer::~FLayerFrameBuffer()
 {
 }
 
-FTexture2DRHIRef UE::PixelStreaming::FLayerFrameBuffer::GetFrame() const
+FTexture2DRHIRef FLayerFrameBuffer::GetFrame() const
 {
 	return TextureSource->GetTexture();
 }
 
-int UE::PixelStreaming::FLayerFrameBuffer::width() const
+int FLayerFrameBuffer::width() const
 {
 	return TextureSource->GetSourceWidth();
 }
 
-int UE::PixelStreaming::FLayerFrameBuffer::height() const
+int FLayerFrameBuffer::height() const
 {
 	return TextureSource->GetSourceHeight();
 }
@@ -95,31 +95,31 @@ int UE::PixelStreaming::FLayerFrameBuffer::height() const
 * ----------------- FFrameBufferI420 -----------------
 */
 
-UE::PixelStreaming::FFrameBufferI420::FFrameBufferI420(TSharedPtr<UE::PixelStreaming::ITextureSource> InTextureSource)
+FFrameBufferI420::FFrameBufferI420(TSharedPtr<IPixelStreamingTextureSource> InTextureSource)
 	: TextureSource(InTextureSource)
 {
 }
 
-UE::PixelStreaming::FFrameBufferI420::~FFrameBufferI420()
+FFrameBufferI420::~FFrameBufferI420()
 {
 }
 
-FTexture2DRHIRef UE::PixelStreaming::FFrameBufferI420::GetFrame() const
+FTexture2DRHIRef FFrameBufferI420::GetFrame() const
 {
 	return TextureSource->GetTexture();
 }
 
-int UE::PixelStreaming::FFrameBufferI420::width() const
+int FFrameBufferI420::width() const
 {
 	return TextureSource->GetSourceWidth();
 }
 
-int UE::PixelStreaming::FFrameBufferI420::height() const
+int FFrameBufferI420::height() const
 {
 	return TextureSource->GetSourceHeight();
 }
 
-const webrtc::I420BufferInterface* UE::PixelStreaming::FFrameBufferI420::GetI420() const
+const webrtc::I420BufferInterface* FFrameBufferI420::GetI420() const
 {
 	return Buffer;
 }
@@ -127,16 +127,15 @@ const webrtc::I420BufferInterface* UE::PixelStreaming::FFrameBufferI420::GetI420
 /*
 * NOTE: Only used for non-hardware encoders (e.g. VP8) - H264 does NOT hit this function.
 */
-rtc::scoped_refptr<webrtc::I420BufferInterface> UE::PixelStreaming::FFrameBufferI420::ToI420()
+rtc::scoped_refptr<webrtc::I420BufferInterface> FFrameBufferI420::ToI420()
 {
-	ITextureSource* TextureSourcePtr = TextureSource.Get();
-	checkf(TextureSourcePtr && TextureSourcePtr->GetName() == FString(TEXT("FTextureSourceBackBufferToCPU")), TEXT("To use VPX encoder texture source must be FTextureSourceBackBufferToCPU"));
+	IBackBufferTextureSource* TextureSourcePtr = static_cast<IBackBufferTextureSource*>(TextureSource.Get());
+	checkf(TextureSourcePtr && TextureSourcePtr->GetName() == FString(TEXT("FBackBufferToCPUTextureSource")), TEXT("To use VPX encoder texture source must be FBackBufferToCPUTextureSource"));
 
-	FTextureSourceBackBufferToCPU* TextureSourceCPU = static_cast<FTextureSourceBackBufferToCPU*>(TextureSourcePtr);
-	TRefCountPtr<FRawPixelsTexture> Current = TextureSourceCPU->GetCurrent();
+	TRefCountPtr<FPixelStreamingCPUReadableBackbufferTexture> Current = static_cast<FPixelStreamingCPUReadableBackbufferTexture*>(TextureSourcePtr->GetCurrent().GetReference());
 
-	uint32 TextureWidth = Current->TextureRef->GetSizeX();
-	uint32 TextureHeight = Current->TextureRef->GetSizeY();
+	uint32 TextureWidth = Current->GetTexture()->GetSizeX();
+	uint32 TextureHeight = Current->GetTexture()->GetSizeY();
 
 	uint32 i = 0;
 	uint32 NumPixels = TextureWidth * TextureHeight;
@@ -152,14 +151,14 @@ rtc::scoped_refptr<webrtc::I420BufferInterface> UE::PixelStreaming::FFrameBuffer
 	Buffer = webrtc::I420Buffer::Create(TextureWidth, TextureHeight, TextureWidth, TextureWidth / 2, TextureWidth / 2);
 
 	// If we have no raw pixels then early exit this swizzle operation
-	if (Current->RawPixels.Num() == 0)
+	if (Current->GetRawPixels().Num() == 0)
 	{
 		webrtc::I420Buffer::SetBlack(Buffer);
 		return Buffer;
 	}
 
 	uint8* yuv420p = Buffer->MutableDataY();
-	const uint8* rgbIn = (const uint8*)Current->RawPixels.GetData();
+	const uint8* rgbIn = (const uint8*)Current->GetRawPixels().GetData();
 	for (uint32 j = 0; j < TextureHeight; j++)
 	{
 		for (uint32 k = 0; k < TextureWidth; k++)
diff --git a/Engine/Plugins/Media/PixelStreaming/Source/PixelStreaming/Private/FrameBuffer.h b/Engine/Plugins/Media/PixelStreaming/Source/PixelStreaming/Private/FrameBuffer.h
deleted file mode 100644
diff --git a/Engine/Plugins/Media/PixelStreaming/Source/PixelStreaming/Private/GPUFencePoller.h b/Engine/Plugins/Media/PixelStreaming/Source/PixelStreaming/Private/GPUFencePoller.h
--- a/Engine/Plugins/Media/PixelStreaming/Source/PixelStreaming/Private/GPUFencePoller.h
+++ b/Engine/Plugins/Media/PixelStreaming/Source/PixelStreaming/Private/GPUFencePoller.h
@@ -11,56 +11,58 @@
 #include "HAL/RunnableThread.h"
 #include "HAL/Runnable.h"
 
-namespace UE {
-    namespace PixelStreaming {
-        /*
+namespace UE
+{
+	namespace PixelStreaming
+	{
+		/*
         * A thread with a single purpose, to poll GPU fences until they are done.
         * Polling jobs can be submitted to this thread using its static methods.
         * Once a polling job is complete a callback is called on the submitted job.
         */
-        class FGPUFencePoller
-        {
-            public:
-                FGPUFencePoller();
-                void Shutdown();
-                virtual ~FGPUFencePoller();
-                void AddJob(FGPUFenceRHIRef Fence, TSharedRef<bool, ESPMode::ThreadSafe> bKeepPolling, TFunction<void()> FenceDoneCallback);
-                static FGPUFencePoller* Get();
+		class PIXELSTREAMING_API FGPUFencePoller
+		{
+		public:
+			FGPUFencePoller();
+			void Shutdown();
+			virtual ~FGPUFencePoller();
+			void AddJob(FGPUFenceRHIRef Fence, TSharedRef<bool, ESPMode::ThreadSafe> bKeepPolling, TFunction<void()> FenceDoneCallback);
+			static FGPUFencePoller* Get();
 
-            private:
-                struct FPollJob
-                {
-                    FPollJob();
-                    FPollJob(FGPUFenceRHIRef InFence, TSharedRef<bool, ESPMode::ThreadSafe> bInKeepPolling, TFunction<void()> InFenceDoneCallback);
-                    FGPUFenceRHIRef Fence;
-                    TSharedRef<bool, ESPMode::ThreadSafe> bKeepPolling;
-                    TFunction<void()> FenceDoneCallback;
-                };
+		private:
+			struct FPollJob
+			{
+				FPollJob();
+				FPollJob(FGPUFenceRHIRef InFence, TSharedRef<bool, ESPMode::ThreadSafe> bInKeepPolling, TFunction<void()> InFenceDoneCallback);
+				FGPUFenceRHIRef Fence;
+				TSharedRef<bool, ESPMode::ThreadSafe> bKeepPolling;
+				TFunction<void()> FenceDoneCallback;
+			};
 
-                class FPollerRunnable : public FRunnable
-                {
-                public:
-                    FPollerRunnable();
-                    void AddJob(FPollJob Job);
-                    bool IsRunning() const;
-                    // Begin FRunnable interface.
-                    virtual bool Init() override;
-                    virtual uint32 Run() override;
-                    virtual void Stop() override;
-                    virtual void Exit() override;
-                    // End FRunnable interface
-                private:
-                    bool bIsRunning;
-                    FEvent* JobAddedEvent;
-                    TArray<FPollJob> PollJobs;
-                    TQueue<FPollJob, EQueueMode::Mpsc> JobsToAdd;
-                };
+			class FPollerRunnable : public FRunnable
+			{
+			public:
+				FPollerRunnable();
+				void AddJob(FPollJob Job);
+				bool IsRunning() const;
+				// Begin FRunnable interface.
+				virtual bool Init() override;
+				virtual uint32 Run() override;
+				virtual void Stop() override;
+				virtual void Exit() override;
+				// End FRunnable interface
+			private:
+				bool bIsRunning;
+				FEvent* JobAddedEvent;
+				TArray<FPollJob> PollJobs;
+				TQueue<FPollJob, EQueueMode::Mpsc> JobsToAdd;
+			};
 
-            private:
-                FPollerRunnable Runnable;
-                FRunnableThread* Thread;
+		private:
+			FPollerRunnable Runnable;
+			FRunnableThread* Thread;
 
-                static FGPUFencePoller* Instance;
-        };
-    }
-}
+			static FGPUFencePoller* Instance;
+		};
+	} // namespace PixelStreaming
+} // namespace UE
diff --git a/Engine/Plugins/Media/PixelStreaming/Source/PixelStreaming/Private/IPixelStreamingSessions.h b/Engine/Plugins/Media/PixelStreaming/Source/PixelStreaming/Private/IPixelStreamingSessions.h
deleted file mode 100644
diff --git a/Engine/Plugins/Media/PixelStreaming/Source/PixelStreaming/Private/IPlayerSession.h b/Engine/Plugins/Media/PixelStreaming/Source/PixelStreaming/Private/IPlayerSession.h
--- a/Engine/Plugins/Media/PixelStreaming/Source/PixelStreaming/Private/IPlayerSession.h
+++ b/Engine/Plugins/Media/PixelStreaming/Source/PixelStreaming/Private/IPlayerSession.h
@@ -2,7 +2,7 @@
 
 #pragma once
 
-#include "ProtocolDefs.h"
+#include "PixelStreamingProtocolDefs.h"
 #include "PixelStreamingPlayerId.h"
 
 class IPixelStreamingAudioSink;
diff --git a/Engine/Plugins/Media/PixelStreaming/Source/PixelStreaming/Private/InputDevice.cpp b/Engine/Plugins/Media/PixelStreaming/Source/PixelStreaming/Private/InputDevice.cpp
--- a/Engine/Plugins/Media/PixelStreaming/Source/PixelStreaming/Private/InputDevice.cpp
+++ b/Engine/Plugins/Media/PixelStreaming/Source/PixelStreaming/Private/InputDevice.cpp
@@ -4,7 +4,7 @@
 #include "PixelStreamingInputComponent.h"
 #include "Settings.h"
 #include "PixelStreamingModule.h"
-#include "ProtocolDefs.h"
+#include "PixelStreamingProtocolDefs.h"
 #include "JavaScriptKeyCodes.inl"
 #include "Engine/Engine.h"
 #include "Engine/GameEngine.h"
@@ -28,6 +28,29 @@ namespace UE
 {
 	namespace PixelStreaming
 	{
+		template <typename T>
+		static const T& ParseBuffer(const uint8*& Data, uint32& Size)
+		{
+			checkf(sizeof(T) <= Size, TEXT("%d - %d"), sizeof(T), Size);
+			const T& Value = *reinterpret_cast<const T*>(Data);
+			Data += sizeof(T);
+			Size -= sizeof(T);
+			return Value;
+		}
+
+		static FString ParseString(const webrtc::DataBuffer& Buffer, const size_t Offset = 0)
+		{
+			FString Res;
+			if (Buffer.data.size() > Offset)
+			{
+				size_t StringLength = (Buffer.data.size() - Offset) / sizeof(TCHAR);
+				Res.GetCharArray().SetNumUninitialized(StringLength + 1);
+				FMemory::Memcpy(Res.GetCharArray().GetData(), Buffer.data.data() + Offset, StringLength * sizeof(TCHAR));
+				Res.GetCharArray()[StringLength] = 0;
+			}
+			return Res;
+		}
+
 		/**
 		 * When reading input from a browser then the cursor position will be sent
 		 * across with mouse events. We want to use this position and avoid getting the
@@ -414,7 +437,6 @@ void UE::PixelStreaming::FInputDevice::Tick(float DeltaTime)
 		}
 	}
 
-	
 	FString Command;
 	while (Commands.Dequeue(Command))
 	{
@@ -631,7 +653,7 @@ void UE::PixelStreaming::FInputDevice::FindFocusedWidget()
 namespace
 {
 
-#define GET(Type, Var) Type Var = UE::PixelStreaming::Protocol::ParseBuffer<Type>(Data, Size)
+#define GET(Type, Var) Type Var = UE::PixelStreaming::ParseBuffer<Type>(Data, Size)
 
 	// XY positions are the ratio (0.0..1.0) along a viewport axis, quantized
 	// into an uint16 (0..65536). This allows the browser viewport and player
@@ -753,14 +775,14 @@ void UE::PixelStreaming::FInputDevice::OnMessage(const webrtc::DataBuffer& Buffe
 	{
 		case EToStreamerMsg::UIInteraction:
 		{
-			FString Descriptor = UE::PixelStreaming::Protocol::ParseString(Buffer, UE::PixelStreaming::FInputDevice::GetMessageHeaderOffset());
+			FString Descriptor = UE::PixelStreaming::ParseString(Buffer, UE::PixelStreaming::FInputDevice::GetMessageHeaderOffset());
 			UE_LOG(LogPixelStreamingInput, Verbose, TEXT("UIInteraction: %s"), *Descriptor);
 			ProcessUIInteraction(Descriptor);
 			break;
 		}
 		case EToStreamerMsg::Command:
 		{
-			FString Descriptor = UE::PixelStreaming::Protocol::ParseString(Buffer, UE::PixelStreaming::FInputDevice::GetMessageHeaderOffset());
+			FString Descriptor = UE::PixelStreaming::ParseString(Buffer, UE::PixelStreaming::FInputDevice::GetMessageHeaderOffset());
 			UE_LOG(LogPixelStreamingInput, Verbose, TEXT("Command: %s"), *Descriptor);
 			ProcessCommand(Descriptor);
 			break;
diff --git a/Engine/Plugins/Media/PixelStreaming/Source/PixelStreaming/Private/PixelStreamingModule.cpp b/Engine/Plugins/Media/PixelStreaming/Source/PixelStreaming/Private/PixelStreamingModule.cpp
--- a/Engine/Plugins/Media/PixelStreaming/Source/PixelStreaming/Private/PixelStreamingModule.cpp
+++ b/Engine/Plugins/Media/PixelStreaming/Source/PixelStreaming/Private/PixelStreamingModule.cpp
@@ -46,6 +46,7 @@ THIRD_PARTY_INCLUDES_END
 #include "IImageWrapperModule.h"
 #include "Async/Async.h"
 #include "Engine/Engine.h"
+#include "EncoderFactory.h"
 
 #if !UE_BUILD_SHIPPING
 	#include "DrawDebugHelpers.h"
@@ -299,7 +300,6 @@ namespace UE::PixelStreaming
 		return ReadyEvent;
 	}
 
-
 	IPixelStreamingModule::FStreamingStartedEvent& FPixelStreamingModule::OnStreamingStarted()
 	{
 		return StreamingStartedEvent;
@@ -494,7 +494,7 @@ namespace UE::PixelStreaming
 			return nullptr;
 		}
 
-		return Streamer->GetPlayerSessions().ForSession<IPixelStreamingAudioSink*>(PlayerId, [](auto Session) { return Session->GetAudioSink(); });
+		return Streamer->GetPlayerSessions().ForSession<IPixelStreamingAudioSink*>(PlayerId, [](TSharedPtr<IPlayerSession> Session) { return Session->GetAudioSink(); });
 	}
 
 	IPixelStreamingAudioSink* FPixelStreamingModule::GetUnlistenedAudioSink()
@@ -517,4 +517,24 @@ namespace UE::PixelStreaming
 	}
 } // namespace UE::PixelStreaming
 
+void UE::PixelStreaming::FPixelStreamingModule::AddGPUFencePollerTask(FGPUFenceRHIRef Fence, TSharedRef<bool, ESPMode::ThreadSafe> bIsEnabled, TFunction<void()> Task)
+{
+	FencePollerThread.AddJob(Fence, bIsEnabled, Task);
+}
+
+void UE::PixelStreaming::FPixelStreamingModule::RegisterVideoSource(FPixelStreamingPlayerId PlayerId, IPumpedVideoSource* VideoSource)
+{
+	PumpThread.RegisterVideoSource(PlayerId, VideoSource);
+}
+
+void UE::PixelStreaming::FPixelStreamingModule::UnregisterVideoSource(FPixelStreamingPlayerId PlayerId)
+{
+	PumpThread.UnregisterVideoSource(PlayerId);
+}
+
+webrtc::VideoEncoderFactory* UE::PixelStreaming::FPixelStreamingModule::CreateVideoEncoderFactory()
+{
+	return new UE::PixelStreaming::FVideoEncoderFactory();
+}
+
 IMPLEMENT_MODULE(UE::PixelStreaming::FPixelStreamingModule, PixelStreaming)
diff --git a/Engine/Plugins/Media/PixelStreaming/Source/PixelStreaming/Private/PixelStreamingModule.h b/Engine/Plugins/Media/PixelStreaming/Source/PixelStreaming/Private/PixelStreamingModule.h
--- a/Engine/Plugins/Media/PixelStreaming/Source/PixelStreaming/Private/PixelStreamingModule.h
+++ b/Engine/Plugins/Media/PixelStreaming/Source/PixelStreaming/Private/PixelStreamingModule.h
@@ -6,6 +6,8 @@
 #include "RHI.h"
 #include "Tickable.h"
 #include "InputDevice.h"
+#include "GPUFencePoller.h"
+#include "FixedFPSPump.h"
 
 class AController;
 class AGameModeBase;
@@ -18,9 +20,9 @@ namespace UE::PixelStreaming
 {
 	class FStreamer;
 
-	/**
-	* This plugin allows the back buffer to be sent as a compressed video across a network.
-	*/
+	/*
+	 * This plugin allows the back buffer to be sent as a compressed video across a network.
+	 */
 	class FPixelStreamingModule : public IPixelStreamingModule, public FTickableGameObject
 	{
 	public:
@@ -29,6 +31,11 @@ namespace UE::PixelStreaming
 		virtual bool StartStreaming(const FString& SignallingServerUrl) override;
 		virtual void StopStreaming() override;
 
+		virtual webrtc::VideoEncoderFactory* CreateVideoEncoderFactory() override;
+		virtual void RegisterVideoSource(FPixelStreamingPlayerId PlayerId, IPumpedVideoSource* VideoSource) override;
+		virtual void UnregisterVideoSource(FPixelStreamingPlayerId PlayerId) override;
+		virtual void AddGPUFencePollerTask(FGPUFenceRHIRef Fence, TSharedRef<bool, ESPMode::ThreadSafe> bIsEnabled, TFunction<void()> Task) override;
+
 	private:
 		/** IModuleInterface implementation */
 		void StartupModule() override;
@@ -47,10 +54,10 @@ namespace UE::PixelStreaming
 		void SendCommand(const FString& Descriptor) override;
 
 		/*
-		* Returns a shared pointer to the device which handles pixel streaming
-		* input.
-		* @return The shared pointer to the input device.
-		*/
+		 * Returns a shared pointer to the device which handles pixel streaming
+		 * input.
+		 * @return The shared pointer to the input device.
+		 */
 		TSharedPtr<FInputDevice> GetInputDevicePtr();
 		void AddInputComponent(UPixelStreamingInput* InInputComponent) override;
 		void RemoveInputComponent(UPixelStreamingInput* InInputComponent) override;
@@ -88,5 +95,8 @@ namespace UE::PixelStreaming
 		bool bCaptureNextBackBufferAndStream = false;
 		double LastVideoEncoderQPReportTime = 0;
 		static IPixelStreamingModule* PixelStreamingModule;
+
+		FFixedFPSPump PumpThread;
+		FGPUFencePoller FencePollerThread;
 	};
 } // namespace UE::PixelStreaming
diff --git a/Engine/Plugins/Media/PixelStreaming/Source/PixelStreaming/Private/PlayerSession.cpp b/Engine/Plugins/Media/PixelStreaming/Source/PixelStreaming/Private/PlayerSession.cpp
--- a/Engine/Plugins/Media/PixelStreaming/Source/PixelStreaming/Private/PlayerSession.cpp
+++ b/Engine/Plugins/Media/PixelStreaming/Source/PixelStreaming/Private/PlayerSession.cpp
@@ -26,7 +26,7 @@ namespace UE::PixelStreaming
 		// Begin IPixelStreamingStatsConsumer
 		void ConsumeStat(FPixelStreamingPlayerId PlayerId, FName StatName, float StatValue) override
 		{
-			Sessions->ForSession(PlayerId, [&](auto Session) {
+			Sessions->ForSession(PlayerId, [&](TSharedPtr<IPlayerSession> Session) {
 				Session->SendVideoEncoderQP((int)StatValue);
 			});
 		}
@@ -78,7 +78,7 @@ namespace UE::PixelStreaming
 	void FPlayerSession::PollWebRTCStats() const
 	{
 		check(PeerConnection);
-		//PeerConnection->GetStats(WebRTCStatsCallback);
+		// PeerConnection->GetStats(WebRTCStatsCallback);
 		std::vector<rtc::scoped_refptr<webrtc::RtpTransceiverInterface>> Transceivers = PeerConnection->GetTransceivers();
 		for (rtc::scoped_refptr<webrtc::RtpTransceiverInterface> Transceiver : Transceivers)
 		{
@@ -254,8 +254,10 @@ namespace UE::PixelStreaming
 	{
 		// Send the mime type first
 		FPlayerSession::SendMessage(Protocol::EToPlayerMsg::FileMimeType, MimeType);
+
 		// Send the extension next
 		FPlayerSession::SendMessage(Protocol::EToPlayerMsg::FileExtension, FileExtension);
+
 		// Send the contents of the file
 		AsyncTask(ENamedThreads::AnyHiPriThreadHiPriTask, [this, ByteData]() {
 			FPlayerSession::SendArbitraryData(ByteData, static_cast<uint8>(Protocol::EToPlayerMsg::FileContents));
@@ -310,8 +312,10 @@ namespace UE::PixelStreaming
 
 			// Write message type
 			Pos = SerializeToBuffer(Buffer, Pos, &MessageType, sizeof(MessageType));
+
 			// Write size of payload
 			Pos = SerializeToBuffer(Buffer, Pos, &DataSize, sizeof(DataSize));
+
 			// Write the data bytes payload
 			Pos = SerializeToBuffer(Buffer, Pos, DataBytes.GetData() + BytesTransmitted, BytesToTransmit);
 
@@ -438,6 +442,7 @@ namespace UE::PixelStreaming
 		}
 
 		rtc::scoped_refptr<webrtc::RtpReceiverInterface> Receiver = transceiver->receiver();
+
 		if (mediaType != cricket::MediaType::MEDIA_TYPE_AUDIO)
 		{
 			return;
diff --git a/Engine/Plugins/Media/PixelStreaming/Source/PixelStreaming/Private/PlayerSession.h b/Engine/Plugins/Media/PixelStreaming/Source/PixelStreaming/Private/PlayerSession.h
--- a/Engine/Plugins/Media/PixelStreaming/Source/PixelStreaming/Private/PlayerSession.h
+++ b/Engine/Plugins/Media/PixelStreaming/Source/PixelStreaming/Private/PlayerSession.h
@@ -7,6 +7,7 @@
 #include "HAL/ThreadSafeBool.h"
 #include "AudioSink.h"
 
+class IPixelStreamingSessions;
 class IPixelStreamingStatsConsumer;
 
 namespace UE::PixelStreaming
diff --git a/Engine/Plugins/Media/PixelStreaming/Source/PixelStreaming/Private/PlayerSessions.cpp b/Engine/Plugins/Media/PixelStreaming/Source/PixelStreaming/Private/PlayerSessions.cpp
--- a/Engine/Plugins/Media/PixelStreaming/Source/PixelStreaming/Private/PlayerSessions.cpp
+++ b/Engine/Plugins/Media/PixelStreaming/Source/PixelStreaming/Private/PlayerSessions.cpp
@@ -1,20 +1,21 @@
 // Copyright Epic Games, Inc. All Rights Reserved.
 
 #include "PlayerSessions.h"
+#include "PixelStreamingDelegates.h"
+#include "PixelStreamingPrivate.h"
 #include "PlayerSession.h"
-#include "PlayerSessionSFU.h"
 #include "PlayerSessionDataOnly.h"
-#include "PixelStreamingPrivate.h"
+#include "PlayerSessionSFU.h"
 #include "Settings.h"
-#include "PixelStreamingDelegates.h"
 
 namespace UE::PixelStreaming
 {
-	TSharedPtr<IPlayerSession> FPlayerSessions::CreatePlayerSession(FPixelStreamingPlayerId PlayerId,
-		rtc::scoped_refptr<webrtc::PeerConnectionFactoryInterface> PeerConnectionFactory,
+	TSharedPtr<IPlayerSession> FPlayerSessions::CreatePlayerSession(
+		FPixelStreamingPlayerId PlayerId,
+		rtc::scoped_refptr<webrtc::PeerConnectionFactoryInterface>
+			PeerConnectionFactory,
 		webrtc::PeerConnectionInterface::RTCConfiguration PeerConnectionConfig,
-		FSignallingServerConnection* SignallingServerConnection,
-		int Flags)
+		FSignallingServerConnection *SignallingServerConnection, int Flags)
 	{
 		check(PeerConnectionFactory);
 
@@ -26,33 +27,46 @@ namespace UE::PixelStreaming
 			return nullptr;
 		}
 
-		UE_LOG(LogPixelStreaming, Log, TEXT("Creating player session for PlayerId=%s"), *PlayerId);
+		UE_LOG(LogPixelStreaming, Log,
+			   TEXT("Creating player session for PlayerId=%s"), *PlayerId);
 
-		// this is called from WebRTC signalling thread, the only thread where `Players` map is modified, so no need to lock it
+		// this is called from WebRTC signalling thread, the only thread where
+		// `Players` map is modified, so no need to lock it
 		bool bIsSFU = (Flags & Protocol::EPlayerFlags::PSPFlag_IsSFU) != 0;
-		bool bMakeQualityController = !bIsSFU && Players.Num() == 0; // first player controls quality by default
+		bool bMakeQualityController =
+			!bIsSFU && Players.Num() == 0; // first player controls quality by default
 		TSharedPtr<IPlayerSession> Session;
 		if (bIsSFU)
 		{
-			Session = MakeShared<FPlayerSessionSFU>(this, SignallingServerConnection, PlayerId);
+			Session = MakeShared<FPlayerSessionSFU>(this, SignallingServerConnection,
+													PlayerId);
 		}
 		else
 		{
-			Session = MakeShared<FPlayerSession>(this, SignallingServerConnection, PlayerId);
+			Session =
+				MakeShared<FPlayerSession>(this, SignallingServerConnection, PlayerId);
 		}
 
-		rtc::scoped_refptr<webrtc::PeerConnectionInterface> PeerConnection = PeerConnectionFactory->CreatePeerConnection(PeerConnectionConfig, webrtc::PeerConnectionDependencies{ Session.Get() });
+		rtc::scoped_refptr<webrtc::PeerConnectionInterface> PeerConnection =
+			PeerConnectionFactory->CreatePeerConnection(
+				PeerConnectionConfig,
+				webrtc::PeerConnectionDependencies{Session.Get()});
 		if (!PeerConnection)
 		{
-			UE_LOG(LogPixelStreaming, Error, TEXT("Failed to created PeerConnection. This may indicate you passed malformed peerConnectionOptions."));
+			UE_LOG(LogPixelStreaming, Error,
+				   TEXT("Failed to created PeerConnection. This may indicate you "
+						"passed malformed peerConnectionOptions."));
 			return nullptr;
 		}
 
 		// Setup suggested bitrate settings on the Peer Connection based on our CVars
 		webrtc::BitrateSettings BitrateSettings;
-		BitrateSettings.min_bitrate_bps = Settings::CVarPixelStreamingWebRTCMinBitrate.GetValueOnAnyThread();
-		BitrateSettings.max_bitrate_bps = Settings::CVarPixelStreamingWebRTCMaxBitrate.GetValueOnAnyThread();
-		BitrateSettings.start_bitrate_bps = Settings::CVarPixelStreamingWebRTCStartBitrate.GetValueOnAnyThread();
+		BitrateSettings.min_bitrate_bps =
+			Settings::CVarPixelStreamingWebRTCMinBitrate.GetValueOnAnyThread();
+		BitrateSettings.max_bitrate_bps =
+			Settings::CVarPixelStreamingWebRTCMaxBitrate.GetValueOnAnyThread();
+		BitrateSettings.start_bitrate_bps =
+			Settings::CVarPixelStreamingWebRTCStartBitrate.GetValueOnAnyThread();
 		PeerConnection->SetBitrate(BitrateSettings);
 
 		Session->SetPeerConnection(PeerConnection);
@@ -68,25 +82,29 @@ namespace UE::PixelStreaming
 			SetQualityController(PlayerId);
 		}
 
-		if (UPixelStreamingDelegates* Delegates = UPixelStreamingDelegates::GetPixelStreamingDelegates())
+		if (UPixelStreamingDelegates *Delegates =
+				UPixelStreamingDelegates::GetPixelStreamingDelegates())
 		{
 			Delegates->OnNewConnection.Broadcast(PlayerId, bMakeQualityController);
-			Delegates->OnNewConnectionNative.Broadcast(PlayerId, bMakeQualityController);
+			Delegates->OnNewConnectionNative.Broadcast(PlayerId,
+													   bMakeQualityController);
 		}
 
 		return Session;
 	}
 
-	void FPlayerSessions::CreateNewDataChannel(FPixelStreamingPlayerId SFUId,
-		FPixelStreamingPlayerId PlayerId,
-		int32 SendStreamId,
-		int32 RecvStreamId,
-		FSignallingServerConnection* SignallingServerConnection)
+	void FPlayerSessions::CreateNewDataChannel(
+		FPixelStreamingPlayerId SFUId, FPixelStreamingPlayerId PlayerId,
+		int32 SendStreamId, int32 RecvStreamId,
+		FSignallingServerConnection *SignallingServerConnection)
 	{
-		if (auto ParentSession = StaticCastSharedPtr<FPlayerSessionSFU>(GetPlayerSession(SFUId)))
+		if (auto ParentSession =
+				StaticCastSharedPtr<FPlayerSessionSFU>(GetPlayerSession(SFUId)))
 		{
 			auto PeerConnection = &ParentSession->GetPeerConnection();
-			auto NewSession = MakeShared<FPlayerSessionDataOnly>(this, SignallingServerConnection, PlayerId, PeerConnection, SendStreamId, RecvStreamId);
+			auto NewSession = MakeShared<FPlayerSessionDataOnly>(
+				this, SignallingServerConnection, PlayerId, PeerConnection,
+				SendStreamId, RecvStreamId);
 			ParentSession->AddChildSession(NewSession);
 
 			{
@@ -101,7 +119,9 @@ namespace UE::PixelStreaming
 		TSharedPtr<IPlayerSession> PlayerSession = GetPlayerSession(PlayerId);
 		if (!PlayerSession)
 		{
-			UE_LOG(LogPixelStreaming, VeryVerbose, TEXT("Failed to delete player %s - that player was not found."), *PlayerId);
+			UE_LOG(LogPixelStreaming, VeryVerbose,
+				   TEXT("Failed to delete player %s - that player was not found."),
+				   *PlayerId);
 			return GetNumPlayers();
 		}
 
@@ -115,11 +135,13 @@ namespace UE::PixelStreaming
 			RemainingCount = Players.Num();
 		}
 
-		UPixelStreamingDelegates* Delegates = UPixelStreamingDelegates::GetPixelStreamingDelegates();
+		UPixelStreamingDelegates *Delegates =
+			UPixelStreamingDelegates::GetPixelStreamingDelegates();
 		if (Delegates && FModuleManager::Get().IsModuleLoaded("PixelStreaming"))
 		{
 			Delegates->OnClosedConnection.Broadcast(PlayerId, bWasQualityController);
-			Delegates->OnClosedConnectionNative.Broadcast(PlayerId, bWasQualityController);
+			Delegates->OnClosedConnectionNative.Broadcast(PlayerId,
+														  bWasQualityController);
 
 			if (RemainingCount == 0)
 			{
@@ -132,7 +154,8 @@ namespace UE::PixelStreaming
 
 		if (bWasQualityController && RemainingCount > 0)
 		{
-			// Quality Controller session has been just removed, set quality control to any of remaining sessions
+			// Quality Controller session has been just removed, set quality control to
+			// any of remaining sessions
 			SetQualityController(Players.begin()->Key);
 		}
 
@@ -145,8 +168,8 @@ namespace UE::PixelStreaming
 		TSet<FPixelStreamingPlayerId> OldPlayerIds;
 		Players.GetKeys(OldPlayerIds);
 
-		// Do the modification to the critical parts separately since quality controller
-		// lock is requested in the destructor of the video source base.
+		// Do the modification to the critical parts separately since quality
+		// controller lock is requested in the destructor of the video source base.
 		{
 			FScopeLock QualityControllerLock(&QualityControllerCS);
 			QualityControllingPlayer = INVALID_PLAYER_ID;
@@ -154,19 +177,20 @@ namespace UE::PixelStreaming
 		{
 			FScopeLock PlayersLock(&PlayersCS);
 			Players.Empty();
-		
 		}
 
 		// Notify all delegates of all the closed players
-		UPixelStreamingDelegates* Delegates = UPixelStreamingDelegates::GetPixelStreamingDelegates();
+		UPixelStreamingDelegates *Delegates =
+			UPixelStreamingDelegates::GetPixelStreamingDelegates();
 		if (Delegates && FModuleManager::Get().IsModuleLoaded("PixelStreaming"))
 		{
-			for (auto&& PlayerId : OldPlayerIds)
+			for (auto &&PlayerId : OldPlayerIds)
 			{
 				bool bWasQualityController = OldQualityController == PlayerId;
 
 				Delegates->OnClosedConnection.Broadcast(PlayerId, bWasQualityController);
-				Delegates->OnClosedConnectionNative.Broadcast(PlayerId, bWasQualityController);
+				Delegates->OnClosedConnectionNative.Broadcast(PlayerId,
+															  bWasQualityController);
 			}
 
 			Delegates->OnAllConnectionsClosed.Broadcast();
@@ -174,12 +198,10 @@ namespace UE::PixelStreaming
 		}
 	}
 
-	int FPlayerSessions::GetNumPlayers() const
-	{
-		return Players.Num();
-	}
+	int FPlayerSessions::GetNumPlayers() const { return Players.Num(); }
 
-	bool FPlayerSessions::IsQualityController(FPixelStreamingPlayerId PlayerId) const
+	bool FPlayerSessions::IsQualityController(
+		FPixelStreamingPlayerId PlayerId) const
 	{
 		FScopeLock Lock(&QualityControllerCS);
 		return QualityControllingPlayer == PlayerId;
@@ -194,34 +216,35 @@ namespace UE::PixelStreaming
 		}
 
 		// Let any listeners know the quality controller has changed
-		if (UPixelStreamingDelegates* Delegates = UPixelStreamingDelegates::GetPixelStreamingDelegates())
+		if (UPixelStreamingDelegates *Delegates = UPixelStreamingDelegates::GetPixelStreamingDelegates())
 		{
 			Delegates->OnQualityControllerChangedNative.Broadcast(PlayerId);
 		}
 
-		UE_LOG(LogPixelStreaming, Log, TEXT("Quality controller is now PlayerId=%s."), *PlayerId);
+		UE_LOG(LogPixelStreaming, Log, TEXT("Quality controller is now PlayerId=%s."),
+			   *PlayerId);
 
 		// Update quality controller status on the browser side too
-		ForEachSession([&](TSharedPtr<IPlayerSession> Session) {
+		ForEachSession([&](TSharedPtr<IPlayerSession> Session)
+					   {
 			bool bIsQualityController = Session->GetPlayerId() == QualityControllingPlayer;
-			Session->SendQualityControlStatus(bIsQualityController);
-		});
+			Session->SendQualityControlStatus(bIsQualityController); });
 	}
 
-	void FPlayerSessions::ForEachSession(const TFunction<void(TSharedPtr<IPlayerSession>)>& Func)
+	void FPlayerSessions::ForEachSession(const TFunction<void(TSharedPtr<IPlayerSession>)> &Func)
 	{
 		// We have to be careful here
-		// We dont lock here because we can end up in a situation where Func dispatches to
-		// another thread and that thread is already waiting to acquire this lock and bam
-		// we have a deadlock. Instead we grab the current list of keys and then iterate
-		// over our copy of those. We DO lock while checking and pulling the session just
-		// in case Players is being modified between Contains and when we actually pull it.
-		// luckily GetPlayerSession handles that for us.
+		// We dont lock here because we can end up in a situation where Func
+		// dispatches to another thread and that thread is already waiting to acquire
+		// this lock and bam we have a deadlock. Instead we grab the current list of
+		// keys and then iterate over our copy of those. We DO lock while checking and
+		// pulling the session just in case Players is being modified between Contains
+		// and when we actually pull it. luckily GetPlayerSession handles that for us.
 		TSet<FPixelStreamingPlayerId> KeySet;
 		Players.GetKeys(KeySet);
-		for (auto&& PlayerId : KeySet)
+		for (auto &&PlayerId : KeySet)
 		{
-			if (auto Session = GetPlayerSession(PlayerId))
+			if (TSharedPtr<IPlayerSession> Session = GetPlayerSession(PlayerId))
 			{
 				Func(Session);
 			}
diff --git a/Engine/Plugins/Media/PixelStreaming/Source/PixelStreaming/Private/PlayerSessions.h b/Engine/Plugins/Media/PixelStreaming/Source/PixelStreaming/Private/PlayerSessions.h
--- a/Engine/Plugins/Media/PixelStreaming/Source/PixelStreaming/Private/PlayerSessions.h
+++ b/Engine/Plugins/Media/PixelStreaming/Source/PixelStreaming/Private/PlayerSessions.h
@@ -38,7 +38,7 @@ namespace UE::PixelStreaming
 		// do something with a single session
 		void ForSession(FPixelStreamingPlayerId PlayerId, const TFunction<void(TSharedPtr<IPlayerSession>)>& Func) const
 		{
-			if (auto Session = GetPlayerSession(PlayerId))
+			if (TSharedPtr<IPlayerSession> Session = GetPlayerSession(PlayerId))
 			{
 				Func(Session);
 			}
@@ -48,7 +48,7 @@ namespace UE::PixelStreaming
 		template <typename R, typename T>
 		R ForSession(FPixelStreamingPlayerId PlayerId, T&& Func) const
 		{
-			if (auto Session = GetPlayerSession(PlayerId))
+			if (TSharedPtr<IPlayerSession> Session = GetPlayerSession(PlayerId))
 			{
 				return Func(Session);
 			}
diff --git a/Engine/Plugins/Media/PixelStreaming/Source/PixelStreaming/Private/ProtocolDefs.h b/Engine/Plugins/Media/PixelStreaming/Source/PixelStreaming/Private/ProtocolDefs.h
deleted file mode 100644
diff --git a/Engine/Plugins/Media/PixelStreaming/Source/PixelStreaming/Private/SignallingServerConnection.cpp b/Engine/Plugins/Media/PixelStreaming/Source/PixelStreaming/Private/SignallingServerConnection.cpp
--- a/Engine/Plugins/Media/PixelStreaming/Source/PixelStreaming/Private/SignallingServerConnection.cpp
+++ b/Engine/Plugins/Media/PixelStreaming/Source/PixelStreaming/Private/SignallingServerConnection.cpp
@@ -9,7 +9,7 @@
 #include "Settings.h"
 #include "TimerManager.h"
 #include "PixelStreamingDelegates.h"
-#include "ProtocolDefs.h"
+#include "PixelStreamingProtocolDefs.h"
 
 DECLARE_LOG_CATEGORY_EXTERN(LogPixelStreamingSS, Log, VeryVerbose);
 DEFINE_LOG_CATEGORY(LogPixelStreamingSS);
@@ -17,8 +17,7 @@ DEFINE_LOG_CATEGORY(LogPixelStreamingSS);
 namespace UE::PixelStreaming
 {
 	FSignallingServerConnection::FSignallingServerConnection(FSignallingServerConnectionObserver& InObserver, FString InStreamerId)
-		: Observer(InObserver),
-		  StreamerId(InStreamerId)
+		: Observer(InObserver), StreamerId(InStreamerId)
 	{
 		RegisterHandler("identify", [this](FJsonObjectPtr JsonMsg) { OnIdRequested(); });
 		RegisterHandler("config", [this](FJsonObjectPtr JsonMsg) { OnConfig(JsonMsg); });
diff --git a/Engine/Plugins/Media/PixelStreaming/Source/PixelStreaming/Private/SimulcastEncoderAdapter.cpp b/Engine/Plugins/Media/PixelStreaming/Source/PixelStreaming/Private/SimulcastEncoderAdapter.cpp
--- a/Engine/Plugins/Media/PixelStreaming/Source/PixelStreaming/Private/SimulcastEncoderAdapter.cpp
+++ b/Engine/Plugins/Media/PixelStreaming/Source/PixelStreaming/Private/SimulcastEncoderAdapter.cpp
@@ -1,6 +1,7 @@
 // Copyright Epic Games, Inc. All Rights Reserved.
 #include "SimulcastEncoderAdapter.h"
-#include "FrameBuffer.h"
+#include "PixelStreamingFrameBuffer.h"
+#include "IPixelStreamingTextureSource.h"
 
 namespace
 {
@@ -121,7 +122,7 @@ int UE::PixelStreaming::FSimulcastEncoderAdapter::Release()
 			StreamInfos.pop_back(); // Deletes callback adapter.
 		}
 	}
-	
+
 	Initialized = false;
 
 	return WEBRTC_VIDEO_CODEC_OK;
@@ -173,7 +174,7 @@ int UE::PixelStreaming::FSimulcastEncoderAdapter::InitEncode(const webrtc::Video
 
 			FVideoEncoderFactory* EncoderFactory = SimulcastEncoderFactory.GetOrCreateEncoderFactory(i);
 			std::unique_ptr<VideoEncoder> Encoder = EncoderFactory->CreateVideoEncoder(Format);
-			
+
 			// std::unique_ptr<VideoEncoder> Encoder = SimulcastEncoderFactory->CreateVideoEncoder(Format);
 			int ReturnCode = Encoder->InitEncode(&StreamCodec, settings);
 			if (ReturnCode < 0)
@@ -199,10 +200,10 @@ int UE::PixelStreaming::FSimulcastEncoderAdapter::InitEncode(const webrtc::Video
 	return WEBRTC_VIDEO_CODEC_OK;
 }
 
-int UE::PixelStreaming::FSimulcastEncoderAdapter::EncodeStream(const webrtc::VideoFrame& InputImage, TSharedPtr<UE::PixelStreaming::ITextureSource> LayerFrameSource, size_t StreamIdx, bool bSendKeyFrame)
+int UE::PixelStreaming::FSimulcastEncoderAdapter::EncodeStream(const webrtc::VideoFrame& InputImage, TSharedPtr<IPixelStreamingTextureSource> LayerFrameSource, size_t StreamIdx, bool bSendKeyFrame)
 {
 	// grab the simulcast frame source, extract the frame source for this layer and wrap that in a new frame buffer
-	rtc::scoped_refptr<UE::PixelStreaming::FLayerFrameBuffer> LayerFrameBuffer = new rtc::RefCountedObject<UE::PixelStreaming::FLayerFrameBuffer>(LayerFrameSource);
+	rtc::scoped_refptr<FLayerFrameBuffer> LayerFrameBuffer = new rtc::RefCountedObject<FLayerFrameBuffer>(LayerFrameSource);
 	webrtc::VideoFrame NewFrame(InputImage);
 	NewFrame.set_video_frame_buffer(LayerFrameBuffer);
 
@@ -245,15 +246,15 @@ int UE::PixelStreaming::FSimulcastEncoderAdapter::Encode(const webrtc::VideoFram
 	// ignore any init frames
 	// NOTE: It seems most of the time they never even get here because the internals of WebRTC decide to drop it. It seems it thinks
 	// the encoder is paused or something.
-	const UE::PixelStreaming::FFrameBuffer* FrameBuffer = static_cast<UE::PixelStreaming::FFrameBuffer*>(input_image.video_frame_buffer().get());
-	if (FrameBuffer->GetFrameBufferType() == UE::PixelStreaming::FFrameBufferType::Initialize)
+	const FPixelStreamingFrameBuffer* FrameBuffer = static_cast<FPixelStreamingFrameBuffer*>(input_image.video_frame_buffer().get());
+	if (FrameBuffer->GetFrameBufferType() == UE::PixelStreaming::EFrameBufferType::Initialize)
 	{
 		return WEBRTC_VIDEO_CODEC_OK;
 	}
 
 	// separate out our pixelstreaming image sources
-	check(FrameBuffer->GetFrameBufferType() == UE::PixelStreaming::FFrameBufferType::Simulcast);
-	const UE::PixelStreaming::FSimulcastFrameBuffer* SimulcastFrameBuffer = static_cast<const UE::PixelStreaming::FSimulcastFrameBuffer*>(FrameBuffer);
+	check(FrameBuffer->GetFrameBufferType() == UE::PixelStreaming::EFrameBufferType::Simulcast);
+	const FSimulcastFrameBuffer* SimulcastFrameBuffer = static_cast<const FSimulcastFrameBuffer*>(FrameBuffer);
 
 	// All active streams should generate a key frame if
 	// a key frame is requested by any stream.
@@ -287,7 +288,7 @@ int UE::PixelStreaming::FSimulcastEncoderAdapter::Encode(const webrtc::VideoFram
 			continue;
 		}
 
-		TSharedPtr<ITextureSource> LayerFrameSource = SimulcastFrameBuffer->GetLayerFrameSource(StreamInfos.size() == 1 ? SimulcastFrameBuffer->GetNumLayers() - 1 : StreamIdx);
+		TSharedPtr<IPixelStreamingTextureSource> LayerFrameSource = SimulcastFrameBuffer->GetLayerFrameSource(StreamInfos.size() == 1 ? SimulcastFrameBuffer->GetNumLayers() - 1 : StreamIdx);
 
 		int RtcError = EncodeStream(input_image, LayerFrameSource, StreamIdx, bSendKeyFrame);
 		if (RtcError != WEBRTC_VIDEO_CODEC_OK)
diff --git a/Engine/Plugins/Media/PixelStreaming/Source/PixelStreaming/Private/SimulcastEncoderAdapter.h b/Engine/Plugins/Media/PixelStreaming/Source/PixelStreaming/Private/SimulcastEncoderAdapter.h
--- a/Engine/Plugins/Media/PixelStreaming/Source/PixelStreaming/Private/SimulcastEncoderAdapter.h
+++ b/Engine/Plugins/Media/PixelStreaming/Source/PixelStreaming/Private/SimulcastEncoderAdapter.h
@@ -5,18 +5,19 @@
 #include "WebRTCIncludes.h"
 #include "EncoderFactory.h"
 
+class IPixelStreamingTextureSource;
+
 namespace webrtc
 {
 	class SimulcastRateAllocator;
 	class VideoEncoderFactory;
 } // namespace webrtc
 
-namespace UE {
-	namespace PixelStreaming {
-		class ITextureSource;
-
+namespace UE
+{
+	namespace PixelStreaming
+	{
 		// This is highly modified version of webrtc::simulcast_encoder_adapter
-
 		class FSimulcastEncoderAdapter : public webrtc::VideoEncoder
 		{
 		public:
@@ -85,7 +86,7 @@ namespace UE {
 
 			int LowestResolutionStreamIndex;
 			int HighestResolutionStreamIndex;
-			int EncodeStream(const webrtc::VideoFrame& InputImage, TSharedPtr<ITextureSource> LayerFrameSource, size_t StreamIdx, bool bSendKeyFrame);
+			int EncodeStream(const webrtc::VideoFrame& InputImage, TSharedPtr<IPixelStreamingTextureSource> LayerFrameSource, size_t StreamIdx, bool bSendKeyFrame);
 		};
-	}
-}
+	} // namespace PixelStreaming
+} // namespace UE
diff --git a/Engine/Plugins/Media/PixelStreaming/Source/PixelStreaming/Private/Stats.cpp b/Engine/Plugins/Media/PixelStreaming/Source/PixelStreaming/Private/Stats.cpp
--- a/Engine/Plugins/Media/PixelStreaming/Source/PixelStreaming/Private/Stats.cpp
+++ b/Engine/Plugins/Media/PixelStreaming/Source/PixelStreaming/Private/Stats.cpp
@@ -367,7 +367,7 @@ namespace UE::PixelStreaming
 		// webrtc::metrics::Histogram* Hist1 = webrtc::metrics::HistogramFactoryGetCounts("WebRTC.Video.NacksSent", 0, 100000, 100);
 		// Will require calling webrtc::metrics::Enable();
 
-		Sessions->ForEachSession([](auto Session) {
+		Sessions->ForEachSession([](TSharedPtr<IPlayerSession> Session) {
 			Session->PollWebRTCStats();
 		});
 
diff --git a/Engine/Plugins/Media/PixelStreaming/Source/PixelStreaming/Private/Streamer.cpp b/Engine/Plugins/Media/PixelStreaming/Source/PixelStreaming/Private/Streamer.cpp
--- a/Engine/Plugins/Media/PixelStreaming/Source/PixelStreaming/Private/Streamer.cpp
+++ b/Engine/Plugins/Media/PixelStreaming/Source/PixelStreaming/Private/Streamer.cpp
@@ -1,6 +1,7 @@
 // Copyright Epic Games, Inc. All Rights Reserved.
 
 #include "Streamer.h"
+#include "WebRTCIncludes.h"
 #include "VideoEncoderFactory.h"
 #include "ToStringExtensions.h"
 #include "AudioCapturer.h"
@@ -22,7 +23,7 @@ namespace UE::PixelStreaming
 		return AVEncoder::FVideoEncoderFactory::Get().HasEncoderForCodec(AVEncoder::ECodecType::H264);
 	}
 
-	FStreamer::FStreamer(const FString& InStreamerId)
+	FStreamer::FStreamer(const FString &InStreamerId)
 		: StreamerId(InStreamerId)
 		, WebRtcSignallingThread(MakeUnique<rtc::Thread>(rtc::SocketServer::CreateDefault()))
 		, SignallingServerConnection(MakeUnique<FSignallingServerConnection>(*this, InStreamerId))
@@ -47,7 +48,7 @@ namespace UE::PixelStreaming
 		rtc::CleanupSSL();
 	}
 
-	bool FStreamer::StartStreaming(const FString& SignallingServerUrl)
+	bool FStreamer::StartStreaming(const FString &SignallingServerUrl)
 	{
 		StopStreaming();
 		ConnectToSignallingServer(SignallingServerUrl);
@@ -129,16 +130,14 @@ namespace UE::PixelStreaming
 		check(SFUPeerConnectionFactory.get() != nullptr);
 	}
 
-	webrtc::AudioProcessing* FStreamer::SetupAudioProcessingModule()
+	webrtc::AudioProcessing *FStreamer::SetupAudioProcessingModule()
 	{
-		webrtc::AudioProcessing* AudioProcessingModule = webrtc::AudioProcessingBuilder().Create();
+		webrtc::AudioProcessing *AudioProcessingModule = webrtc::AudioProcessingBuilder().Create();
 		webrtc::AudioProcessing::Config Config;
-
 		// Enabled multi channel audio capture/render
 		Config.pipeline.multi_channel_capture = true;
 		Config.pipeline.multi_channel_render = true;
 		Config.pipeline.maximum_internal_processing_rate = 48000;
-
 		// Turn off all other audio processing effects in UE's WebRTC. We want to stream audio from UE as pure as possible.
 		Config.pre_amplifier.enabled = false;
 		Config.high_pass_filter.enabled = false;
@@ -150,14 +149,13 @@ namespace UE::PixelStreaming
 		Config.gain_controller2.enabled = false;
 		Config.residual_echo_detector.enabled = false;
 		Config.level_estimation.enabled = false;
-
 		// Apply the config.
 		AudioProcessingModule->ApplyConfig(Config);
 
 		return AudioProcessingModule;
 	}
 
-	void FStreamer::ConnectToSignallingServer(const FString& Url)
+	void FStreamer::ConnectToSignallingServer(const FString &Url)
 	{
 		PreviousSignallingServerURL = Url;
 		SignallingServerConnection->Connect(Url);
@@ -168,12 +166,12 @@ namespace UE::PixelStreaming
 		SignallingServerConnection->Disconnect();
 	}
 
-	void FStreamer::OnConfig(const webrtc::PeerConnectionInterface::RTCConfiguration& Config)
+	void FStreamer::OnConfig(const webrtc::PeerConnectionInterface::RTCConfiguration &Config)
 	{
 		PeerConnectionConfig = Config;
 	}
 
-	void FStreamer::OnDataChannelOpen(FPixelStreamingPlayerId PlayerId, webrtc::DataChannelInterface* DataChannel)
+	void FStreamer::OnDataChannelOpen(FPixelStreamingPlayerId PlayerId, webrtc::DataChannelInterface *DataChannel)
 	{
 		if (DataChannel)
 		{
@@ -190,16 +188,15 @@ namespace UE::PixelStreaming
 
 		rtc::scoped_refptr<webrtc::PeerConnectionFactoryInterface> PCFactory = bIsSFU ? SFUPeerConnectionFactory : P2PPeerConnectionFactory;
 
-		auto NewSession = PlayerSessions.CreatePlayerSession(
+		TSharedPtr<IPlayerSession> NewSession = PlayerSessions.CreatePlayerSession(
 			PlayerId,
 			PCFactory,
 			PeerConnectionConfig,
 			SignallingServerConnection.Get(),
 			Flags);
-
 		if (NewSession)
 		{
-			auto& PeerConnection = NewSession->GetPeerConnection();
+			webrtc::PeerConnectionInterface &PeerConnection = NewSession->GetPeerConnection();
 			const bool SupportsDataChannel = (Flags & Protocol::EPlayerFlags::PSPFlag_SupportsDataChannel) != 0;
 
 			if (SupportsDataChannel)
@@ -217,9 +214,9 @@ namespace UE::PixelStreaming
 				NewSession->SetDataChannel(DataChannel);
 
 				// Bind to the on data channel open delegate (we use this as the earliest time peer is ready for freeze frame)
-				FDataChannelObserver* DataChannelObserver = NewSession->GetDataChannelObserver();
+				FDataChannelObserver *DataChannelObserver = NewSession->GetDataChannelObserver();
 				check(DataChannelObserver);
-				if (UPixelStreamingDelegates* Delegates = UPixelStreamingDelegates::GetPixelStreamingDelegates())
+				if (UPixelStreamingDelegates *Delegates = UPixelStreamingDelegates::GetPixelStreamingDelegates())
 				{
 					Delegates->OnDataChannelOpenNative.AddRaw(this, &FStreamer::OnDataChannelOpen);
 				}
@@ -230,29 +227,28 @@ namespace UE::PixelStreaming
 		return NewSession;
 	}
 
-	void FStreamer::OnSessionDescription(FPixelStreamingPlayerId PlayerId, webrtc::SdpType Type, const FString& Sdp)
+	void FStreamer::OnSessionDescription(FPixelStreamingPlayerId PlayerId, webrtc::SdpType Type, const FString &Sdp)
 	{
 		switch (Type)
 		{
-			case webrtc::SdpType::kOffer:
-				OnOffer(PlayerId, Sdp);
-				break;
-			case webrtc::SdpType::kAnswer:
-			case webrtc::SdpType::kPrAnswer:
-			{
-				PlayerSessions.ForSession(PlayerId, [&](auto Session) {
-					Session->OnAnswer(Sdp);
-				});
-				bStreamingStarted = true;
-				break;
-			}
-			case webrtc::SdpType::kRollback:
-				UE_LOG(LogPixelStreaming, Error, TEXT("Rollback SDP is currently unsupported. SDP is: %s"), *Sdp);
-				break;
+		case webrtc::SdpType::kOffer:
+			OnOffer(PlayerId, Sdp);
+			break;
+		case webrtc::SdpType::kAnswer:
+		case webrtc::SdpType::kPrAnswer:
+		{
+			PlayerSessions.ForSession(PlayerId, [&](TSharedPtr<IPlayerSession> Session)
+									  { Session->OnAnswer(Sdp); });
+			bStreamingStarted = true;
+			break;
+		}
+		case webrtc::SdpType::kRollback:
+			UE_LOG(LogPixelStreaming, Error, TEXT("Rollback SDP is currently unsupported. SDP is: %s"), *Sdp);
+			break;
 		}
 	}
 
-	void FStreamer::OnOffer(FPixelStreamingPlayerId PlayerId, const FString& Sdp)
+	void FStreamer::OnOffer(FPixelStreamingPlayerId PlayerId, const FString &Sdp)
 	{
 		TSharedPtr<IPlayerSession> NewSession = CreateSession(PlayerId, Protocol::EPlayerFlags::PSPFlag_SupportsDataChannel);
 		if (NewSession)
@@ -269,12 +265,10 @@ namespace UE::PixelStreaming
 			ForceKeyFrame();
 		}
 		else
-		{
 			UE_LOG(LogPixelStreaming, Error, TEXT("Failed to create player session, peer connection was nullptr."));
-		}
 	}
 
-	void FStreamer::SetLocalDescription(webrtc::PeerConnectionInterface* PeerConnection, FSetSessionDescriptionObserver* Observer, webrtc::SessionDescriptionInterface* SDP)
+	void FStreamer::SetLocalDescription(webrtc::PeerConnectionInterface *PeerConnection, FSetSessionDescriptionObserver *Observer, webrtc::SessionDescriptionInterface *SDP)
 	{
 		// Note from Luke about WebRTC: the sink of video capturer will be added as a direct result
 		// of `PeerConnection->SetLocalDescription()` call but video encoder will be created later on
@@ -296,7 +290,7 @@ namespace UE::PixelStreaming
 				webrtc::RTCError Err = Sender->SetParameters(ExistingParams);
 				if (!Err.ok())
 				{
-					const char* ErrMsg = Err.message();
+					const char *ErrMsg = Err.message();
 					FString ErrorStr(ErrMsg);
 					UE_LOG(LogPixelStreaming, Error, TEXT("Failed to set RTP Sender params: %s"), *ErrorStr);
 				}
@@ -312,33 +306,35 @@ namespace UE::PixelStreaming
 		//		PeerConnection.SetLocalDescription(Answer);
 		//		SignallingServerConnection.SendAnswer(Answer);
 
-		FSetSessionDescriptionObserver* SetLocalDescriptionObserver = FSetSessionDescriptionObserver::Create(
+		FSetSessionDescriptionObserver *SetLocalDescriptionObserver = FSetSessionDescriptionObserver::Create(
 			[this, Session]() // on success
 			{
 				SignallingServerConnection->SendAnswer(Session->GetPlayerId(), *Session->GetPeerConnection().local_description());
 				bStreamingStarted = true;
 			},
-			[this, Session](const FString& Error) // on failure
+			[this, Session](const FString &Error) // on failure
 			{
 				UE_LOG(LogPixelStreaming, Error, TEXT("Failed to set local description: %s"), *Error);
 				Session->DisconnectPlayer(Error);
 				PlayerSessions.DeletePlayerSession(Session->GetPlayerId());
 			});
 
-		auto OnCreateAnswerSuccess = [this, Session, SetLocalDescriptionObserver](webrtc::SessionDescriptionInterface* SDP) {
+		auto OnCreateAnswerSuccess = [this, Session, SetLocalDescriptionObserver](webrtc::SessionDescriptionInterface *SDP)
+		{
 			SetLocalDescription(&Session->GetPeerConnection(), SetLocalDescriptionObserver, SDP);
 		};
 
-		FCreateSessionDescriptionObserver* CreateAnswerObserver = FCreateSessionDescriptionObserver::Create(
+		FCreateSessionDescriptionObserver *CreateAnswerObserver = FCreateSessionDescriptionObserver::Create(
 			MoveTemp(OnCreateAnswerSuccess),
-			[this, Session](const FString& Error) // on failure
+			[this, Session](const FString &Error) // on failure
 			{
 				UE_LOG(LogPixelStreaming, Error, TEXT("Failed to create answer: %s"), *Error);
 				Session->DisconnectPlayer(Error);
 				PlayerSessions.DeletePlayerSession(Session->GetPlayerId());
 			});
 
-		auto OnSetRemoteDescriptionSuccess = [this, Session, CreateAnswerObserver]() {
+		auto OnSetRemoteDescriptionSuccess = [this, Session, CreateAnswerObserver]()
+		{
 			// Note: these offer to receive are superseded now we are use transceivers to setup our peer connection media
 			int offer_to_receive_video = webrtc::PeerConnectionInterface::RTCOfferAnswerOptions::kUndefined;
 			int offer_to_receive_audio = webrtc::PeerConnectionInterface::RTCOfferAnswerOptions::kUndefined;
@@ -351,41 +347,39 @@ namespace UE::PixelStreaming
 				offer_to_receive_audio,
 				voice_activity_detection,
 				ice_restart,
-				use_rtp_mux
-			};
+				use_rtp_mux};
 
 			AddStreams(Session, Protocol::EPlayerFlags::PSPFlag_SupportsDataChannel);
 			ModifyTransceivers(Session->GetPeerConnection().GetTransceivers(), Protocol::EPlayerFlags::PSPFlag_SupportsDataChannel);
 			Session->GetPeerConnection().CreateAnswer(CreateAnswerObserver, AnswerOption);
 		};
-
-		FSetSessionDescriptionObserver* SetRemoteDescriptionObserver = FSetSessionDescriptionObserver::Create(
+		FSetSessionDescriptionObserver *SetRemoteDescriptionObserver = FSetSessionDescriptionObserver::Create(
 			MoveTemp(OnSetRemoteDescriptionSuccess),
-			[this, Session](const FString& Error) // on failure
+			[this, Session](const FString &Error) // on failure
 			{
 				UE_LOG(LogPixelStreaming, Error, TEXT("Failed to set remote description: %s"), *Error);
 				Session->DisconnectPlayer(Error);
 				PlayerSessions.DeletePlayerSession(Session->GetPlayerId());
 			});
 
-		cricket::SessionDescription* RemoteDescription = Sdp->description();
+		cricket::SessionDescription *RemoteDescription = Sdp->description();
 		MungeRemoteSDP(RemoteDescription);
 
 		Session->GetPeerConnection().SetRemoteDescription(SetRemoteDescriptionObserver, Sdp.Release());
 	}
 
-	void FStreamer::MungeRemoteSDP(cricket::SessionDescription* RemoteDescription)
+	void FStreamer::MungeRemoteSDP(cricket::SessionDescription *RemoteDescription)
 	{
 		// Munge SDP of remote description to inject min, max, start bitrates
-		std::vector<cricket::ContentInfo>& ContentInfos = RemoteDescription->contents();
-		for (cricket::ContentInfo& Content : ContentInfos)
+		std::vector<cricket::ContentInfo> &ContentInfos = RemoteDescription->contents();
+		for (cricket::ContentInfo &Content : ContentInfos)
 		{
-			cricket::MediaContentDescription* MediaDescription = Content.media_description();
+			cricket::MediaContentDescription *MediaDescription = Content.media_description();
 			if (MediaDescription->type() == cricket::MediaType::MEDIA_TYPE_VIDEO)
 			{
-				cricket::VideoContentDescription* VideoDescription = MediaDescription->as_video();
+				cricket::VideoContentDescription *VideoDescription = MediaDescription->as_video();
 				std::vector<cricket::VideoCodec> CodecsCopy = VideoDescription->codecs();
-				for (cricket::VideoCodec& Codec : CodecsCopy)
+				for (cricket::VideoCodec &Codec : CodecsCopy)
 				{
 					// Note: These params are passed as kilobits, so divide by 1000.
 					Codec.SetParam(cricket::kCodecParamMinBitrate, Settings::CVarPixelStreamingWebRTCMinBitrate.GetValueOnAnyThread() / 1000);
@@ -397,46 +391,45 @@ namespace UE::PixelStreaming
 		}
 	}
 
-	void FStreamer::OnRemoteIceCandidate(FPixelStreamingPlayerId PlayerId, const FString& SdpMid, int SdpMLineIndex, const FString& Sdp)
+	void FStreamer::OnRemoteIceCandidate(FPixelStreamingPlayerId PlayerId, const FString &SdpMid, int SdpMLineIndex, const FString &Sdp)
 	{
-		PlayerSessions.ForSession(PlayerId, [&](auto Session) {
-			Session->OnRemoteIceCandidate(SdpMid, SdpMLineIndex, Sdp);
-		});
+		PlayerSessions.ForSession(PlayerId, [&](TSharedPtr<IPlayerSession> Session)
+								  { Session->OnRemoteIceCandidate(SdpMid, SdpMLineIndex, Sdp); });
 	}
 
 	void FStreamer::OnPlayerConnected(FPixelStreamingPlayerId PlayerId, int Flags)
 	{
 		// create peer connection
-		if (auto NewSession = CreateSession(PlayerId, Flags))
+		if (TSharedPtr<IPlayerSession> NewSession = CreateSession(PlayerId, Flags))
 		{
 			AddStreams(NewSession, Flags);
 
-			auto* PeerConnection = &NewSession->GetPeerConnection();
-			ModifyTransceivers(PeerConnection->GetTransceivers(), Flags);
+			webrtc::PeerConnectionInterface &PeerConnection = NewSession->GetPeerConnection();
+			ModifyTransceivers(PeerConnection.GetTransceivers(), Flags);
 
 			// observer for creating offer
-			FCreateSessionDescriptionObserver* CreateOfferObserver = FCreateSessionDescriptionObserver::Create(
-				[this, PlayerId, PeerConnection](webrtc::SessionDescriptionInterface* SDP) // on SDP create success
+			FCreateSessionDescriptionObserver *CreateOfferObserver = FCreateSessionDescriptionObserver::Create(
+				[this, PlayerId, PeerConnectionPtr = &PeerConnection](webrtc::SessionDescriptionInterface *SDP) // on SDP create success
 				{
-					FSetSessionDescriptionObserver* SetLocalDescriptionObserver = FSetSessionDescriptionObserver::Create(
+					FSetSessionDescriptionObserver *SetLocalDescriptionObserver = FSetSessionDescriptionObserver::Create(
 						[this, PlayerId, SDP]() // on SDP set success
 						{
 							SignallingServerConnection->SendOffer(PlayerId, *SDP);
 						},
-						[](const FString& Error) // on SDP set failure
+						[](const FString &Error) // on SDP set failure
 						{
 							UE_LOG(LogPixelStreaming, Error, TEXT("Failed to set local description: %s"), *Error);
 						});
 
 					MungeLocalSDP(SDP->description());
-					SetLocalDescription(PeerConnection, SetLocalDescriptionObserver, SDP);
+					SetLocalDescription(PeerConnectionPtr, SetLocalDescriptionObserver, SDP);
 				},
-				[](const FString& Error) // on SDP create failure
+				[](const FString &Error) // on SDP create failure
 				{
 					UE_LOG(LogPixelStreaming, Error, TEXT("Failed to create offer: %s"), *Error);
 				});
 
-			PeerConnection->CreateOffer(CreateOfferObserver, {});
+			PeerConnection.CreateOffer(CreateOfferObserver, {});
 		}
 	}
 
@@ -464,7 +457,6 @@ namespace UE::PixelStreaming
 		SignallingServerConnection->Disconnect();
 		ConnectToSignallingServer(PreviousSignallingServerURL); // maybe leave this up to the user?
 	}
-
 	void FStreamer::ForceKeyFrame()
 	{
 		if (P2PVideoEncoderFactory)
@@ -523,18 +515,19 @@ namespace UE::PixelStreaming
 			using FLayer = Settings::FSimulcastParameters::FLayer;
 
 			// encodings should be lowest res to highest
-			TArray<FLayer*> SortedLayers;
-			for (FLayer& Layer : Settings::SimulcastParameters.Layers)
+			TArray<FLayer *> SortedLayers;
+			for (FLayer &Layer : Settings::SimulcastParameters.Layers)
 			{
 				SortedLayers.Add(&Layer);
 			}
 
-			SortedLayers.Sort([](const FLayer& LayerA, const FLayer& LayerB) { return LayerA.Scaling > LayerB.Scaling; });
+			SortedLayers.Sort([](const FLayer &LayerA, const FLayer &LayerB)
+							  { return LayerA.Scaling > LayerB.Scaling; });
 
 			const int LayerCount = SortedLayers.Num();
 			for (int i = 0; i < LayerCount; ++i)
 			{
-				const FLayer* SimulcastLayer = SortedLayers[i];
+				const FLayer *SimulcastLayer = SortedLayers[i];
 				webrtc::RtpEncodingParameters LayerEncoding{};
 				LayerEncoding.rid = TCHAR_TO_UTF8(*(FString("simulcast") + FString::FromInt(LayerCount - i)));
 				LayerEncoding.min_bitrate_bps = SimulcastLayer->MinBitrate;
@@ -587,20 +580,20 @@ namespace UE::PixelStreaming
 		webrtc::DegradationPreference DegradationPref = Settings::GetDegradationPreference();
 		switch (DegradationPref)
 		{
-			case webrtc::DegradationPreference::MAINTAIN_FRAMERATE:
-				VideoTrack->set_content_hint(webrtc::VideoTrackInterface::ContentHint::kFluid);
-				break;
-			case webrtc::DegradationPreference::MAINTAIN_RESOLUTION:
-				VideoTrack->set_content_hint(webrtc::VideoTrackInterface::ContentHint::kDetailed);
-				break;
-			default:
-				break;
+		case webrtc::DegradationPreference::MAINTAIN_FRAMERATE:
+			VideoTrack->set_content_hint(webrtc::VideoTrackInterface::ContentHint::kFluid);
+			break;
+		case webrtc::DegradationPreference::MAINTAIN_RESOLUTION:
+			VideoTrack->set_content_hint(webrtc::VideoTrackInterface::ContentHint::kDetailed);
+			break;
+		default:
+			break;
 		}
 
 		bool bHasVideoTransceiver = false;
 
-		auto& PeerConnection = Session->GetPeerConnection();
-		for (auto& Transceiver : PeerConnection.GetTransceivers())
+		webrtc::PeerConnectionInterface &PeerConnection = Session->GetPeerConnection();
+		for (auto &Transceiver : PeerConnection.GetTransceivers())
 		{
 			rtc::scoped_refptr<webrtc::RtpSenderInterface> Sender = Transceiver->sender();
 			if (Transceiver->media_type() == cricket::MediaType::MEDIA_TYPE_VIDEO)
@@ -614,7 +607,7 @@ namespace UE::PixelStreaming
 		if (!bHasVideoTransceiver)
 		{
 			webrtc::RtpTransceiverInit TransceiverOptions;
-			TransceiverOptions.stream_ids = { TCHAR_TO_UTF8(*InVideoStreamId) };
+			TransceiverOptions.stream_ids = {TCHAR_TO_UTF8(*InVideoStreamId)};
 			TransceiverOptions.direction = webrtc::RtpTransceiverDirection::kSendOnly;
 			TransceiverOptions.send_encodings = CreateRTPEncodingParams(Flags);
 
@@ -623,25 +616,25 @@ namespace UE::PixelStreaming
 		}
 	}
 
-	void FStreamer::MungeLocalSDP(cricket::SessionDescription* SessionDescription)
+	void FStreamer::MungeLocalSDP(cricket::SessionDescription *SessionDescription)
 	{
 
-		std::vector<cricket::ContentInfo>& ContentInfos = SessionDescription->contents();
-		for (cricket::ContentInfo& ContentInfo : ContentInfos)
+		std::vector<cricket::ContentInfo> &ContentInfos = SessionDescription->contents();
+		for (cricket::ContentInfo &ContentInfo : ContentInfos)
 		{
-			cricket::MediaContentDescription* MediaDescription = ContentInfo.media_description();
+			cricket::MediaContentDescription *MediaDescription = ContentInfo.media_description();
 			cricket::MediaType MediaType = MediaDescription->type();
 			if (MediaType != cricket::MediaType::MEDIA_TYPE_AUDIO)
 			{
 				continue;
 			}
-			cricket::AudioContentDescription* AudioDescription = MediaDescription->as_audio();
+			cricket::AudioContentDescription *AudioDescription = MediaDescription->as_audio();
 			if (AudioDescription == nullptr)
 			{
 				continue;
 			}
 			std::vector<cricket::AudioCodec> CodecsCopy = AudioDescription->codecs();
-			for (cricket::AudioCodec& Codec : CodecsCopy)
+			for (cricket::AudioCodec &Codec : CodecsCopy)
 			{
 				if (Codec.name == "opus")
 				{
@@ -700,8 +693,8 @@ namespace UE::PixelStreaming
 		rtc::scoped_refptr<webrtc::AudioTrackInterface> AudioTrack = PCFactory->CreateAudioTrack(TCHAR_TO_UTF8(*InAudioTrackLabel), AudioSource);
 
 		bool bHasAudioTransceiver = false;
-		auto& PeerConnection = Session->GetPeerConnection();
-		for (auto& Transceiver : PeerConnection.GetTransceivers())
+		webrtc::PeerConnectionInterface &PeerConnection = Session->GetPeerConnection();
+		for (auto &Transceiver : PeerConnection.GetTransceivers())
 		{
 			rtc::scoped_refptr<webrtc::RtpSenderInterface> Sender = Transceiver->sender();
 			if (Transceiver->media_type() == cricket::MediaType::MEDIA_TYPE_AUDIO)
@@ -714,7 +707,7 @@ namespace UE::PixelStreaming
 		if (!bHasAudioTransceiver)
 		{
 			// Add the track
-			webrtc::RTCErrorOr<rtc::scoped_refptr<webrtc::RtpSenderInterface>> Result = PeerConnection.AddTrack(AudioTrack, { TCHAR_TO_UTF8(*InAudioStreamId) });
+			webrtc::RTCErrorOr<rtc::scoped_refptr<webrtc::RtpSenderInterface>> Result = PeerConnection.AddTrack(AudioTrack, {TCHAR_TO_UTF8(*InAudioStreamId)});
 
 			if (!Result.ok())
 			{
@@ -729,7 +722,7 @@ namespace UE::PixelStreaming
 		bool bIsSFU = (Flags & Protocol::EPlayerFlags::PSPFlag_IsSFU) != Protocol::EPlayerFlags::PSPFlag_None;
 		bool bReceiveBrowserAudio = !bIsSFU && !Settings::CVarPixelStreamingWebRTCDisableReceiveAudio.GetValueOnAnyThread();
 
-		for (auto& Transceiver : Transceivers)
+		for (auto &Transceiver : Transceivers)
 		{
 			rtc::scoped_refptr<webrtc::RtpSenderInterface> Sender = Transceiver->sender();
 
@@ -756,12 +749,12 @@ namespace UE::PixelStreaming
 				}
 
 				Transceiver->SetDirection(AudioTransceiverDirection);
-				Sender->SetStreams({ TCHAR_TO_UTF8(*GetAudioStreamID()) });
+				Sender->SetStreams({TCHAR_TO_UTF8(*GetAudioStreamID())});
 			}
 			else if (Transceiver->media_type() == cricket::MediaType::MEDIA_TYPE_VIDEO)
 			{
 				Transceiver->SetDirection(webrtc::RtpTransceiverDirection::kSendOnly);
-				Sender->SetStreams({ TCHAR_TO_UTF8(*GetVideoStreamID()) });
+				Sender->SetStreams({TCHAR_TO_UTF8(*GetVideoStreamID())});
 
 				webrtc::RtpParameters RtpParams = Sender->GetParameters();
 				RtpParams.encodings = CreateRTPEncodingParams(Flags);
@@ -770,23 +763,21 @@ namespace UE::PixelStreaming
 		}
 	}
 
-	void FStreamer::SendPlayerMessage(Protocol::EToPlayerMsg Type, const FString& Descriptor)
+	void FStreamer::SendPlayerMessage(Protocol::EToPlayerMsg Type, const FString &Descriptor)
 	{
-		PlayerSessions.ForEachSession([&](TSharedPtr<IPlayerSession> PlayerSession) {
-			PlayerSession->SendMessage(Type, Descriptor);
-		});
+		PlayerSessions.ForEachSession([&](TSharedPtr<IPlayerSession> PlayerSession)
+									  { PlayerSession->SendMessage(Type, Descriptor); });
 	}
 
-	void FStreamer::SendFreezeFrame(const TArray64<uint8>& JpegBytes)
+	void FStreamer::SendFreezeFrame(const TArray64<uint8> &JpegBytes)
 	{
 		if (!bStreamingStarted)
 		{
 			return;
 		}
 
-		PlayerSessions.ForEachSession([&JpegBytes](TSharedPtr<IPlayerSession> PlayerSession) {
-			PlayerSession->SendFreezeFrame(JpegBytes);
-		});
+		PlayerSessions.ForEachSession([&JpegBytes](TSharedPtr<IPlayerSession> PlayerSession)
+									  { PlayerSession->SendFreezeFrame(JpegBytes); });
 
 		CachedJpegBytes = JpegBytes;
 	}
@@ -795,9 +786,8 @@ namespace UE::PixelStreaming
 	{
 		if (CachedJpegBytes.Num() > 0)
 		{
-			PlayerSessions.ForSession(PlayerId, [&](auto Session) {
-				Session->SendFreezeFrame(CachedJpegBytes);
-			});
+			PlayerSessions.ForSession(PlayerId, [&](TSharedPtr<IPlayerSession> Session)
+									  { Session->SendFreezeFrame(CachedJpegBytes); });
 		}
 	}
 
@@ -806,18 +796,15 @@ namespace UE::PixelStreaming
 		// Force a keyframe so when stream unfreezes if player has never received a h.264 frame before they can still connect.
 		ForceKeyFrame();
 
-		PlayerSessions.ForEachSession([](TSharedPtr<IPlayerSession> PlayerSession) {
-			PlayerSession->SendUnfreezeFrame();
-		});
+		PlayerSessions.ForEachSession([](TSharedPtr<IPlayerSession> PlayerSession)
+									  { PlayerSession->SendUnfreezeFrame(); });
 
 		CachedJpegBytes.Empty();
 	}
-
-	void FStreamer::SendFileData(TArray<uint8>& ByteData, FString& MimeType, FString& FileExtension)
+	void FStreamer::SendFileData(TArray<uint8> &ByteData, FString &MimeType, FString &FileExtension)
 	{
-		PlayerSessions.ForEachSession([&](TSharedPtr<IPlayerSession> PlayerSession) {
-			PlayerSession->SendFileData(ByteData, MimeType, FileExtension);
-		});
+		PlayerSessions.ForEachSession([&](TSharedPtr<IPlayerSession> PlayerSession)
+									  { PlayerSession->SendFileData(ByteData, MimeType, FileExtension); });
 	}
 
 	void FStreamer::KickPlayer(FPixelStreamingPlayerId PlayerId)
diff --git a/Engine/Plugins/Media/PixelStreaming/Source/PixelStreaming/Private/Streamer.h b/Engine/Plugins/Media/PixelStreaming/Source/PixelStreaming/Private/Streamer.h
--- a/Engine/Plugins/Media/PixelStreaming/Source/PixelStreaming/Private/Streamer.h
+++ b/Engine/Plugins/Media/PixelStreaming/Source/PixelStreaming/Private/Streamer.h
@@ -3,7 +3,7 @@
 #pragma once
 
 #include "SignallingServerConnection.h"
-#include "ProtocolDefs.h"
+#include "PixelStreamingProtocolDefs.h"
 #include "RHI.h"
 #include "PlayerSessions.h"
 #include "FixedFPSPump.h"
diff --git a/Engine/Plugins/Media/PixelStreaming/Source/PixelStreaming/Private/TextureSource.cpp b/Engine/Plugins/Media/PixelStreaming/Source/PixelStreaming/Private/TextureSource.cpp
--- a/Engine/Plugins/Media/PixelStreaming/Source/PixelStreaming/Private/TextureSource.cpp
+++ b/Engine/Plugins/Media/PixelStreaming/Source/PixelStreaming/Private/TextureSource.cpp
@@ -1,69 +1,160 @@
 // Copyright Epic Games, Inc. All Rights Reserved.
 
-#include "TextureSource.h"
+#include "PixelStreamingTextureSource.h"
+#include "Stats.h"
 #include "Utils.h"
 
-
 /*
-* -------- UE::PixelStreaming::FTextureSourceBackBuffer ----------------
-*/
+ * --------- IBackBufferTextureSource ------------------------------------------
+ */
 
-UE::PixelStreaming::FTextureSourceBackBuffer::FTextureSourceBackBuffer() 
-	: TTextureSourceBackBufferBase()
+IBackBufferTextureSource::IBackBufferTextureSource(float InFrameScale)
+	: FrameScale(InFrameScale)
 {
+	// Explictly make clear we are adding another ref to this shared bool for the purposes of using in the below lambda
+	TSharedRef<bool, ESPMode::ThreadSafe> bEnabledClone = bEnabled;
+
+	// The backbuffer delegate can only be accessed using GameThread
+	AsyncTask(ENamedThreads::GameThread, [this, bEnabledClone]() {
+		/*Early exit if `this` died before game thread ran.*/
+		if (bEnabledClone.IsUnique())
+		{
+			return;
+		}
+		OnBackbuffer = &FSlateApplication::Get().GetRenderer()->OnBackBufferReadyToPresent();
+		BackbufferDelegateHandle = OnBackbuffer->AddSP(this, &IBackBufferTextureSource::OnBackBufferReady_RenderThread);
+	});
+}
 
-};
-
-UE::PixelStreaming::FTextureSourceBackBuffer::FTextureSourceBackBuffer(float InScale) 
-	: TTextureSourceBackBufferBase(InScale)
+IBackBufferTextureSource::~IBackBufferTextureSource()
 {
+	if (OnBackbuffer)
+	{
+		OnBackbuffer->Remove(BackbufferDelegateHandle);
+	}
 
-};
+	*bEnabled = false;
+}
 
-void UE::PixelStreaming::FTextureSourceBackBuffer::CopyTexture(const FTexture2DRHIRef& SourceTexture, TRefCountPtr<FRHITexture2D> DestTexture, FGPUFenceRHIRef& CopyFence)
+TRefCountPtr<IPixelStreamingBackBufferTextureWrapper> IBackBufferTextureSource::GetCurrent()
 {
-	UE::PixelStreaming::CopyTexture(SourceTexture, DestTexture, CopyFence);
+	if (bIsTempDirty)
+	{
+		FScopeLock Lock(&CriticalSection);
+		ReadBuffer.Swap(TempBuffer);
+		bIsTempDirty = false;
+	}
+	return ReadBuffer;
 }
 
-TRefCountPtr<FRHITexture2D> UE::PixelStreaming::FTextureSourceBackBuffer::CreateTexture(int Width, int Height)
+void IBackBufferTextureSource::SetEnabled(bool bInEnabled)
 {
-	return UE::PixelStreaming::CreateTexture(Width, Height);
+	*bEnabled = bInEnabled;
+	// This source has been disabled, so set `bInitialized` to false so `OnBackBufferReady_RenderThread`
+	// will make new textures next time it is called.
+	if (bInitialized && bInEnabled == false)
+	{
+		bInitialized = false;
+	}
 }
 
-FTexture2DRHIRef UE::PixelStreaming::FTextureSourceBackBuffer::ToTextureRef(TRefCountPtr<FRHITexture2D> Texture)
+void IBackBufferTextureSource::OnBackBufferReady_RenderThread(SWindow& SlateWindow, const FTexture2DRHIRef& FrameBuffer)
 {
-	return Texture;
+	if (!bInitialized)
+	{
+		Initialize(FrameBuffer->GetSizeXY().X * FrameScale, FrameBuffer->GetSizeXY().Y * FrameScale);
+	}
+
+	if (!IsEnabled())
+	{
+		return;
+	}
+
+	auto& WriteBuffer = bWriteParity ? WriteBuffers[0] : WriteBuffers[1];
+	bWriteParity = !bWriteParity;
+
+	// for safety we just make sure that the buffer is not currently waiting for a copy
+	if (WriteBuffer.bAvailable)
+	{
+		WriteBuffer.bAvailable = false;
+
+		FRHICommandListImmediate& RHICmdList = FRHICommandListExecutor::GetImmediateCommandList();
+		WriteBuffer.Fence->Clear();
+
+		RHICmdList.EnqueueLambda([&WriteBuffer](FRHICommandListImmediate& RHICmdList) {
+			WriteBuffer.PreWaitingOnCopy = FPlatformTime::Cycles64();
+		});
+
+		UE::PixelStreaming::CopyTexture(FrameBuffer, WriteBuffer.TextureWrapper->GetTexture(), WriteBuffer.Fence);
+
+		IPixelStreamingModule::Get().AddGPUFencePollerTask(WriteBuffer.Fence, bEnabled, [this, &WriteBuffer]() {
+			// This lambda is called only once the GPUFence is done
+			{
+				FScopeLock Lock(&CriticalSection);
+				TempBuffer.Swap(WriteBuffer.TextureWrapper);
+				WriteBuffer.Fence->Clear();
+				WriteBuffer.bAvailable = true;
+
+				bIsTempDirty = true;
+			}
+
+			// For debugging timing information about the copy operation
+			// Turning it on all the time is a bit too much log spam if logging stats
+			uint64 PostWaitingOnCopy = FPlatformTime::Cycles64();
+			UE::PixelStreaming::FStats* Stats = UE::PixelStreaming::FStats::Get();
+			if (Stats)
+			{
+				double CaptureLatencyMs = FPlatformTime::ToMilliseconds64(PostWaitingOnCopy - WriteBuffer.PreWaitingOnCopy);
+				Stats->StoreApplicationStat(UE::PixelStreaming::FStatData(FName(*FString::Printf(TEXT("Layer (x%.2f) Capture time (ms)"), FrameScale)), CaptureLatencyMs, 2, true));
+			}
+		});
+	}
 }
 
-/*
-* -------- UE::PixelStreaming::FTextureSourceBackBufferToCPU ----------------
-*/
-
-UE::PixelStreaming::FTextureSourceBackBufferToCPU::FTextureSourceBackBufferToCPU() 
-	: TTextureSourceBackBufferBase()
+void IBackBufferTextureSource::Initialize(int Width, int Height)
 {
+	SourceWidth = Width;
+	SourceHeight = Height;
+
+	for (auto& Buffer : WriteBuffers)
+	{
+		Buffer.TextureWrapper = CreateTexture(SourceWidth, SourceHeight);
+		Buffer.Fence = GDynamicRHI->RHICreateGPUFence(TEXT("VideoCapturerCopyFence"));
+		Buffer.bAvailable = true;
+	}
+	bWriteParity = true;
+
+	TempBuffer = CreateTexture(SourceWidth, SourceHeight);
+	ReadBuffer = CreateTexture(SourceWidth, SourceHeight);
+	bIsTempDirty = false;
+
+	bInitialized = true;
+}
 
-};
-
-UE::PixelStreaming::FTextureSourceBackBufferToCPU::FTextureSourceBackBufferToCPU(float InScale) 
-	: TTextureSourceBackBufferBase(InScale)
+/*
+ * -------- FTextureSourceBackBuffer ----------------
+ */
+void FBackBufferTextureSource::CopyTexture(const FTexture2DRHIRef& SourceTexture, TRefCountPtr<IPixelStreamingBackBufferTextureWrapper> DestTexture, FGPUFenceRHIRef& CopyFence)
 {
+	UE::PixelStreaming::CopyTexture(SourceTexture, DestTexture->GetTexture(), CopyFence);
+}
 
-};
-
-void UE::PixelStreaming::FTextureSourceBackBufferToCPU::CopyTexture(const FTexture2DRHIRef& SourceTexture, TRefCountPtr<UE::PixelStreaming::FRawPixelsTexture> DestTexture, FGPUFenceRHIRef& CopyFence)
+TRefCountPtr<IPixelStreamingBackBufferTextureWrapper> FBackBufferTextureSource::CreateTexture(int Width, int Height)
 {
-	UE::PixelStreaming::CopyTexture(SourceTexture, DestTexture->TextureRef, CopyFence);
-	UE::PixelStreaming::ReadTextureToCPU(FRHICommandListExecutor::GetImmediateCommandList(), DestTexture->TextureRef, DestTexture->RawPixels);
+	return new FPixelStreamingRHIBackBufferTexture(UE::PixelStreaming::CreateTexture(Width, Height));
 }
 
-TRefCountPtr<UE::PixelStreaming::FRawPixelsTexture> UE::PixelStreaming::FTextureSourceBackBufferToCPU::CreateTexture(int Width, int Height)
+/*
+ * -------- FBackBufferToCPUTextureSource ----------------
+ */
+
+void FBackBufferToCPUTextureSource::CopyTexture(const FTexture2DRHIRef& SourceTexture, TRefCountPtr<IPixelStreamingBackBufferTextureWrapper> DestTexture, FGPUFenceRHIRef& CopyFence)
 {
-	FRawPixelsTexture* Tex = new FRawPixelsTexture(UE::PixelStreaming::CreateTexture(Width, Height));
-	return TRefCountPtr<FRawPixelsTexture>(Tex);
+	UE::PixelStreaming::CopyTexture(SourceTexture, DestTexture->GetTexture(), CopyFence);
+	UE::PixelStreaming::ReadTextureToCPU(FRHICommandListExecutor::GetImmediateCommandList(), DestTexture->GetTexture(), static_cast<FPixelStreamingCPUReadableBackbufferTexture*>(DestTexture.GetReference())->GetRawPixels());
 }
 
-FTexture2DRHIRef UE::PixelStreaming::FTextureSourceBackBufferToCPU::ToTextureRef(TRefCountPtr<UE::PixelStreaming::FRawPixelsTexture> Texture)
+TRefCountPtr<IPixelStreamingBackBufferTextureWrapper> FBackBufferToCPUTextureSource::CreateTexture(int Width, int Height)
 {
-	return Texture->TextureRef;
-}
\ No newline at end of file
+	return new FPixelStreamingCPUReadableBackbufferTexture(UE::PixelStreaming::CreateTexture(Width, Height));
+}
diff --git a/Engine/Plugins/Media/PixelStreaming/Source/PixelStreaming/Private/TextureSource.h b/Engine/Plugins/Media/PixelStreaming/Source/PixelStreaming/Private/TextureSource.h
deleted file mode 100644
diff --git a/Engine/Plugins/Media/PixelStreaming/Source/PixelStreaming/Private/VideoEncoderH264Wrapper.cpp b/Engine/Plugins/Media/PixelStreaming/Source/PixelStreaming/Private/VideoEncoderH264Wrapper.cpp
--- a/Engine/Plugins/Media/PixelStreaming/Source/PixelStreaming/Private/VideoEncoderH264Wrapper.cpp
+++ b/Engine/Plugins/Media/PixelStreaming/Source/PixelStreaming/Private/VideoEncoderH264Wrapper.cpp
@@ -5,15 +5,15 @@
 #include "Stats.h"
 #include "VideoEncoderFactory.h"
 #include "EncoderFactory.h"
+#include "PixelStreamingFrameBuffer.h"
 
 UE::PixelStreaming::FVideoEncoderH264Wrapper::FVideoEncoderH264Wrapper(TUniquePtr<FEncoderFrameFactory> InFrameFactory, TUniquePtr<AVEncoder::FVideoEncoder> InEncoder)
 	: FrameFactory(MoveTemp(InFrameFactory))
-	, Encoder(MoveTemp(InEncoder)) 
-{
-	checkf(Encoder, TEXT("Encoder is nullptr."))
-}
+	, Encoder(MoveTemp(InEncoder)){
+		checkf(Encoder, TEXT("Encoder is nullptr."))
+	}
 
-UE::PixelStreaming::FVideoEncoderH264Wrapper::~FVideoEncoderH264Wrapper()
+	UE::PixelStreaming::FVideoEncoderH264Wrapper::~FVideoEncoderH264Wrapper()
 {
 	Encoder->ClearOnEncodedPacket();
 	Encoder->Shutdown();
@@ -23,9 +23,9 @@ void UE::PixelStreaming::FVideoEncoderH264Wrapper::Encode(const webrtc::VideoFra
 {
 	FTexture2DRHIRef SourceTexture;
 
-	UE::PixelStreaming::FFrameBuffer* FrameBuffer = static_cast<UE::PixelStreaming::FFrameBuffer*>(WebRTCFrame.video_frame_buffer().get());
-	check(FrameBuffer->GetFrameBufferType() == UE::PixelStreaming::FFrameBufferType::Layer);
-	UE::PixelStreaming::FLayerFrameBuffer* LayerFrameBuffer = static_cast<UE::PixelStreaming::FLayerFrameBuffer*>(FrameBuffer);
+	FPixelStreamingFrameBuffer* FrameBuffer = static_cast<FPixelStreamingFrameBuffer*>(WebRTCFrame.video_frame_buffer().get());
+	check(FrameBuffer->GetFrameBufferType() == UE::PixelStreaming::EFrameBufferType::Layer);
+	FLayerFrameBuffer* LayerFrameBuffer = static_cast<FLayerFrameBuffer*>(FrameBuffer);
 	SourceTexture = LayerFrameBuffer->GetFrame();
 
 	if (SourceTexture)
diff --git a/Engine/Plugins/Media/PixelStreaming/Source/PixelStreaming/Private/VideoEncoderRTC.cpp b/Engine/Plugins/Media/PixelStreaming/Source/PixelStreaming/Private/VideoEncoderRTC.cpp
--- a/Engine/Plugins/Media/PixelStreaming/Source/PixelStreaming/Private/VideoEncoderRTC.cpp
+++ b/Engine/Plugins/Media/PixelStreaming/Source/PixelStreaming/Private/VideoEncoderRTC.cpp
@@ -3,7 +3,7 @@
 #include "VideoEncoderRTC.h"
 #include "EncoderFactory.h"
 #include "VideoEncoderFactory.h"
-#include "FrameBuffer.h"
+#include "PixelStreamingFrameBuffer.h"
 #include "PlayerSession.h"
 #include "Stats.h"
 #include "UnrealEngine.h"
@@ -29,8 +29,8 @@ int UE::PixelStreaming::FVideoEncoderRTC::InitEncode(webrtc::VideoCodec const* I
 	checkf(AVEncoder::FVideoEncoderFactory::Get().IsSetup(), TEXT("FVideoEncoderFactory not setup"));
 	// Try and get the encoder. If it doesn't exist, create it.
 	FVideoEncoderH264Wrapper* Encoder = Factory.GetOrCreateHardwareEncoder(InCodecSettings->width, InCodecSettings->height, InCodecSettings->maxBitrate, InCodecSettings->startBitrate, InCodecSettings->maxFramerate);
-	
-	if(Encoder != nullptr)
+
+	if (Encoder != nullptr)
 	{
 		HardwareEncoder = Encoder;
 		UpdateConfig();
@@ -57,11 +57,11 @@ int32 UE::PixelStreaming::FVideoEncoderRTC::Release()
 
 int32 UE::PixelStreaming::FVideoEncoderRTC::Encode(webrtc::VideoFrame const& frame, std::vector<webrtc::VideoFrameType> const* frame_types)
 {
-	UE::PixelStreaming::FFrameBuffer* FrameBuffer = static_cast<UE::PixelStreaming::FFrameBuffer*>(frame.video_frame_buffer().get());
+	FPixelStreamingFrameBuffer* FrameBuffer = static_cast<FPixelStreamingFrameBuffer*>(frame.video_frame_buffer().get());
 
 	// If initialize frame is passed, this is a dummy frame so WebRTC is happy frames are passing through
 	// This is used so that an encoder can be active as far as WebRTC is concerned by simply have frames transmitted by some other encoder on its behalf.
-	if(FrameBuffer->GetFrameBufferType() == UE::PixelStreaming::FFrameBufferType::Initialize)
+	if (FrameBuffer->GetFrameBufferType() == UE::PixelStreaming::EFrameBufferType::Initialize)
 	{
 		return WEBRTC_VIDEO_CODEC_OK;
 	}
diff --git a/Engine/Plugins/Media/PixelStreaming/Source/PixelStreaming/Private/VideoSource.cpp b/Engine/Plugins/Media/PixelStreaming/Source/PixelStreaming/Private/VideoSource.cpp
--- a/Engine/Plugins/Media/PixelStreaming/Source/PixelStreaming/Private/VideoSource.cpp
+++ b/Engine/Plugins/Media/PixelStreaming/Source/PixelStreaming/Private/VideoSource.cpp
@@ -2,10 +2,11 @@
 
 #include "VideoSource.h"
 #include "FixedFPSPump.h"
-#include "FrameBuffer.h"
+#include "PixelStreamingFrameBuffer.h"
 #include "Stats.h"
 #include "PixelStreamingStatNames.h"
 #include "PlayerSessions.h"
+#include "IPixelStreamingTextureSource.h"
 
 namespace UE::PixelStreaming
 {
@@ -84,11 +85,11 @@ namespace UE::PixelStreaming
 	{
 		if (Settings::IsCodecVPX())
 		{
-			TextureSource = MakeShared<FTextureSourceBackBufferToCPU>((Settings::CVarPixelStreamingFrameScale.GetValueOnAnyThread() <= 0 ? 1.0 : Settings::CVarPixelStreamingFrameScale.GetValueOnAnyThread()));
+			TextureSource = TSharedPtr<IPixelStreamingTextureSource, ESPMode::ThreadSafe>(new FBackBufferToCPUTextureSource(Settings::CVarPixelStreamingFrameScale.GetValueOnAnyThread() <= 0 ? 1.0 : Settings::CVarPixelStreamingFrameScale.GetValueOnAnyThread()));
 		}
 		else
 		{
-			TextureSource = MakeShared<FTextureSourceBackBuffer>((Settings::CVarPixelStreamingFrameScale.GetValueOnAnyThread() <= 0 ? 1.0 : Settings::CVarPixelStreamingFrameScale.GetValueOnAnyThread()));
+			TextureSource = TSharedPtr<IPixelStreamingTextureSource, ESPMode::ThreadSafe>(new FBackBufferTextureSource(Settings::CVarPixelStreamingFrameScale.GetValueOnAnyThread() <= 0 ? 1.0 : Settings::CVarPixelStreamingFrameScale.GetValueOnAnyThread()));
 		}
 
 		// We store the codec during construction as querying it everytime seem overly wasteful
@@ -185,7 +186,7 @@ namespace UE::PixelStreaming
 		for (FLayer* SimulcastLayer : SortedLayers)
 		{
 			const float Scale = 1.0f / SimulcastLayer->Scaling;
-			TSharedPtr<FTextureSourceBackBuffer> TextureSource = MakeShared<FTextureSourceBackBuffer>(Scale);
+			TSharedPtr<FBackBufferTextureSource> TextureSource = MakeShared<FBackBufferTextureSource>(Scale);
 			TextureSource->SetEnabled(true);
 			LayerTextures.Add(TextureSource);
 		}
diff --git a/Engine/Plugins/Media/PixelStreaming/Source/PixelStreaming/Private/VideoSource.h b/Engine/Plugins/Media/PixelStreaming/Source/PixelStreaming/Private/VideoSource.h
--- a/Engine/Plugins/Media/PixelStreaming/Source/PixelStreaming/Private/VideoSource.h
+++ b/Engine/Plugins/Media/PixelStreaming/Source/PixelStreaming/Private/VideoSource.h
@@ -3,13 +3,15 @@
 #pragma once
 
 #include "FixedFPSPump.h"
+#include "PixelStreamingTextureSource.h"
 #include "Settings.h"
 
+class IPixelStreamingTextureSource;
+
 namespace UE::PixelStreaming
 {
 	class FVideoSourceFactory;
 	class FPlayerSessions;
-	class ITextureSource;
 
 	/*
 		* Base class for pumped Pixel Streaming video sources. Video sources are WebRTC video sources that populate WebRTC tracks
@@ -38,6 +40,8 @@ namespace UE::PixelStreaming
 		virtual bool HasOneRef() const { return RefCount.HasOneRef(); }
 		/* End IPumpedVideoSource */
 
+		/* Begin rtc::RefCountInterface */
+
 	protected:
 		virtual bool AdaptCaptureFrame(const int64 TimestampUs, FIntPoint Resolution);
 		virtual webrtc::VideoFrame CreateFrame(int32 FrameId) = 0;
@@ -60,7 +64,7 @@ namespace UE::PixelStreaming
 
 	protected:
 		FPlayerSessions* Sessions;
-		TSharedPtr<ITextureSource> TextureSource;
+		TSharedPtr<IPixelStreamingTextureSource> TextureSource;
 		Settings::ECodec Codec;
 
 	protected:
@@ -82,7 +86,7 @@ namespace UE::PixelStreaming
 		FVideoSourceSFU();
 
 	protected:
-		TArray<TSharedPtr<ITextureSource>> LayerTextures;
+		TArray<TSharedPtr<IPixelStreamingTextureSource>> LayerTextures;
 
 	protected:
 		/* Begin FVideoSourceBase */
diff --git a/Engine/Plugins/Media/PixelStreaming/Source/PixelStreaming/Public/IPixelStreamingModule.h b/Engine/Plugins/Media/PixelStreaming/Source/PixelStreaming/Public/IPixelStreamingModule.h
--- a/Engine/Plugins/Media/PixelStreaming/Source/PixelStreaming/Public/IPixelStreamingModule.h
+++ b/Engine/Plugins/Media/PixelStreaming/Source/PixelStreaming/Public/IPixelStreamingModule.h
@@ -9,6 +9,7 @@
 #include "Templates/SharedPointer.h"
 #include "PixelStreamingPlayerId.h"
 #include "IPixelStreamingAudioSink.h"
+#include "IPixelStreamingPumpedVideoSource.h"
 #include "Templates/SharedPointer.h"
 
 #include "IInputDevice.h"
@@ -16,6 +17,11 @@
 class UTexture2D;
 class UPixelStreamingInput;
 
+namespace webrtc
+{
+	class VideoEncoderFactory;
+}
+
 /**
 * The public interface to this module
 */
@@ -172,4 +178,24 @@ public:
 	 * Get the input components currently attached to Pixel Streaming.
 	 */
 	virtual const TArray<UPixelStreamingInput*> GetInputComponents() = 0;
+
+	/*
+	 * Create a webrtc::VideoEncoderFactory pointer.
+	 */
+	virtual webrtc::VideoEncoderFactory* CreateVideoEncoderFactory() = 0;
+
+	/*
+	 * Register a VideoSource with the module
+	 */
+	virtual void RegisterVideoSource(FPixelStreamingPlayerId PlayerId, IPumpedVideoSource* VideoSource) = 0;
+
+	/*
+	 * Unregister a VideoSource with the module
+	 */
+	virtual void UnregisterVideoSource(FPixelStreamingPlayerId PlayerId) = 0;
+
+	/*
+	 * Create add a task to the GPU fence poller.
+	 */
+	virtual void AddGPUFencePollerTask(FGPUFenceRHIRef Fence, TSharedRef<bool, ESPMode::ThreadSafe> bIsEnabled, TFunction<void()> Task) = 0;
 };
diff --git a/Engine/Plugins/Media/PixelStreaming/Source/PixelStreaming/Public/IPixelStreamingPumpedVideoSource.h b/Engine/Plugins/Media/PixelStreaming/Source/PixelStreaming/Public/IPixelStreamingPumpedVideoSource.h
new file mode 100644
--- /dev/null
+++ b/Engine/Plugins/Media/PixelStreaming/Source/PixelStreaming/Public/IPixelStreamingPumpedVideoSource.h
@@ -0,0 +1,12 @@
+// Copyright Epic Games, Inc. All Rights Reserved.
+#pragma once
+
+/*
+* Interface for video sources that get pumped through `OnPump`.
+*/
+class PIXELSTREAMING_API IPumpedVideoSource
+{
+public:
+	virtual void OnPump(int32 FrameId) = 0;
+	virtual bool IsReadyForPump() const = 0;
+};
\ No newline at end of file
diff --git a/Engine/Plugins/Media/PixelStreaming/Source/PixelStreaming/Public/IPixelStreamingSessions.h b/Engine/Plugins/Media/PixelStreaming/Source/PixelStreaming/Public/IPixelStreamingSessions.h
new file mode 100644
--- /dev/null
+++ b/Engine/Plugins/Media/PixelStreaming/Source/PixelStreaming/Public/IPixelStreamingSessions.h
@@ -0,0 +1,25 @@
+// Copyright Epic Games, Inc. All Rights Reserved.
+
+#pragma once
+
+#include "IPixelStreamingAudioSink.h"
+#include "PixelStreamingPlayerId.h"
+#include "ProtocolDefs.h"
+
+namespace UE::PixelStreaming
+{
+	class IPixelStreamingSessions
+	{
+	public:
+		virtual int GetNumPlayers() const = 0;
+		virtual IPixelStreamingAudioSink* GetAudioSink(FPixelStreamingPlayerId PlayerId) const = 0;
+		virtual IPixelStreamingAudioSink* GetUnlistenedAudioSink() const = 0;
+		virtual bool IsQualityController(FPixelStreamingPlayerId PlayerId) const = 0;
+		virtual void SetQualityController(FPixelStreamingPlayerId PlayerId) = 0;
+		virtual bool SendMessage(FPixelStreamingPlayerId PlayerId, Protocol::EToPlayerMsg Type, const FString& Descriptor) const = 0;
+		virtual void SendLatestQP(FPixelStreamingPlayerId PlayerId, int LatestQP) const = 0;
+		virtual void SendFreezeFrameTo(FPixelStreamingPlayerId PlayerId, const TArray64<uint8>& JpegBytes) const = 0;
+		virtual void PollWebRTCStats() const = 0;
+		virtual void OnDataChannelClosed(FPixelStreamingPlayerId PlayerId) const = 0;
+	};
+} // namespace UE::PixelStreaming
diff --git a/Engine/Plugins/Media/PixelStreaming/Source/PixelStreaming/Public/IPixelStreamingTextureSource.h b/Engine/Plugins/Media/PixelStreaming/Source/PixelStreaming/Public/IPixelStreamingTextureSource.h
new file mode 100644
--- /dev/null
+++ b/Engine/Plugins/Media/PixelStreaming/Source/PixelStreaming/Public/IPixelStreamingTextureSource.h
@@ -0,0 +1,23 @@
+// Copyright Epic Games, Inc. All Rights Reserved.
+#pragma once
+
+#include "Containers/UnrealString.h"
+#include "RHI.h"
+
+/*
+* Interface for all texture sources in Pixel Streaming.
+* These texture sources are used to populate video sources for video tracks.
+*/
+class IPixelStreamingTextureSource
+{
+public:
+	IPixelStreamingTextureSource() = default;
+	virtual ~IPixelStreamingTextureSource() = default;
+	virtual bool IsAvailable() const = 0;
+	virtual bool IsEnabled() const = 0;
+	virtual void SetEnabled(bool bInEnabled) = 0;
+	virtual int GetSourceHeight() const = 0;
+	virtual int GetSourceWidth() const = 0;
+	virtual FTexture2DRHIRef GetTexture() = 0;
+	virtual FString GetName() const = 0;
+};
\ No newline at end of file
diff --git a/Engine/Plugins/Media/PixelStreaming/Source/PixelStreaming/Public/PixelStreamingFrameBuffer.h b/Engine/Plugins/Media/PixelStreaming/Source/PixelStreaming/Public/PixelStreamingFrameBuffer.h
new file mode 100644
--- /dev/null
+++ b/Engine/Plugins/Media/PixelStreaming/Source/PixelStreaming/Public/PixelStreamingFrameBuffer.h
@@ -0,0 +1,141 @@
+// Copyright Epic Games, Inc. All Rights Reserved.
+
+#pragma once
+
+#include "api/video/video_frame_buffer.h"
+#include "api/video/i420_buffer.h"
+#include "api/scoped_refptr.h"
+
+class FTexture2DRHI;
+class IPixelStreamingTextureSource;
+
+namespace UE
+{
+	namespace PixelStreaming
+	{
+		enum EFrameBufferType
+		{
+			Initialize,
+			Simulcast,
+			Layer
+		};
+	} // namespace PixelStreaming
+} // namespace UE
+
+/*
+ * ----------------- FPixelStreamingFrameBuffer -----------------
+ * The base framebuffer that extends the WebRTC type.
+ */
+class PIXELSTREAMING_API FPixelStreamingFrameBuffer : public webrtc::VideoFrameBuffer
+{
+public:
+	virtual ~FPixelStreamingFrameBuffer() {}
+
+	virtual UE::PixelStreaming::EFrameBufferType GetFrameBufferType() const = 0;
+
+	// Begin webrtc::VideoFrameBuffer interface
+	virtual webrtc::VideoFrameBuffer::Type type() const override
+	{
+		return webrtc::VideoFrameBuffer::Type::kNative;
+	}
+
+	virtual int width() const = 0;
+	virtual int height() const = 0;
+
+	virtual rtc::scoped_refptr<webrtc::I420BufferInterface> ToI420() override
+	{
+		unimplemented();
+		return nullptr;
+	}
+
+	virtual const webrtc::I420BufferInterface* GetI420() const override
+	{
+		unimplemented();
+		return nullptr;
+	}
+	// End webrtc::VideoFrameBuffer interface
+};
+
+/*
+ * ----------------- FInitializeFrameBuffer -----------------
+ * We use this frame to force webrtc to create encoders that will siphon off other streams
+ * but never get real frames directly pumped to them.
+ */
+class PIXELSTREAMING_API FInitializeFrameBuffer : public FPixelStreamingFrameBuffer
+{
+public:
+	FInitializeFrameBuffer(TSharedPtr<IPixelStreamingTextureSource> InTextureSource);
+	virtual ~FInitializeFrameBuffer();
+
+	virtual UE::PixelStreaming::EFrameBufferType GetFrameBufferType() const { return UE::PixelStreaming::EFrameBufferType::Initialize; }
+
+	virtual int width() const override;
+	virtual int height() const override;
+
+private:
+	TSharedPtr<IPixelStreamingTextureSource> TextureSource;
+};
+
+/*
+ * ----------------- FSimulcastFrameBuffer -----------------
+ * Holds a number of textures (called layers) and each is assigned an index.
+ */
+class PIXELSTREAMING_API FSimulcastFrameBuffer : public FPixelStreamingFrameBuffer
+{
+public:
+	FSimulcastFrameBuffer(TArray<TSharedPtr<IPixelStreamingTextureSource>>& InTextureSources);
+	virtual ~FSimulcastFrameBuffer();
+
+	virtual UE::PixelStreaming::EFrameBufferType GetFrameBufferType() const { return UE::PixelStreaming::EFrameBufferType::Simulcast; }
+	TSharedPtr<IPixelStreamingTextureSource> GetLayerFrameSource(int LayerIndex) const;
+	int GetNumLayers() const;
+
+	virtual int width() const override;
+	virtual int height() const override;
+
+private:
+	TArray<TSharedPtr<IPixelStreamingTextureSource>>& TextureSources;
+};
+
+/*
+ * ----------------- FLayerFrameBuffer -----------------
+ * Holds a genuine single texture for encoding.
+ */
+class PIXELSTREAMING_API FLayerFrameBuffer : public FPixelStreamingFrameBuffer
+{
+public:
+	FLayerFrameBuffer(TSharedPtr<IPixelStreamingTextureSource> InTextureSource);
+	virtual ~FLayerFrameBuffer();
+
+	virtual UE::PixelStreaming::EFrameBufferType GetFrameBufferType() const { return UE::PixelStreaming::EFrameBufferType::Layer; }
+	FTexture2DRHIRef GetFrame() const;
+
+	virtual int width() const override;
+	virtual int height() const override;
+
+private:
+	TSharedPtr<IPixelStreamingTextureSource> TextureSource;
+};
+
+/*
+ * ----------------- FFrameBufferI420 -----------------
+ * Holds a texture source that is capable of being converted to I420 for encoding.
+ */
+class PIXELSTREAMING_API FFrameBufferI420 : public FPixelStreamingFrameBuffer
+{
+public:
+	FFrameBufferI420(TSharedPtr<IPixelStreamingTextureSource> InTextureSource);
+	virtual ~FFrameBufferI420();
+
+	virtual UE::PixelStreaming::EFrameBufferType GetFrameBufferType() const { return UE::PixelStreaming::EFrameBufferType::Layer; }
+	FTexture2DRHIRef GetFrame() const;
+
+	virtual int width() const override;
+	virtual int height() const override;
+	virtual rtc::scoped_refptr<webrtc::I420BufferInterface> ToI420() override;
+	virtual const webrtc::I420BufferInterface* GetI420() const override;
+
+private:
+	TSharedPtr<IPixelStreamingTextureSource> TextureSource;
+	rtc::scoped_refptr<webrtc::I420Buffer> Buffer;
+};
diff --git a/Engine/Plugins/Media/PixelStreaming/Source/PixelStreaming/Public/PixelStreamingProtocolDefs.h b/Engine/Plugins/Media/PixelStreaming/Source/PixelStreaming/Public/PixelStreamingProtocolDefs.h
new file mode 100644
--- /dev/null
+++ b/Engine/Plugins/Media/PixelStreaming/Source/PixelStreaming/Public/PixelStreamingProtocolDefs.h
@@ -0,0 +1,119 @@
+// Copyright Epic Games, Inc. All Rights Reserved.
+
+#pragma once
+
+#include "HAL/Platform.h"
+#include "api/data_channel_interface.h"
+#include "WebRTCIncludes.h"
+
+namespace UE::PixelStreaming::Protocol
+{
+	enum EPlayerFlags
+	{
+		PSPFlag_None = 0x00,
+		PSPFlag_SupportsDataChannel = 0x01,
+		PSPFlag_IsSFU = 0x02,
+	};
+
+	enum class EToStreamerMsg : uint8
+	{
+		/**********************************************************************/
+
+		/*
+				* Control Messages. Range = 0..49.
+				*/
+
+		IFrameRequest = 0,
+		RequestQualityControl = 1,
+		FpsRequest = 2,
+		AverageBitrateRequest = 3,
+		StartStreaming = 4,
+		StopStreaming = 5,
+		LatencyTest = 6,
+		RequestInitialSettings = 7,
+		TestEcho = 8,
+
+		/**********************************************************************/
+
+		/*
+				* Input Messages. Range = 50..89.
+				*/
+
+		// Generic Input Messages. Range = 50..59.
+		UIInteraction = 50,
+		Command = 51,
+
+		// Keyboard Input Message. Range = 60..69.
+		KeyDown = 60,
+		KeyUp = 61,
+		KeyPress = 62,
+
+		// Mouse Input Messages. Range = 70..79.
+		MouseEnter = 70,
+		MouseLeave = 71,
+		MouseDown = 72,
+		MouseUp = 73,
+		MouseMove = 74,
+		MouseWheel = 75,
+
+		// Touch Input Messages. Range = 80..89.
+		TouchStart = 80,
+		TouchEnd = 81,
+		TouchMove = 82,
+
+		// Gamepad Input Messages. Range = 90..99
+		GamepadButtonPressed = 90,
+		GamepadButtonReleased = 91,
+		GamepadAnalog = 92,
+
+		/**********************************************************************/
+
+		/*
+				* Ensure Count is the final entry.
+				*/
+		Count
+
+		/**********************************************************************/
+	};
+
+	//! Messages that can be sent to the webrtc players
+	enum class EToPlayerMsg : uint8
+	{
+		QualityControlOwnership = 0,
+		Response = 1,
+		Command = 2,
+		FreezeFrame = 3,
+		UnfreezeFrame = 4,
+		VideoEncoderAvgQP = 5, // average Quantisation Parameter value of Video Encoder, roughly depicts video encoding quality
+		LatencyTest = 6,
+		InitialSettings = 7,
+		FileExtension = 8,
+		FileMimeType = 9,
+		FileContents = 10,
+		TestEcho = 11
+	};
+
+	template <typename T>
+	static const T& ParseBuffer(const uint8*& Data, uint32& Size)
+	{
+		checkf(sizeof(T) <= Size, TEXT("%d - %d"), sizeof(T), Size);
+		const T& Value = *reinterpret_cast<const T*>(Data);
+		Data += sizeof(T);
+		Size -= sizeof(T);
+		return Value;
+	}
+
+	static FString ParseString(const webrtc::DataBuffer& Buffer, const size_t Offset = 0)
+	{
+		FString Res;
+		if (Buffer.data.size() > Offset)
+		{
+			size_t StringLength = (Buffer.data.size() - Offset) / sizeof(TCHAR);
+			Res.GetCharArray().SetNumUninitialized(StringLength + 1);
+			FMemory::Memcpy(Res.GetCharArray().GetData(), Buffer.data.data() + Offset, StringLength * sizeof(TCHAR));
+			Res.GetCharArray()[StringLength] = 0;
+		}
+		return Res;
+	}
+
+}; // namespace UE::PixelStreaming::Protocol
diff --git a/Engine/Plugins/Media/PixelStreaming/Source/PixelStreaming/Public/PixelStreamingTextureSource.h b/Engine/Plugins/Media/PixelStreaming/Source/PixelStreaming/Public/PixelStreamingTextureSource.h
new file mode 100644
--- /dev/null
+++ b/Engine/Plugins/Media/PixelStreaming/Source/PixelStreaming/Public/PixelStreamingTextureSource.h
@@ -0,0 +1,169 @@
+// Copyright Epic Games, Inc. All Rights Reserved.
+#pragma once
+
+#include "Templates/SharedPointer.h"
+#include "Templates/RefCounting.h"
+#include "Async/Async.h"
+#include "RHI.h"
+#include "HAL/ThreadSafeBool.h"
+#include "Rendering/SlateRenderer.h"
+#include "Framework/Application/SlateApplication.h"
+#include "IPixelStreamingTextureSource.h"
+#include "IPixelStreamingModule.h"
+
+/*
+ * Interface to wrap a Texture for use with the TextureSource
+ * By extending this class you can add additional operations to a class derived from BackBufferTextureSource
+ */
+
+class IPixelStreamingBackBufferTextureWrapper : public FRefCountBase
+{
+public:
+	explicit IPixelStreamingBackBufferTextureWrapper(FTexture2DRHIRef InTexture)
+	{
+		Texture = InTexture;
+	}
+
+	FTexture2DRHIRef& GetTexture()
+	{
+		return Texture;
+	}
+
+private:
+	FTexture2DRHIRef Texture;
+};
+
+/*
+ * Base class for TextureSources that get their textures from the UE backbuffer.
+ * Textures are copied from the backbuffer using a triplebuffering mechanism so that texture read access is always thread safe while writes are occurring.
+ * If no texture has been written since the last read then the same texture will be read again.
+ * This class also has the additional functionality of scaling textures from the backbuffer.
+ * Note: This is a template class that uses CRTP - see: https://en.wikipedia.org/wiki/Curiously_recurring_template_pattern
+ * For derived types of this class they must contain the following static methods:
+ *
+ * static void CopyTexture(const FTexture2DRHIRef& SourceTexture, TRefCountPtr<IPixelStreamingBackBufferTextureWrapper> DestTexture, FGPUFenceRHIRef& CopyFence)
+ * static TRefCountPtr<IPixelStreamingBackBufferTextureWrapper> CreateTexture(int Width, int Height)
+ * static FTexture2DRHIRef ToTextureRef(TRefCountPtr<IPixelStreamingBackBufferTextureWrapper> Texture)
+ * static FString& GetName()
+ */
+
+class PIXELSTREAMING_API IBackBufferTextureSource : public IPixelStreamingTextureSource, public TSharedFromThis<IBackBufferTextureSource, ESPMode::ThreadSafe>
+{
+public:
+	IBackBufferTextureSource(float InFrameScale);
+	IBackBufferTextureSource()
+		: IBackBufferTextureSource(1.0) {}
+
+	virtual ~IBackBufferTextureSource();
+
+	TRefCountPtr<IPixelStreamingBackBufferTextureWrapper> GetCurrent();
+
+	// abstract methods for overriding
+	virtual void CopyTexture(const FTexture2DRHIRef& SourceTexture, TRefCountPtr<IPixelStreamingBackBufferTextureWrapper> DestTexture, FGPUFenceRHIRef& CopyFence) = 0;
+	virtual TRefCountPtr<IPixelStreamingBackBufferTextureWrapper> CreateTexture(int Width, int Height) = 0;
+	virtual FString GetName() const = 0;
+
+	/* Begin IPixelStreamingTextureSource interface */
+	virtual void SetEnabled(bool bInEnabled) override;
+	virtual FTexture2DRHIRef GetTexture() override { return GetCurrent()->GetTexture(); }
+
+	virtual bool IsEnabled() const override { return *bEnabled; }
+	virtual bool IsAvailable() const override { return bInitialized; }
+	virtual int GetSourceWidth() const override { return SourceWidth; }
+	virtual int GetSourceHeight() const override { return SourceHeight; }
+	/* End IPixelStreamingTextureSource interface */
+
+	float GetFrameScaling() const { return FrameScale; }
+
+private:
+	void OnBackBufferReady_RenderThread(SWindow& SlateWindow, const FTexture2DRHIRef& FrameBuffer);
+	void Initialize(int Width, int Height);
+
+	struct FCaptureFrame
+	{
+		TRefCountPtr<IPixelStreamingBackBufferTextureWrapper> TextureWrapper;
+		FGPUFenceRHIRef Fence;
+		bool bAvailable = true;
+		uint64 PreWaitingOnCopy;
+	};
+
+private:
+	const float FrameScale;
+	int SourceWidth = 0;
+	int SourceHeight = 0;
+	FThreadSafeBool bInitialized = false;
+	TSharedRef<bool, ESPMode::ThreadSafe> bEnabled = MakeShared<bool, ESPMode::ThreadSafe>(true);
+	FSlateRenderer::FOnBackBufferReadyToPresent* OnBackbuffer = nullptr;
+	FDelegateHandle BackbufferDelegateHandle;
+	FCriticalSection DestructorCriticalSection;
+
+	/*
+	 * Triple buffer setup with queued write buffers (since we have to wait for RHI copy).
+	 * 1 Read buffer (read the captured texture)
+	 * 1 Temp buffer (for swapping what is read and written)
+	 * 2 Write buffers (2 write buffers because UE can render two frames before presenting sometimes)
+	 */
+	FCriticalSection CriticalSection;
+	bool bWriteParity = true;
+	FCaptureFrame WriteBuffers[2];
+	TRefCountPtr<IPixelStreamingBackBufferTextureWrapper> TempBuffer;
+	TRefCountPtr<IPixelStreamingBackBufferTextureWrapper> ReadBuffer;
+	FThreadSafeBool bIsTempDirty;
+};
+
+/*
+ * Basic IPixelStreamingBackBufferTextureWrapper for use with FBackBufferTextureSource
+ */
+
+class PIXELSTREAMING_API FPixelStreamingRHIBackBufferTexture : public IPixelStreamingBackBufferTextureWrapper
+{
+public:
+	FPixelStreamingRHIBackBufferTexture(FTexture2DRHIRef InTexture)
+		: IPixelStreamingBackBufferTextureWrapper(InTexture) {}
+};
+
+class PIXELSTREAMING_API FBackBufferTextureSource : public IBackBufferTextureSource
+{
+public:
+	FBackBufferTextureSource()
+		: IBackBufferTextureSource() {}
+	FBackBufferTextureSource(float InScale)
+		: IBackBufferTextureSource(InScale) {}
+
+	virtual void CopyTexture(const FTexture2DRHIRef& SourceTexture, TRefCountPtr<IPixelStreamingBackBufferTextureWrapper> DestTexture, FGPUFenceRHIRef& CopyFence) override;
+	virtual TRefCountPtr<IPixelStreamingBackBufferTextureWrapper> CreateTexture(int Width, int Height) override;
+
+	virtual FString GetName() const override { return TEXT("FBackBufferTextureSource"); }
+};
+
+/*
+ * Extention for IPixelStreamingBackBufferTextureWrapper which also stores a CPU accessible memory address of the texture
+ */
+class PIXELSTREAMING_API FPixelStreamingCPUReadableBackbufferTexture : public IPixelStreamingBackBufferTextureWrapper
+{
+public:
+	FPixelStreamingCPUReadableBackbufferTexture(FTexture2DRHIRef InTexture)
+		: IPixelStreamingBackBufferTextureWrapper(InTexture) {}
+
+	TArray<FColor>& GetRawPixels()
+	{
+		return RawPixels;
+	}
+
+private:
+	TArray<FColor> RawPixels;
+};
+
+class PIXELSTREAMING_API FBackBufferToCPUTextureSource : public IBackBufferTextureSource
+{
+public:
+	FBackBufferToCPUTextureSource()
+		: IBackBufferTextureSource() {}
+	FBackBufferToCPUTextureSource(float InScale)
+		: IBackBufferTextureSource(InScale) {}
+
+	virtual void CopyTexture(const FTexture2DRHIRef& SourceTexture, TRefCountPtr<IPixelStreamingBackBufferTextureWrapper> DestTexture, FGPUFenceRHIRef& CopyFence) override;
+	virtual TRefCountPtr<IPixelStreamingBackBufferTextureWrapper> CreateTexture(int Width, int Height) override;
+
+	virtual FString GetName() const override { return TEXT("FBackBufferToCPUTextureSource"); }
+};
\ No newline at end of file
