NoPCH Compile UnrealEditor Win64 - IPixelStreamingModule.h syntax error: identifier 'FGPUFenceRHIRef'

diff --git a/Engine/Plugins/Media/PixelStreaming/Source/PixelStreaming/Private/EncoderFrameFactory.h b/Engine/Plugins/Media/PixelStreaming/Source/PixelStreaming/Private/EncoderFrameFactory.h
--- a/Engine/Plugins/Media/PixelStreaming/Source/PixelStreaming/Private/EncoderFrameFactory.h
+++ b/Engine/Plugins/Media/PixelStreaming/Source/PixelStreaming/Private/EncoderFrameFactory.h
@@ -3,6 +3,10 @@
 #pragma once
 
 #include "VideoEncoderInput.h"
+#include "Templates/RefCounting.h"
+#include "Templates/SharedPointer.h"
+
+class FRHITexture2D;
 
 namespace UE
 {
@@ -19,26 +23,26 @@ namespace UE
 		public:
 			FEncoderFrameFactory();
 			~FEncoderFrameFactory();
-			TSharedPtr<AVEncoder::FVideoEncoderInputFrame> GetFrameAndSetTexture(FTexture2DRHIRef InTexture);
+			TSharedPtr<AVEncoder::FVideoEncoderInputFrame> GetFrameAndSetTexture(TRefCountPtr<FRHITexture2D> InTexture);
 			TSharedPtr<AVEncoder::FVideoEncoderInput> GetOrCreateVideoEncoderInput();
 
 		private:
-			TSharedPtr<AVEncoder::FVideoEncoderInputFrame> GetOrCreateFrame(const FTexture2DRHIRef InTexture);
+			TSharedPtr<AVEncoder::FVideoEncoderInputFrame> GetOrCreateFrame(const TRefCountPtr<FRHITexture2D> InTexture);
 			void RemoveStaleTextures();
 			void FlushFrames();
 			TSharedPtr<AVEncoder::FVideoEncoderInput> CreateVideoEncoderInput() const;
-			void SetTexture(TSharedPtr<AVEncoder::FVideoEncoderInputFrame> Frame, const FTexture2DRHIRef& Texture);
-			void SetTextureCUDAVulkan(TSharedPtr<AVEncoder::FVideoEncoderInputFrame> Frame, const FTexture2DRHIRef& Texture);
+			void SetTexture(TSharedPtr<AVEncoder::FVideoEncoderInputFrame> Frame, const TRefCountPtr<FRHITexture2D>& Texture);
+			void SetTextureCUDAVulkan(TSharedPtr<AVEncoder::FVideoEncoderInputFrame> Frame, const TRefCountPtr<FRHITexture2D>& Texture);
 #if PLATFORM_WINDOWS
-			void SetTextureCUDAD3D11(TSharedPtr<AVEncoder::FVideoEncoderInputFrame> Frame, const FTexture2DRHIRef& Texture);
-			void SetTextureCUDAD3D12(TSharedPtr<AVEncoder::FVideoEncoderInputFrame> Frame, const FTexture2DRHIRef& Texture);
+			void SetTextureCUDAD3D11(TSharedPtr<AVEncoder::FVideoEncoderInputFrame> Frame, const TRefCountPtr<FRHITexture2D>& Texture);
+			void SetTextureCUDAD3D12(TSharedPtr<AVEncoder::FVideoEncoderInputFrame> Frame, const TRefCountPtr<FRHITexture2D>& Texture);
 #endif // PLATFORM_WINDOWS
 
 		private:
 			uint64 FrameId = 0;
 			TSharedPtr<AVEncoder::FVideoEncoderInput> EncoderInput;
 			// Store a mapping between raw textures and the FVideoEncoderInputFrames that wrap them
-			TMap<FTexture2DRHIRef, TSharedPtr<AVEncoder::FVideoEncoderInputFrame>> TextureToFrameMapping;
+			TMap<TRefCountPtr<FRHITexture2D>, TSharedPtr<AVEncoder::FVideoEncoderInputFrame>> TextureToFrameMapping;
 		};
 	} // namespace PixelStreaming
 } // namespace UE
diff --git a/Engine/Plugins/Media/PixelStreaming/Source/PixelStreaming/Private/VideoEncoderH264Wrapper.h b/Engine/Plugins/Media/PixelStreaming/Source/PixelStreaming/Private/VideoEncoderH264Wrapper.h
--- a/Engine/Plugins/Media/PixelStreaming/Source/PixelStreaming/Private/VideoEncoderH264Wrapper.h
+++ b/Engine/Plugins/Media/PixelStreaming/Source/PixelStreaming/Private/VideoEncoderH264Wrapper.h
@@ -4,6 +4,11 @@
 #include "VideoEncoder.h"
 #include "EncoderFrameFactory.h"
 
+namespace webrtc
+{
+	class VideoFrame;
+}
+
 namespace UE {
 	namespace PixelStreaming {
 		class FVideoEncoderFactory;
diff --git a/Engine/Plugins/Media/PixelStreaming/Source/PixelStreaming/Public/IPixelStreamingModule.h b/Engine/Plugins/Media/PixelStreaming/Source/PixelStreaming/Public/IPixelStreamingModule.h
--- a/Engine/Plugins/Media/PixelStreaming/Source/PixelStreaming/Public/IPixelStreamingModule.h
+++ b/Engine/Plugins/Media/PixelStreaming/Source/PixelStreaming/Public/IPixelStreamingModule.h
@@ -7,15 +7,16 @@
 #include "Modules/ModuleManager.h"
 #include "IInputDeviceModule.h"
 #include "Templates/SharedPointer.h"
+#include "Templates/RefCounting.h"
 #include "PixelStreamingPlayerId.h"
 #include "IPixelStreamingAudioSink.h"
 #include "IPixelStreamingPumpedVideoSource.h"
-#include "Templates/SharedPointer.h"
 
 #include "IInputDevice.h"
 
 class UTexture2D;
 class UPixelStreamingInput;
+class FRHIGPUFence;
 
 namespace webrtc
 {
@@ -197,5 +198,5 @@ public:
 	/*
 	 * Create add a task to the GPU fence poller.
 	 */
-	virtual void AddGPUFencePollerTask(FGPUFenceRHIRef Fence, TSharedRef<bool, ESPMode::ThreadSafe> bIsEnabled, TFunction<void()> Task) = 0;
+	virtual void AddGPUFencePollerTask(TRefCountPtr<FRHIGPUFence> Fence, TSharedRef<bool, ESPMode::ThreadSafe> bIsEnabled, TFunction<void()> Task) = 0;
 };
diff --git a/Engine/Plugins/Media/PixelStreaming/Source/PixelStreaming/Public/PixelStreamingFrameBuffer.h b/Engine/Plugins/Media/PixelStreaming/Source/PixelStreaming/Public/PixelStreamingFrameBuffer.h
--- a/Engine/Plugins/Media/PixelStreaming/Source/PixelStreaming/Public/PixelStreamingFrameBuffer.h
+++ b/Engine/Plugins/Media/PixelStreaming/Source/PixelStreaming/Public/PixelStreamingFrameBuffer.h
@@ -2,11 +2,14 @@
 
 #pragma once
 
+#include "Templates/RefCounting.h"
+#include "Templates/SharedPointer.h"
+
 #include "api/video/video_frame_buffer.h"
 #include "api/video/i420_buffer.h"
 #include "api/scoped_refptr.h"
 
-class FTexture2DRHI;
+class FRHITexture2D;
 class IPixelStreamingTextureSource;
 
 namespace UE
@@ -108,7 +111,7 @@ public:
 	virtual ~FLayerFrameBuffer();
 
 	virtual UE::PixelStreaming::EFrameBufferType GetFrameBufferType() const { return UE::PixelStreaming::EFrameBufferType::Layer; }
-	FTexture2DRHIRef GetFrame() const;
+	TRefCountPtr<FRHITexture2D> GetFrame() const;
 
 	virtual int width() const override;
 	virtual int height() const override;
@@ -128,7 +131,7 @@ public:
 	virtual ~FFrameBufferI420();
 
 	virtual UE::PixelStreaming::EFrameBufferType GetFrameBufferType() const { return UE::PixelStreaming::EFrameBufferType::Layer; }
-	FTexture2DRHIRef GetFrame() const;
+	TRefCountPtr<FRHITexture2D> GetFrame() const;
 
 	virtual int width() const override;
 	virtual int height() const override;
