Fix streamer to work with older versions of the cirrus signalling server where the offer/answer pattern is reversed

diff --git a/Engine/Plugins/Media/PixelStreaming/Source/PixelStreaming/Private/Streamer.cpp b/Engine/Plugins/Media/PixelStreaming/Source/PixelStreaming/Private/Streamer.cpp
--- a/Engine/Plugins/Media/PixelStreaming/Source/PixelStreaming/Private/Streamer.cpp
+++ b/Engine/Plugins/Media/PixelStreaming/Source/PixelStreaming/Private/Streamer.cpp
@@ -262,7 +262,6 @@ namespace UE::PixelStreaming
 			}
 
 			SendAnswer(NewSession, TUniquePtr<webrtc::SessionDescriptionInterface>(SessionDesc.release()));
-			ForceKeyFrame();
 		}
 		else
 			UE_LOG(LogPixelStreaming, Error, TEXT("Failed to create player session, peer connection was nullptr."));
