Add network priority of video packet to be as high as possible.

diff --git a/Engine/Plugins/Media/PixelStreaming/Source/PixelStreaming/Private/Streamer.cpp b/Engine/Plugins/Media/PixelStreaming/Source/PixelStreaming/Private/Streamer.cpp
--- a/Engine/Plugins/Media/PixelStreaming/Source/PixelStreaming/Private/Streamer.cpp
+++ b/Engine/Plugins/Media/PixelStreaming/Source/PixelStreaming/Private/Streamer.cpp
@@ -23,7 +23,7 @@ namespace UE::PixelStreaming
 		return AVEncoder::FVideoEncoderFactory::Get().HasEncoderForCodec(AVEncoder::ECodecType::H264);
 	}
 
-	FStreamer::FStreamer(const FString &InStreamerId)
+	FStreamer::FStreamer(const FString& InStreamerId)
 		: StreamerId(InStreamerId)
 		, WebRtcSignallingThread(MakeUnique<rtc::Thread>(rtc::SocketServer::CreateDefault()))
 		, SignallingServerConnection(MakeUnique<FSignallingServerConnection>(*this, InStreamerId))
@@ -48,7 +48,7 @@ namespace UE::PixelStreaming
 		rtc::CleanupSSL();
 	}
 
-	bool FStreamer::StartStreaming(const FString &SignallingServerUrl)
+	bool FStreamer::StartStreaming(const FString& SignallingServerUrl)
 	{
 		StopStreaming();
 		ConnectToSignallingServer(SignallingServerUrl);
@@ -130,9 +130,9 @@ namespace UE::PixelStreaming
 		check(SFUPeerConnectionFactory.get() != nullptr);
 	}
 
-	webrtc::AudioProcessing *FStreamer::SetupAudioProcessingModule()
+	webrtc::AudioProcessing* FStreamer::SetupAudioProcessingModule()
 	{
-		webrtc::AudioProcessing *AudioProcessingModule = webrtc::AudioProcessingBuilder().Create();
+		webrtc::AudioProcessing* AudioProcessingModule = webrtc::AudioProcessingBuilder().Create();
 		webrtc::AudioProcessing::Config Config;
 		// Enabled multi channel audio capture/render
 		Config.pipeline.multi_channel_capture = true;
@@ -155,7 +155,7 @@ namespace UE::PixelStreaming
 		return AudioProcessingModule;
 	}
 
-	void FStreamer::ConnectToSignallingServer(const FString &Url)
+	void FStreamer::ConnectToSignallingServer(const FString& Url)
 	{
 		PreviousSignallingServerURL = Url;
 		SignallingServerConnection->Connect(Url);
@@ -166,12 +166,12 @@ namespace UE::PixelStreaming
 		SignallingServerConnection->Disconnect();
 	}
 
-	void FStreamer::OnConfig(const webrtc::PeerConnectionInterface::RTCConfiguration &Config)
+	void FStreamer::OnConfig(const webrtc::PeerConnectionInterface::RTCConfiguration& Config)
 	{
 		PeerConnectionConfig = Config;
 	}
 
-	void FStreamer::OnDataChannelOpen(FPixelStreamingPlayerId PlayerId, webrtc::DataChannelInterface *DataChannel)
+	void FStreamer::OnDataChannelOpen(FPixelStreamingPlayerId PlayerId, webrtc::DataChannelInterface* DataChannel)
 	{
 		if (DataChannel)
 		{
@@ -196,7 +196,7 @@ namespace UE::PixelStreaming
 			Flags);
 		if (NewSession)
 		{
-			webrtc::PeerConnectionInterface &PeerConnection = NewSession->GetPeerConnection();
+			webrtc::PeerConnectionInterface& PeerConnection = NewSession->GetPeerConnection();
 			const bool SupportsDataChannel = (Flags & Protocol::EPlayerFlags::PSPFlag_SupportsDataChannel) != 0;
 
 			if (SupportsDataChannel)
@@ -214,9 +214,9 @@ namespace UE::PixelStreaming
 				NewSession->SetDataChannel(DataChannel);
 
 				// Bind to the on data channel open delegate (we use this as the earliest time peer is ready for freeze frame)
-				FDataChannelObserver *DataChannelObserver = NewSession->GetDataChannelObserver();
+				FDataChannelObserver* DataChannelObserver = NewSession->GetDataChannelObserver();
 				check(DataChannelObserver);
-				if (UPixelStreamingDelegates *Delegates = UPixelStreamingDelegates::GetPixelStreamingDelegates())
+				if (UPixelStreamingDelegates* Delegates = UPixelStreamingDelegates::GetPixelStreamingDelegates())
 				{
 					Delegates->OnDataChannelOpenNative.AddRaw(this, &FStreamer::OnDataChannelOpen);
 				}
@@ -227,28 +227,27 @@ namespace UE::PixelStreaming
 		return NewSession;
 	}
 
-	void FStreamer::OnSessionDescription(FPixelStreamingPlayerId PlayerId, webrtc::SdpType Type, const FString &Sdp)
+	void FStreamer::OnSessionDescription(FPixelStreamingPlayerId PlayerId, webrtc::SdpType Type, const FString& Sdp)
 	{
 		switch (Type)
 		{
-		case webrtc::SdpType::kOffer:
-			OnOffer(PlayerId, Sdp);
-			break;
-		case webrtc::SdpType::kAnswer:
-		case webrtc::SdpType::kPrAnswer:
-		{
-			PlayerSessions.ForSession(PlayerId, [&](TSharedPtr<IPlayerSession> Session)
-									  { Session->OnAnswer(Sdp); });
-			bStreamingStarted = true;
-			break;
-		}
-		case webrtc::SdpType::kRollback:
-			UE_LOG(LogPixelStreaming, Error, TEXT("Rollback SDP is currently unsupported. SDP is: %s"), *Sdp);
-			break;
+			case webrtc::SdpType::kOffer:
+				OnOffer(PlayerId, Sdp);
+				break;
+			case webrtc::SdpType::kAnswer:
+			case webrtc::SdpType::kPrAnswer:
+			{
+				PlayerSessions.ForSession(PlayerId, [&](TSharedPtr<IPlayerSession> Session) { Session->OnAnswer(Sdp); });
+				bStreamingStarted = true;
+				break;
+			}
+			case webrtc::SdpType::kRollback:
+				UE_LOG(LogPixelStreaming, Error, TEXT("Rollback SDP is currently unsupported. SDP is: %s"), *Sdp);
+				break;
 		}
 	}
 
-	void FStreamer::OnOffer(FPixelStreamingPlayerId PlayerId, const FString &Sdp)
+	void FStreamer::OnOffer(FPixelStreamingPlayerId PlayerId, const FString& Sdp)
 	{
 		TSharedPtr<IPlayerSession> NewSession = CreateSession(PlayerId, Protocol::EPlayerFlags::PSPFlag_SupportsDataChannel);
 		if (NewSession)
@@ -267,7 +266,7 @@ namespace UE::PixelStreaming
 			UE_LOG(LogPixelStreaming, Error, TEXT("Failed to create player session, peer connection was nullptr."));
 	}
 
-	void FStreamer::SetLocalDescription(webrtc::PeerConnectionInterface *PeerConnection, FSetSessionDescriptionObserver *Observer, webrtc::SessionDescriptionInterface *SDP)
+	void FStreamer::SetLocalDescription(webrtc::PeerConnectionInterface* PeerConnection, FSetSessionDescriptionObserver* Observer, webrtc::SessionDescriptionInterface* SDP)
 	{
 		// Note from Luke about WebRTC: the sink of video capturer will be added as a direct result
 		// of `PeerConnection->SetLocalDescription()` call but video encoder will be created later on
@@ -289,7 +288,7 @@ namespace UE::PixelStreaming
 				webrtc::RTCError Err = Sender->SetParameters(ExistingParams);
 				if (!Err.ok())
 				{
-					const char *ErrMsg = Err.message();
+					const char* ErrMsg = Err.message();
 					FString ErrorStr(ErrMsg);
 					UE_LOG(LogPixelStreaming, Error, TEXT("Failed to set RTP Sender params: %s"), *ErrorStr);
 				}
@@ -305,35 +304,33 @@ namespace UE::PixelStreaming
 		//		PeerConnection.SetLocalDescription(Answer);
 		//		SignallingServerConnection.SendAnswer(Answer);
 
-		FSetSessionDescriptionObserver *SetLocalDescriptionObserver = FSetSessionDescriptionObserver::Create(
+		FSetSessionDescriptionObserver* SetLocalDescriptionObserver = FSetSessionDescriptionObserver::Create(
 			[this, Session]() // on success
 			{
 				SignallingServerConnection->SendAnswer(Session->GetPlayerId(), *Session->GetPeerConnection().local_description());
 				bStreamingStarted = true;
 			},
-			[this, Session](const FString &Error) // on failure
+			[this, Session](const FString& Error) // on failure
 			{
 				UE_LOG(LogPixelStreaming, Error, TEXT("Failed to set local description: %s"), *Error);
 				Session->DisconnectPlayer(Error);
 				PlayerSessions.DeletePlayerSession(Session->GetPlayerId());
 			});
 
-		auto OnCreateAnswerSuccess = [this, Session, SetLocalDescriptionObserver](webrtc::SessionDescriptionInterface *SDP)
-		{
+		auto OnCreateAnswerSuccess = [this, Session, SetLocalDescriptionObserver](webrtc::SessionDescriptionInterface* SDP) {
 			SetLocalDescription(&Session->GetPeerConnection(), SetLocalDescriptionObserver, SDP);
 		};
 
-		FCreateSessionDescriptionObserver *CreateAnswerObserver = FCreateSessionDescriptionObserver::Create(
+		FCreateSessionDescriptionObserver* CreateAnswerObserver = FCreateSessionDescriptionObserver::Create(
 			MoveTemp(OnCreateAnswerSuccess),
-			[this, Session](const FString &Error) // on failure
+			[this, Session](const FString& Error) // on failure
 			{
 				UE_LOG(LogPixelStreaming, Error, TEXT("Failed to create answer: %s"), *Error);
 				Session->DisconnectPlayer(Error);
 				PlayerSessions.DeletePlayerSession(Session->GetPlayerId());
 			});
 
-		auto OnSetRemoteDescriptionSuccess = [this, Session, CreateAnswerObserver]()
-		{
+		auto OnSetRemoteDescriptionSuccess = [this, Session, CreateAnswerObserver]() {
 			// Note: these offer to receive are superseded now we are use transceivers to setup our peer connection media
 			int offer_to_receive_video = webrtc::PeerConnectionInterface::RTCOfferAnswerOptions::kUndefined;
 			int offer_to_receive_audio = webrtc::PeerConnectionInterface::RTCOfferAnswerOptions::kUndefined;
@@ -346,39 +343,40 @@ namespace UE::PixelStreaming
 				offer_to_receive_audio,
 				voice_activity_detection,
 				ice_restart,
-				use_rtp_mux};
+				use_rtp_mux
+			};
 
 			AddStreams(Session, Protocol::EPlayerFlags::PSPFlag_SupportsDataChannel);
 			ModifyTransceivers(Session->GetPeerConnection().GetTransceivers(), Protocol::EPlayerFlags::PSPFlag_SupportsDataChannel);
 			Session->GetPeerConnection().CreateAnswer(CreateAnswerObserver, AnswerOption);
 		};
-		FSetSessionDescriptionObserver *SetRemoteDescriptionObserver = FSetSessionDescriptionObserver::Create(
+		FSetSessionDescriptionObserver* SetRemoteDescriptionObserver = FSetSessionDescriptionObserver::Create(
 			MoveTemp(OnSetRemoteDescriptionSuccess),
-			[this, Session](const FString &Error) // on failure
+			[this, Session](const FString& Error) // on failure
 			{
 				UE_LOG(LogPixelStreaming, Error, TEXT("Failed to set remote description: %s"), *Error);
 				Session->DisconnectPlayer(Error);
 				PlayerSessions.DeletePlayerSession(Session->GetPlayerId());
 			});
 
-		cricket::SessionDescription *RemoteDescription = Sdp->description();
+		cricket::SessionDescription* RemoteDescription = Sdp->description();
 		MungeRemoteSDP(RemoteDescription);
 
 		Session->GetPeerConnection().SetRemoteDescription(SetRemoteDescriptionObserver, Sdp.Release());
 	}
 
-	void FStreamer::MungeRemoteSDP(cricket::SessionDescription *RemoteDescription)
+	void FStreamer::MungeRemoteSDP(cricket::SessionDescription* RemoteDescription)
 	{
 		// Munge SDP of remote description to inject min, max, start bitrates
-		std::vector<cricket::ContentInfo> &ContentInfos = RemoteDescription->contents();
-		for (cricket::ContentInfo &Content : ContentInfos)
+		std::vector<cricket::ContentInfo>& ContentInfos = RemoteDescription->contents();
+		for (cricket::ContentInfo& Content : ContentInfos)
 		{
-			cricket::MediaContentDescription *MediaDescription = Content.media_description();
+			cricket::MediaContentDescription* MediaDescription = Content.media_description();
 			if (MediaDescription->type() == cricket::MediaType::MEDIA_TYPE_VIDEO)
 			{
-				cricket::VideoContentDescription *VideoDescription = MediaDescription->as_video();
+				cricket::VideoContentDescription* VideoDescription = MediaDescription->as_video();
 				std::vector<cricket::VideoCodec> CodecsCopy = VideoDescription->codecs();
-				for (cricket::VideoCodec &Codec : CodecsCopy)
+				for (cricket::VideoCodec& Codec : CodecsCopy)
 				{
 					// Note: These params are passed as kilobits, so divide by 1000.
 					Codec.SetParam(cricket::kCodecParamMinBitrate, Settings::CVarPixelStreamingWebRTCMinBitrate.GetValueOnAnyThread() / 1000);
@@ -390,10 +388,9 @@ namespace UE::PixelStreaming
 		}
 	}
 
-	void FStreamer::OnRemoteIceCandidate(FPixelStreamingPlayerId PlayerId, const FString &SdpMid, int SdpMLineIndex, const FString &Sdp)
+	void FStreamer::OnRemoteIceCandidate(FPixelStreamingPlayerId PlayerId, const FString& SdpMid, int SdpMLineIndex, const FString& Sdp)
 	{
-		PlayerSessions.ForSession(PlayerId, [&](TSharedPtr<IPlayerSession> Session)
-								  { Session->OnRemoteIceCandidate(SdpMid, SdpMLineIndex, Sdp); });
+		PlayerSessions.ForSession(PlayerId, [&](TSharedPtr<IPlayerSession> Session) { Session->OnRemoteIceCandidate(SdpMid, SdpMLineIndex, Sdp); });
 	}
 
 	void FStreamer::OnPlayerConnected(FPixelStreamingPlayerId PlayerId, int Flags)
@@ -403,19 +400,19 @@ namespace UE::PixelStreaming
 		{
 			AddStreams(NewSession, Flags);
 
-			webrtc::PeerConnectionInterface &PeerConnection = NewSession->GetPeerConnection();
+			webrtc::PeerConnectionInterface& PeerConnection = NewSession->GetPeerConnection();
 			ModifyTransceivers(PeerConnection.GetTransceivers(), Flags);
 
 			// observer for creating offer
-			FCreateSessionDescriptionObserver *CreateOfferObserver = FCreateSessionDescriptionObserver::Create(
-				[this, PlayerId, PeerConnectionPtr = &PeerConnection](webrtc::SessionDescriptionInterface *SDP) // on SDP create success
+			FCreateSessionDescriptionObserver* CreateOfferObserver = FCreateSessionDescriptionObserver::Create(
+				[this, PlayerId, PeerConnectionPtr = &PeerConnection](webrtc::SessionDescriptionInterface* SDP) // on SDP create success
 				{
-					FSetSessionDescriptionObserver *SetLocalDescriptionObserver = FSetSessionDescriptionObserver::Create(
+					FSetSessionDescriptionObserver* SetLocalDescriptionObserver = FSetSessionDescriptionObserver::Create(
 						[this, PlayerId, SDP]() // on SDP set success
 						{
 							SignallingServerConnection->SendOffer(PlayerId, *SDP);
 						},
-						[](const FString &Error) // on SDP set failure
+						[](const FString& Error) // on SDP set failure
 						{
 							UE_LOG(LogPixelStreaming, Error, TEXT("Failed to set local description: %s"), *Error);
 						});
@@ -423,7 +420,7 @@ namespace UE::PixelStreaming
 					MungeLocalSDP(SDP->description());
 					SetLocalDescription(PeerConnectionPtr, SetLocalDescriptionObserver, SDP);
 				},
-				[](const FString &Error) // on SDP create failure
+				[](const FString& Error) // on SDP create failure
 				{
 					UE_LOG(LogPixelStreaming, Error, TEXT("Failed to create offer: %s"), *Error);
 				});
@@ -514,37 +511,38 @@ namespace UE::PixelStreaming
 			using FLayer = Settings::FSimulcastParameters::FLayer;
 
 			// encodings should be lowest res to highest
-			TArray<FLayer *> SortedLayers;
-			for (FLayer &Layer : Settings::SimulcastParameters.Layers)
+			TArray<FLayer*> SortedLayers;
+			for (FLayer& Layer : Settings::SimulcastParameters.Layers)
 			{
 				SortedLayers.Add(&Layer);
 			}
 
-			SortedLayers.Sort([](const FLayer &LayerA, const FLayer &LayerB)
-							  { return LayerA.Scaling > LayerB.Scaling; });
+			SortedLayers.Sort([](const FLayer& LayerA, const FLayer& LayerB) { return LayerA.Scaling > LayerB.Scaling; });
 
 			const int LayerCount = SortedLayers.Num();
 			for (int i = 0; i < LayerCount; ++i)
 			{
-				const FLayer *SimulcastLayer = SortedLayers[i];
+				const FLayer* SimulcastLayer = SortedLayers[i];
 				webrtc::RtpEncodingParameters LayerEncoding{};
 				LayerEncoding.rid = TCHAR_TO_UTF8(*(FString("simulcast") + FString::FromInt(LayerCount - i)));
 				LayerEncoding.min_bitrate_bps = SimulcastLayer->MinBitrate;
 				LayerEncoding.max_bitrate_bps = SimulcastLayer->MaxBitrate;
 				LayerEncoding.scale_resolution_down_by = SimulcastLayer->Scaling;
+				LayerEncoding.network_priority = webrtc::Priority::kHigh;
 				LayerEncoding.max_framerate = FMath::Max(60, Settings::CVarPixelStreamingWebRTCFps.GetValueOnAnyThread());
 				EncodingParams.push_back(LayerEncoding);
 			}
 		}
 		else
 		{
-			webrtc::RtpEncodingParameters encoding{};
-			encoding.rid = "base";
-			encoding.max_bitrate_bps = Settings::CVarPixelStreamingWebRTCMaxBitrate.GetValueOnAnyThread();
-			encoding.min_bitrate_bps = Settings::CVarPixelStreamingWebRTCMinBitrate.GetValueOnAnyThread();
-			encoding.max_framerate = FMath::Max(60, Settings::CVarPixelStreamingWebRTCFps.GetValueOnAnyThread());
-			encoding.scale_resolution_down_by.reset();
-			EncodingParams.push_back(encoding);
+			webrtc::RtpEncodingParameters Encoding{};
+			Encoding.rid = "base";
+			Encoding.max_bitrate_bps = Settings::CVarPixelStreamingWebRTCMaxBitrate.GetValueOnAnyThread();
+			Encoding.min_bitrate_bps = Settings::CVarPixelStreamingWebRTCMinBitrate.GetValueOnAnyThread();
+			Encoding.max_framerate = FMath::Max(60, Settings::CVarPixelStreamingWebRTCFps.GetValueOnAnyThread());
+			Encoding.scale_resolution_down_by.reset();
+			Encoding.network_priority = webrtc::Priority::kHigh;
+			EncodingParams.push_back(Encoding);
 		}
 
 		return EncodingParams;
@@ -579,20 +577,20 @@ namespace UE::PixelStreaming
 		webrtc::DegradationPreference DegradationPref = Settings::GetDegradationPreference();
 		switch (DegradationPref)
 		{
-		case webrtc::DegradationPreference::MAINTAIN_FRAMERATE:
-			VideoTrack->set_content_hint(webrtc::VideoTrackInterface::ContentHint::kFluid);
-			break;
-		case webrtc::DegradationPreference::MAINTAIN_RESOLUTION:
-			VideoTrack->set_content_hint(webrtc::VideoTrackInterface::ContentHint::kDetailed);
-			break;
-		default:
-			break;
+			case webrtc::DegradationPreference::MAINTAIN_FRAMERATE:
+				VideoTrack->set_content_hint(webrtc::VideoTrackInterface::ContentHint::kFluid);
+				break;
+			case webrtc::DegradationPreference::MAINTAIN_RESOLUTION:
+				VideoTrack->set_content_hint(webrtc::VideoTrackInterface::ContentHint::kDetailed);
+				break;
+			default:
+				break;
 		}
 
 		bool bHasVideoTransceiver = false;
 
-		webrtc::PeerConnectionInterface &PeerConnection = Session->GetPeerConnection();
-		for (auto &Transceiver : PeerConnection.GetTransceivers())
+		webrtc::PeerConnectionInterface& PeerConnection = Session->GetPeerConnection();
+		for (auto& Transceiver : PeerConnection.GetTransceivers())
 		{
 			rtc::scoped_refptr<webrtc::RtpSenderInterface> Sender = Transceiver->sender();
 			if (Transceiver->media_type() == cricket::MediaType::MEDIA_TYPE_VIDEO)
@@ -606,7 +604,7 @@ namespace UE::PixelStreaming
 		if (!bHasVideoTransceiver)
 		{
 			webrtc::RtpTransceiverInit TransceiverOptions;
-			TransceiverOptions.stream_ids = {TCHAR_TO_UTF8(*InVideoStreamId)};
+			TransceiverOptions.stream_ids = { TCHAR_TO_UTF8(*InVideoStreamId) };
 			TransceiverOptions.direction = webrtc::RtpTransceiverDirection::kSendOnly;
 			TransceiverOptions.send_encodings = CreateRTPEncodingParams(Flags);
 
@@ -615,25 +613,25 @@ namespace UE::PixelStreaming
 		}
 	}
 
-	void FStreamer::MungeLocalSDP(cricket::SessionDescription *SessionDescription)
+	void FStreamer::MungeLocalSDP(cricket::SessionDescription* SessionDescription)
 	{
 
-		std::vector<cricket::ContentInfo> &ContentInfos = SessionDescription->contents();
-		for (cricket::ContentInfo &ContentInfo : ContentInfos)
+		std::vector<cricket::ContentInfo>& ContentInfos = SessionDescription->contents();
+		for (cricket::ContentInfo& ContentInfo : ContentInfos)
 		{
-			cricket::MediaContentDescription *MediaDescription = ContentInfo.media_description();
+			cricket::MediaContentDescription* MediaDescription = ContentInfo.media_description();
 			cricket::MediaType MediaType = MediaDescription->type();
 			if (MediaType != cricket::MediaType::MEDIA_TYPE_AUDIO)
 			{
 				continue;
 			}
-			cricket::AudioContentDescription *AudioDescription = MediaDescription->as_audio();
+			cricket::AudioContentDescription* AudioDescription = MediaDescription->as_audio();
 			if (AudioDescription == nullptr)
 			{
 				continue;
 			}
 			std::vector<cricket::AudioCodec> CodecsCopy = AudioDescription->codecs();
-			for (cricket::AudioCodec &Codec : CodecsCopy)
+			for (cricket::AudioCodec& Codec : CodecsCopy)
 			{
 				if (Codec.name == "opus")
 				{
@@ -692,8 +690,8 @@ namespace UE::PixelStreaming
 		rtc::scoped_refptr<webrtc::AudioTrackInterface> AudioTrack = PCFactory->CreateAudioTrack(TCHAR_TO_UTF8(*InAudioTrackLabel), AudioSource);
 
 		bool bHasAudioTransceiver = false;
-		webrtc::PeerConnectionInterface &PeerConnection = Session->GetPeerConnection();
-		for (auto &Transceiver : PeerConnection.GetTransceivers())
+		webrtc::PeerConnectionInterface& PeerConnection = Session->GetPeerConnection();
+		for (auto& Transceiver : PeerConnection.GetTransceivers())
 		{
 			rtc::scoped_refptr<webrtc::RtpSenderInterface> Sender = Transceiver->sender();
 			if (Transceiver->media_type() == cricket::MediaType::MEDIA_TYPE_AUDIO)
@@ -706,7 +704,7 @@ namespace UE::PixelStreaming
 		if (!bHasAudioTransceiver)
 		{
 			// Add the track
-			webrtc::RTCErrorOr<rtc::scoped_refptr<webrtc::RtpSenderInterface>> Result = PeerConnection.AddTrack(AudioTrack, {TCHAR_TO_UTF8(*InAudioStreamId)});
+			webrtc::RTCErrorOr<rtc::scoped_refptr<webrtc::RtpSenderInterface>> Result = PeerConnection.AddTrack(AudioTrack, { TCHAR_TO_UTF8(*InAudioStreamId) });
 
 			if (!Result.ok())
 			{
@@ -721,7 +719,7 @@ namespace UE::PixelStreaming
 		bool bIsSFU = (Flags & Protocol::EPlayerFlags::PSPFlag_IsSFU) != Protocol::EPlayerFlags::PSPFlag_None;
 		bool bReceiveBrowserAudio = !bIsSFU && !Settings::CVarPixelStreamingWebRTCDisableReceiveAudio.GetValueOnAnyThread();
 
-		for (auto &Transceiver : Transceivers)
+		for (auto& Transceiver : Transceivers)
 		{
 			rtc::scoped_refptr<webrtc::RtpSenderInterface> Sender = Transceiver->sender();
 
@@ -748,12 +746,12 @@ namespace UE::PixelStreaming
 				}
 
 				Transceiver->SetDirection(AudioTransceiverDirection);
-				Sender->SetStreams({TCHAR_TO_UTF8(*GetAudioStreamID())});
+				Sender->SetStreams({ TCHAR_TO_UTF8(*GetAudioStreamID()) });
 			}
 			else if (Transceiver->media_type() == cricket::MediaType::MEDIA_TYPE_VIDEO)
 			{
 				Transceiver->SetDirection(webrtc::RtpTransceiverDirection::kSendOnly);
-				Sender->SetStreams({TCHAR_TO_UTF8(*GetVideoStreamID())});
+				Sender->SetStreams({ TCHAR_TO_UTF8(*GetVideoStreamID()) });
 
 				webrtc::RtpParameters RtpParams = Sender->GetParameters();
 				RtpParams.encodings = CreateRTPEncodingParams(Flags);
@@ -762,21 +760,19 @@ namespace UE::PixelStreaming
 		}
 	}
 
-	void FStreamer::SendPlayerMessage(Protocol::EToPlayerMsg Type, const FString &Descriptor)
+	void FStreamer::SendPlayerMessage(Protocol::EToPlayerMsg Type, const FString& Descriptor)
 	{
-		PlayerSessions.ForEachSession([&](TSharedPtr<IPlayerSession> PlayerSession)
-									  { PlayerSession->SendMessage(Type, Descriptor); });
+		PlayerSessions.ForEachSession([&](TSharedPtr<IPlayerSession> PlayerSession) { PlayerSession->SendMessage(Type, Descriptor); });
 	}
 
-	void FStreamer::SendFreezeFrame(const TArray64<uint8> &JpegBytes)
+	void FStreamer::SendFreezeFrame(const TArray64<uint8>& JpegBytes)
 	{
 		if (!bStreamingStarted)
 		{
 			return;
 		}
 
-		PlayerSessions.ForEachSession([&JpegBytes](TSharedPtr<IPlayerSession> PlayerSession)
-									  { PlayerSession->SendFreezeFrame(JpegBytes); });
+		PlayerSessions.ForEachSession([&JpegBytes](TSharedPtr<IPlayerSession> PlayerSession) { PlayerSession->SendFreezeFrame(JpegBytes); });
 
 		CachedJpegBytes = JpegBytes;
 	}
@@ -785,8 +781,7 @@ namespace UE::PixelStreaming
 	{
 		if (CachedJpegBytes.Num() > 0)
 		{
-			PlayerSessions.ForSession(PlayerId, [&](TSharedPtr<IPlayerSession> Session)
-									  { Session->SendFreezeFrame(CachedJpegBytes); });
+			PlayerSessions.ForSession(PlayerId, [&](TSharedPtr<IPlayerSession> Session) { Session->SendFreezeFrame(CachedJpegBytes); });
 		}
 	}
 
@@ -795,15 +790,13 @@ namespace UE::PixelStreaming
 		// Force a keyframe so when stream unfreezes if player has never received a h.264 frame before they can still connect.
 		ForceKeyFrame();
 
-		PlayerSessions.ForEachSession([](TSharedPtr<IPlayerSession> PlayerSession)
-									  { PlayerSession->SendUnfreezeFrame(); });
+		PlayerSessions.ForEachSession([](TSharedPtr<IPlayerSession> PlayerSession) { PlayerSession->SendUnfreezeFrame(); });
 
 		CachedJpegBytes.Empty();
 	}
-	void FStreamer::SendFileData(TArray<uint8> &ByteData, FString &MimeType, FString &FileExtension)
+	void FStreamer::SendFileData(TArray<uint8>& ByteData, FString& MimeType, FString& FileExtension)
 	{
-		PlayerSessions.ForEachSession([&](TSharedPtr<IPlayerSession> PlayerSession)
-									  { PlayerSession->SendFileData(ByteData, MimeType, FileExtension); });
+		PlayerSessions.ForEachSession([&](TSharedPtr<IPlayerSession> PlayerSession) { PlayerSession->SendFileData(ByteData, MimeType, FileExtension); });
 	}
 
 	void FStreamer::KickPlayer(FPixelStreamingPlayerId PlayerId)
