Fix comment out unsupported RTP parameters that was causing simulcast streaming to fail.

diff --git a/Engine/Plugins/Media/PixelStreaming/Source/PixelStreaming/Private/Streamer.cpp b/Engine/Plugins/Media/PixelStreaming/Source/PixelStreaming/Private/Streamer.cpp
--- a/Engine/Plugins/Media/PixelStreaming/Source/PixelStreaming/Private/Streamer.cpp
+++ b/Engine/Plugins/Media/PixelStreaming/Source/PixelStreaming/Private/Streamer.cpp
@@ -528,7 +528,11 @@ namespace UE::PixelStreaming
 				LayerEncoding.min_bitrate_bps = SimulcastLayer->MinBitrate;
 				LayerEncoding.max_bitrate_bps = SimulcastLayer->MaxBitrate;
 				LayerEncoding.scale_resolution_down_by = SimulcastLayer->Scaling;
-				LayerEncoding.network_priority = webrtc::Priority::kHigh;
+				
+				// In M84 this will crash with "Attempted to set an unimplemented parameter of RtpParameters".
+				// Try re-enabling this when we upgrade WebRTC versions.
+				//LayerEncoding.network_priority = webrtc::Priority::kHigh;
+				
 				LayerEncoding.max_framerate = FMath::Max(60, Settings::CVarPixelStreamingWebRTCFps.GetValueOnAnyThread());
 				EncodingParams.push_back(LayerEncoding);
 			}
