Fix a spelling error in the public facing Pixel Streaming delegates

diff --git a/Engine/Plugins/Media/PixelStreaming/Source/PixelStreaming/Private/SignallingServerConnection.cpp b/Engine/Plugins/Media/PixelStreaming/Source/PixelStreaming/Private/SignallingServerConnection.cpp
--- a/Engine/Plugins/Media/PixelStreaming/Source/PixelStreaming/Private/SignallingServerConnection.cpp
+++ b/Engine/Plugins/Media/PixelStreaming/Source/PixelStreaming/Private/SignallingServerConnection.cpp
@@ -255,8 +255,8 @@ namespace UE::PixelStreaming
 
 		if (UPixelStreamingDelegates* Delegates = UPixelStreamingDelegates::GetPixelStreamingDelegates())
 		{
-			Delegates->OnConnecedToSignallingServer.Broadcast();
-			Delegates->OnConnecedToSignallingServerNative.Broadcast();
+			Delegates->OnConnectedToSignallingServer.Broadcast();
+			Delegates->OnConnectedToSignallingServerNative.Broadcast();
 		}
 	}
 
diff --git a/Engine/Plugins/Media/PixelStreaming/Source/PixelStreaming/Public/PixelStreamingDelegates.h b/Engine/Plugins/Media/PixelStreaming/Source/PixelStreaming/Public/PixelStreamingDelegates.h
--- a/Engine/Plugins/Media/PixelStreaming/Source/PixelStreaming/Public/PixelStreamingDelegates.h
+++ b/Engine/Plugins/Media/PixelStreaming/Source/PixelStreaming/Public/PixelStreamingDelegates.h
@@ -25,10 +25,10 @@ public:
 	// BP Delegate
 	DECLARE_DYNAMIC_MULTICAST_DELEGATE(FConnectedToSignallingServer);
 	UPROPERTY(BlueprintAssignable, Category = "Pixel Streaming Delegates")
-	FConnectedToSignallingServer OnConnecedToSignallingServer;
+	FConnectedToSignallingServer OnConnectedToSignallingServer;
 	// C++ Delegate
 	DECLARE_MULTICAST_DELEGATE(FConnectedToSignallingServerNative);
-	FConnectedToSignallingServerNative OnConnecedToSignallingServerNative;
+	FConnectedToSignallingServerNative OnConnectedToSignallingServerNative;
 
 
 	/**
