fix rename TextureSource.cpp to PixelStreamingTextureSource.cpp

diff --git a/Engine/Plugins/Media/PixelStreaming/Source/PixelStreaming/Private/PixelStreamingTextureSource.cpp b/Engine/Plugins/Media/PixelStreaming/Source/PixelStreaming/Private/PixelStreamingTextureSource.cpp
new file mode 100644
--- /dev/null
+++ b/Engine/Plugins/Media/PixelStreaming/Source/PixelStreaming/Private/PixelStreamingTextureSource.cpp
@@ -0,0 +1,172 @@
+// Copyright Epic Games, Inc. All Rights Reserved.
+
+#include "PixelStreamingTextureSource.h"
+#include "Stats.h"
+#include "Utils.h"
+
+/**
+* @@!!
+* this is PixelStreamingTextureSource.cpp
+* the cpp is named wrong
+*/
+
+/*
+ * --------- IBackBufferTextureSource ------------------------------------------
+ */
+
+IBackBufferTextureSource::IBackBufferTextureSource(float InFrameScale)
+	: FrameScale(InFrameScale)
+{
+	// Explictly make clear we are adding another ref to this shared bool for the purposes of using in the below lambda
+	TSharedRef<bool, ESPMode::ThreadSafe> bEnabledClone = bEnabled;
+
+	// The backbuffer delegate can only be accessed using GameThread
+	AsyncTask(ENamedThreads::GameThread, [this, bEnabledClone]() {
+		/*Early exit if `this` died before game thread ran.*/
+		if (bEnabledClone.IsUnique())
+		{
+			return;
+		}
+		OnBackbuffer = &FSlateApplication::Get().GetRenderer()->OnBackBufferReadyToPresent();
+		BackbufferDelegateHandle = OnBackbuffer->AddSP(this, &IBackBufferTextureSource::OnBackBufferReady_RenderThread);
+	});
+}
+
+IBackBufferTextureSource::~IBackBufferTextureSource()
+{
+	if (OnBackbuffer)
+	{
+		OnBackbuffer->Remove(BackbufferDelegateHandle);
+	}
+
+	*bEnabled = false;
+}
+
+TRefCountPtr<IPixelStreamingBackBufferTextureWrapper> IBackBufferTextureSource::GetCurrent()
+{
+	if (bIsTempDirty)
+	{
+		FScopeLock Lock(&CriticalSection);
+		ReadBuffer.Swap(TempBuffer);
+		bIsTempDirty = false;
+	}
+	return ReadBuffer;
+}
+
+void IBackBufferTextureSource::SetEnabled(bool bInEnabled)
+{
+	*bEnabled = bInEnabled;
+	// This source has been disabled, so set `bInitialized` to false so `OnBackBufferReady_RenderThread`
+	// will make new textures next time it is called.
+	if (bInitialized && bInEnabled == false)
+	{
+		bInitialized = false;
+	}
+}
+
+void IBackBufferTextureSource::OnBackBufferReady_RenderThread(SWindow& SlateWindow, const FTexture2DRHIRef& FrameBuffer)
+{
+	if (!bInitialized)
+	{
+		Initialize(FrameBuffer->GetSizeXY().X * FrameScale, FrameBuffer->GetSizeXY().Y * FrameScale);
+	}
+
+	// Note: this is the logic that makes mid-stream resolution changing work.
+	if ((FrameBuffer->GetSizeXY().X * FrameScale != SourceWidth) || (FrameBuffer->GetSizeXY().Y * FrameScale != SourceHeight))
+	{
+		Initialize(FrameBuffer->GetSizeXY().X * FrameScale, FrameBuffer->GetSizeXY().Y * FrameScale);
+	}
+
+	if (!IsEnabled())
+	{
+		return;
+	}
+
+	auto& WriteBuffer = bWriteParity ? WriteBuffers[0] : WriteBuffers[1];
+	bWriteParity = !bWriteParity;
+
+	// for safety we just make sure that the buffer is not currently waiting for a copy
+	if (WriteBuffer.bAvailable)
+	{
+		WriteBuffer.bAvailable = false;
+
+		FRHICommandListImmediate& RHICmdList = FRHICommandListExecutor::GetImmediateCommandList();
+		WriteBuffer.Fence->Clear();
+
+		RHICmdList.EnqueueLambda([&WriteBuffer](FRHICommandListImmediate& RHICmdList) {
+			WriteBuffer.PreWaitingOnCopy = FPlatformTime::Cycles64();
+		});
+
+		UE::PixelStreaming::CopyTexture(FrameBuffer, WriteBuffer.TextureWrapper->GetTexture(), WriteBuffer.Fence);
+
+		IPixelStreamingModule::Get().AddGPUFencePollerTask(WriteBuffer.Fence, bEnabled, [this, &WriteBuffer]() {
+			// This lambda is called only once the GPUFence is done
+			{
+				FScopeLock Lock(&CriticalSection);
+				TempBuffer.Swap(WriteBuffer.TextureWrapper);
+				WriteBuffer.Fence->Clear();
+				WriteBuffer.bAvailable = true;
+
+				bIsTempDirty = true;
+			}
+
+			// For debugging timing information about the copy operation
+			// Turning it on all the time is a bit too much log spam if logging stats
+			uint64 PostWaitingOnCopy = FPlatformTime::Cycles64();
+			UE::PixelStreaming::FStats* Stats = UE::PixelStreaming::FStats::Get();
+			if (Stats)
+			{
+				double CaptureLatencyMs = FPlatformTime::ToMilliseconds64(PostWaitingOnCopy - WriteBuffer.PreWaitingOnCopy);
+				Stats->StoreApplicationStat(UE::PixelStreaming::FStatData(FName(*FString::Printf(TEXT("Layer (x%.2f) Capture time (ms)"), FrameScale)), CaptureLatencyMs, 2, true));
+			}
+		});
+	}
+}
+
+void IBackBufferTextureSource::Initialize(int Width, int Height)
+{
+	SourceWidth = Width;
+	SourceHeight = Height;
+
+	for (auto& Buffer : WriteBuffers)
+	{
+		Buffer.TextureWrapper = CreateTexture(SourceWidth, SourceHeight);
+		Buffer.Fence = GDynamicRHI->RHICreateGPUFence(TEXT("VideoCapturerCopyFence"));
+		Buffer.bAvailable = true;
+	}
+	bWriteParity = true;
+
+	TempBuffer = CreateTexture(SourceWidth, SourceHeight);
+	ReadBuffer = CreateTexture(SourceWidth, SourceHeight);
+	bIsTempDirty = false;
+
+	bInitialized = true;
+}
+
+/*
+ * -------- FTextureSourceBackBuffer ----------------
+ */
+void FBackBufferTextureSource::CopyTexture(const FTexture2DRHIRef& SourceTexture, TRefCountPtr<IPixelStreamingBackBufferTextureWrapper> DestTexture, FGPUFenceRHIRef& CopyFence)
+{
+	UE::PixelStreaming::CopyTexture(SourceTexture, DestTexture->GetTexture(), CopyFence);
+}
+
+TRefCountPtr<IPixelStreamingBackBufferTextureWrapper> FBackBufferTextureSource::CreateTexture(int Width, int Height)
+{
+	return new FPixelStreamingRHIBackBufferTexture(UE::PixelStreaming::CreateTexture(Width, Height));
+}
+
+/*
+ * -------- FBackBufferToCPUTextureSource ----------------
+ */
+
+void FBackBufferToCPUTextureSource::CopyTexture(const FTexture2DRHIRef& SourceTexture, TRefCountPtr<IPixelStreamingBackBufferTextureWrapper> DestTexture, FGPUFenceRHIRef& CopyFence)
+{
+	UE::PixelStreaming::CopyTexture(SourceTexture, DestTexture->GetTexture(), CopyFence);
+	UE::PixelStreaming::ReadTextureToCPU(FRHICommandListExecutor::GetImmediateCommandList(), DestTexture->GetTexture(), static_cast<FPixelStreamingCPUReadableBackbufferTexture*>(DestTexture.GetReference())->GetRawPixels());
+}
+
+TRefCountPtr<IPixelStreamingBackBufferTextureWrapper> FBackBufferToCPUTextureSource::CreateTexture(int Width, int Height)
+{
+	return new FPixelStreamingCPUReadableBackbufferTexture(UE::PixelStreaming::CreateTexture(Width, Height));
+}
diff --git a/Engine/Plugins/Media/PixelStreaming/Source/PixelStreaming/Private/TextureSource.cpp b/Engine/Plugins/Media/PixelStreaming/Source/PixelStreaming/Private/TextureSource.cpp
deleted file mode 100644
