This commit is a large collection (from various authors) of Pixel Streaming 5.1 work
that is useful to be released early onto Main.

--------------------------
Bugfixes
--------------------------

Fix update default max bitrate to support 4k
Fix resolution changing and fix match viewport resolution browser side option.
Fix stream not playable when streamer disconnects but page not refreshed.
Fix SFU/data-only peers were able to request/gain quality controller status.
Fix delegate registration for backbuffer texture sources would sometimes be called when deleted and crash.
Fix uncessary bundling of pixel streaming servers when building non-game targets such as Editor/Server.
Fix PixelStreaming plugin now works with "-onethread" enabled and some methods have been made more thread safe.

--------------------------
Refactors
--------------------------
Refactor texture source API surface to enable texture sources that do not use Unreal Engine texture types.
Refactor to make FixedFPSPump be more usable by any arbitrary thing that wants updates at a fixed interval (not just video sources).
Refactor to remove player id from video sources to make them easier to create outside of a true peer connection existing.
Refactor to remove unecessary CVarPixelStreamingWaitForTextureCopy.
Refactor to expose functionality on module `CreateExternalVideoSource`. This is experimental and will likely be changed before 5.1.

--------------------------
Features
--------------------------
Added foundation for desktop streaming on Pixel Streaming.
Added simulcast support for VP8.
Added a compute shader for faster I420 conversion when using VPX, this is faster than CPU readback and conversion in our tests.
Added some PixelStreaming tests
	1. A test that will start streaming, connect to a client, and check that a data channel message can be sent.
	2. A mock signalling server test and functionality to test the C++ connection to the signalling server.
Added "stat pixelstreaming" and made PixelStreaming.HUDStats able to be toggled on/off.
Added custom UE implementation webrtc::AdaptedVideoTrackSource that can add/remove remove video sinks at runtime.
Added ability on module to start/stop streaming as needed.
Added peers can now swap roles between P2P and simulcast mid-stream.
	To test this new feature:
	1. Simply toggle the PreferSFU toggle in the settings panel.
	2. Click the new restart stream button in the settings panel.

diff --git a/Engine/Plugins/Media/HardwareEncoders/Source/EncoderNVENC/Private/NVENC_EncoderH264.cpp b/Engine/Plugins/Media/HardwareEncoders/Source/EncoderNVENC/Private/NVENC_EncoderH264.cpp
--- a/Engine/Plugins/Media/HardwareEncoders/Source/EncoderNVENC/Private/NVENC_EncoderH264.cpp
+++ b/Engine/Plugins/Media/HardwareEncoders/Source/EncoderNVENC/Private/NVENC_EncoderH264.cpp
@@ -632,9 +632,8 @@ namespace AVEncoder
 		// This encode is single threaded and blocking, if someone else is encoding right now, we wait.
 		while (bIsProcessingFrame) 
 		{
-			// According to UE docs
-			// Sleep this thread for Seconds. 0.0 means release the current time slice to let other threads get some attention.
-			FPlatformProcess::Sleep(0.0f);
+			// Encoder is already processing a frame, don't try and encode another
+			return;
 		}
 
 		checkf(!bIsProcessingFrame, TEXT("NVENC encoder should only ever be called one thread at a time."));
diff --git a/Engine/Plugins/Media/PixelStreaming/PixelStreaming.uplugin b/Engine/Plugins/Media/PixelStreaming/PixelStreaming.uplugin
--- a/Engine/Plugins/Media/PixelStreaming/PixelStreaming.uplugin
+++ b/Engine/Plugins/Media/PixelStreaming/PixelStreaming.uplugin
@@ -21,7 +21,14 @@
 			"LoadingPhase": "PreDefault",
 			"PlatformAllowList": [ "Win64", "Linux" ],
 			"TargetDenyList": [ "Server" ]
-		}
+		},
+		{
+            "Name": "PixelStreamingShaders",
+            "Type": "Runtime",
+            "LoadingPhase": "PostConfigInit",
+            "PlatformAllowList": [ "Win64", "Linux" ],
+            "TargetDenyList": [ "Server" ]
+		},
 	],
 	"Plugins": [
 		{
diff --git a/Engine/Plugins/Media/PixelStreaming/Source/PixelStreaming/PixelStreaming.Build.cs b/Engine/Plugins/Media/PixelStreaming/Source/PixelStreaming/PixelStreaming.Build.cs
--- a/Engine/Plugins/Media/PixelStreaming/Source/PixelStreaming/PixelStreaming.Build.cs
+++ b/Engine/Plugins/Media/PixelStreaming/Source/PixelStreaming/PixelStreaming.Build.cs
@@ -24,14 +24,11 @@ namespace UnrealBuildTool.Rules
 			List<string> DependenciesToAdd = new List<string>();
 			DependenciesToAdd.AddRange(Directory.GetFiles(DirectoryToAdd, "*.*", SearchOption.AllDirectories));
 
-			string NodeModulesDirPath = new DirectoryInfo(DirectoryToAdd + "/node_modules").FullName;
-			string LogsDirPath = new DirectoryInfo(DirectoryToAdd + "/logs").FullName;
-			foreach (string Dependency in DependenciesToAdd)
+			foreach(var DependencySource in DependenciesToAdd)
 			{
-				if (!Dependency.StartsWith(NodeModulesDirPath) &&
-					!Dependency.StartsWith(LogsDirPath))
+				if ( ! (DependencySource.Contains("node") | DependencySource.Contains("logs")))
 				{
-					RuntimeDependencies.Add(Dependency, StagedFileType.NonUFS);
+					RuntimeDependencies.Add(DependencySource, StagedFileType.NonUFS);
 				}
 			}
 		}
@@ -47,24 +44,32 @@ namespace UnrealBuildTool.Rules
 			// This is so for game projects using our public headers don't have to include extra modules they might not know about.
 			PublicDependencyModuleNames.AddRange(new string[]
 			{
+				"ApplicationCore",
 				"InputDevice",
 				"WebRTC"
 			});
 
 			// NOTE: General rule is not to access the private folder of another module
 			PrivateIncludePaths.AddRange(new string[]
+			{
+				Path.Combine(EngineDir, "Source/Runtime/AudioMixer/Private"),
+			});
+
+			// WebRTC third party includes (just libyuv for colour format conversions for now)
+			PublicIncludePaths.AddRange(new string[]
 				{
-					Path.Combine(EngineDir, "Source/Runtime/AudioMixer/Private"),
+					Path.Combine(EngineDir, "Source/ThirdParty/WebRTC/4147/Include/third_party/libyuv/include"),
 				});
 
 			PrivateDependencyModuleNames.AddRange(new string[]
 			{
-				"ApplicationCore",
 				"Core",
 				"CoreUObject",
 				"Engine",
+				"EngineSettings",
 				"InputCore",
 				"Json",
+				"Renderer",
 				"RenderCore",
 				"RHI",
 				"Slate",
@@ -75,7 +80,8 @@ namespace UnrealBuildTool.Rules
 				"Sockets",
 				"MediaUtils",
 				"DeveloperSettings",
-				"AVEncoder"
+				"AVEncoder",
+				"PixelStreamingShaders"
 			});
 
 			PrivateDependencyModuleNames.Add("VulkanRHI");
@@ -89,9 +95,13 @@ namespace UnrealBuildTool.Rules
 				AddEngineThirdPartyPrivateStaticDependencies(Target, "DX11", "DX12");
 			}
 
-			AddFolder("SignallingWebServer");
-			AddFolder("Matchmaker");
-			AddFolder("SFU");
+			// When we build a Game target we also package the servers with it as runtime dependencies
+			if(Target.Type == TargetType.Game)
+			{
+				AddFolder("SignallingWebServer");
+				AddFolder("Matchmaker");
+				AddFolder("SFU");
+			}
 		}
 	}
 }
diff --git a/Engine/Plugins/Media/PixelStreaming/Source/PixelStreaming/Private/AudioCapturer.cpp b/Engine/Plugins/Media/PixelStreaming/Source/PixelStreaming/Private/AudioCapturer.cpp
--- a/Engine/Plugins/Media/PixelStreaming/Source/PixelStreaming/Private/AudioCapturer.cpp
+++ b/Engine/Plugins/Media/PixelStreaming/Source/PixelStreaming/Private/AudioCapturer.cpp
@@ -24,340 +24,343 @@ DEFINE_LOG_CATEGORY(LogPixelStreamingAudioCapturer);
 		};                       \
 	}
 
-constexpr int UE::PixelStreaming::FAudioCapturer::SampleRate;
-constexpr int UE::PixelStreaming::FAudioCapturer::NumChannels;
-
-void UE::PixelStreaming::FAudioCapturer::OnNewSubmixBuffer(const USoundSubmix* OwningSubmix, float* AudioData, int32 NumSamples, int32 InNumChannels, const int32 InSampleRate, double AudioClock)
+namespace UE::PixelStreaming
 {
-	if (!(bInitialized && bRecordingInitialized))
-	{
-		return;
-	}
+	constexpr int FAudioCapturer::SampleRate;
+	constexpr int FAudioCapturer::NumChannels;
 
-	// Only 48000hz supported for now
-	if (InSampleRate != SampleRate)
+	void FAudioCapturer::OnNewSubmixBuffer(const USoundSubmix* OwningSubmix, float* AudioData, int32 NumSamples, int32 InNumChannels, const int32 InSampleRate, double AudioClock)
 	{
-		// Only report the problem once
-		if (!bFormatChecked)
+		if (!(bInitialized && bRecordingInitialized))
 		{
-			bFormatChecked = true;
-			UE_LOG(LogPixelStreamingAudioCapturer, Error, TEXT("Audio samplerate needs to be 48000hz"));
+			return;
+		}
+
+		// Only 48000hz supported for now
+		if (InSampleRate != SampleRate)
+		{
+			// Only report the problem once
+			if (!bFormatChecked)
+			{
+				bFormatChecked = true;
+				UE_LOG(LogPixelStreamingAudioCapturer, Error, TEXT("Audio samplerate needs to be 48000hz"));
+			}
+			return;
 		}
-		return;
-	}
 
-	UE_LOG(LogPixelStreamingAudioCapturer, VeryVerbose, TEXT("captured %d samples, %dc, %dHz"), NumSamples, NumChannels, SampleRate);
+		UE_LOG(LogPixelStreamingAudioCapturer, VeryVerbose, TEXT("captured %d samples, %dc, %dHz"), NumSamples, NumChannels, SampleRate);
 
-	Audio::TSampleBuffer<float> Buffer(AudioData, NumSamples, InNumChannels, SampleRate);
-	// Mix to stereo if required, since PixelStreaming only accepts stereo at the moment
-	if (Buffer.GetNumChannels() != NumChannels)
-	{
-		Buffer.MixBufferToChannels(NumChannels);
+		Audio::TSampleBuffer<float> Buffer(AudioData, NumSamples, InNumChannels, SampleRate);
+		// Mix to stereo if required, since PixelStreaming only accepts stereo at the moment
+		if (Buffer.GetNumChannels() != NumChannels)
+		{
+			Buffer.MixBufferToChannels(NumChannels);
+		}
+
+		// Convert to signed PCM 16-bits
+		PCM16.Reset(Buffer.GetNumSamples());
+		PCM16.AddZeroed(Buffer.GetNumSamples());
+		const float* Ptr = reinterpret_cast<const float*>(Buffer.GetData());
+		for (int16& S : PCM16)
+		{
+			int32 N = *Ptr >= 0 ? *Ptr * int32(MAX_int16) : *Ptr * (int32(MAX_int16) + 1);
+			S = static_cast<int16>(FMath::Clamp(N, int32(MIN_int16), int32(MAX_int16)));
+			Ptr++;
+		}
+
+		RecordingBuffer.Append(reinterpret_cast<const uint8*>(PCM16.GetData()), PCM16.Num() * sizeof(PCM16[0]));
+		int BytesPer10Ms = (SampleRate * NumChannels * static_cast<int>(sizeof(uint16))) / 100;
+
+		// Feed in 10ms chunks
+		while (RecordingBuffer.Num() >= BytesPer10Ms)
+		{
+			{
+				FScopeLock Lock(&DeviceBufferCS);
+				if (DeviceBuffer)
+				{
+					DeviceBuffer->SetRecordedBuffer(RecordingBuffer.GetData(), BytesPer10Ms / (sizeof(uint16) * NumChannels));
+					DeviceBuffer->DeliverRecordedData();
+					UE_LOG(LogPixelStreamingAudioCapturer, VeryVerbose, TEXT("passed %d bytes"), BytesPer10Ms);
+				}
+			}
+
+			RecordingBuffer.RemoveAt(0, BytesPer10Ms, false);
+		}
 	}
 
-	// Convert to signed PCM 16-bits
-	PCM16.Reset(Buffer.GetNumSamples());
-	PCM16.AddZeroed(Buffer.GetNumSamples());
-	const float* Ptr = reinterpret_cast<const float*>(Buffer.GetData());
-	for (int16& S : PCM16)
+	int32 FAudioCapturer::ActiveAudioLayer(AudioLayer* audioLayer) const
 	{
-		int32 N = *Ptr >= 0 ? *Ptr * int32(MAX_int16) : *Ptr * (int32(MAX_int16) + 1);
-		S = static_cast<int16>(FMath::Clamp(N, int32(MIN_int16), int32(MAX_int16)));
-		Ptr++;
+		*audioLayer = AudioDeviceModule::kDummyAudio;
+		return 0;
 	}
 
-	RecordingBuffer.Append(reinterpret_cast<const uint8*>(PCM16.GetData()), PCM16.Num() * sizeof(PCM16[0]));
-	int BytesPer10Ms = (SampleRate * NumChannels * static_cast<int>(sizeof(uint16))) / 100;
+	int32 FAudioCapturer::RegisterAudioCallback(webrtc::AudioTransport* audioCallback)
+	{
+		FScopeLock Lock(&DeviceBufferCS);
+		DeviceBuffer->RegisterAudioCallback(audioCallback);
+
+		return 0;
+	}
 
-	// Feed in 10ms chunks
-	while (RecordingBuffer.Num() >= BytesPer10Ms)
+	int32 FAudioCapturer::Init()
 	{
+		if (bInitialized)
+			return 0;
+
 		{
 			FScopeLock Lock(&DeviceBufferCS);
-			if (DeviceBuffer)
-			{
-				DeviceBuffer->SetRecordedBuffer(RecordingBuffer.GetData(), BytesPer10Ms / (sizeof(uint16) * NumChannels));
-				DeviceBuffer->DeliverRecordedData();
-				UE_LOG(LogPixelStreamingAudioCapturer, VeryVerbose, TEXT("passed %d bytes"), BytesPer10Ms);
-			}
+
+			m_taskQueueFactory = webrtc::CreateDefaultTaskQueueFactory();
+			DeviceBuffer = MakeUnique<webrtc::AudioDeviceBuffer>(m_taskQueueFactory.get());
 		}
 
-		RecordingBuffer.RemoveAt(0, BytesPer10Ms, false);
-	}
-}
+		// subscribe to audio data
+		if (!GEngine)
+		{
+			return -1;
+		}
 
-int32 UE::PixelStreaming::FAudioCapturer::ActiveAudioLayer(AudioLayer* audioLayer) const
-{
-	*audioLayer = AudioDeviceModule::kDummyAudio;
-	return 0;
-}
+		FAudioDeviceHandle AudioDevice = GEngine->GetMainAudioDevice();
+		if (!AudioDevice)
+		{
+			UE_LOG(LogPixelStreamingAudioCapturer, Warning, TEXT("No audio device"));
+			return -1;
+		}
 
-int32 UE::PixelStreaming::FAudioCapturer::RegisterAudioCallback(webrtc::AudioTransport* audioCallback)
-{
-	FScopeLock Lock(&DeviceBufferCS);
-	DeviceBuffer->RegisterAudioCallback(audioCallback);
+		bInitialized = true;
+		AudioDevice->RegisterSubmixBufferListener(this);
 
-	return 0;
-}
+		UE_LOG(LogPixelStreamingAudioCapturer, Verbose, TEXT("Init"));
 
-int32 UE::PixelStreaming::FAudioCapturer::Init()
-{
-	if (bInitialized)
 		return 0;
+	}
 
+	int32 FAudioCapturer::Terminate()
 	{
-		FScopeLock Lock(&DeviceBufferCS);
+		if (!bInitialized)
+			return 0;
+
+		// unsubscribe from audio data
+		if (!GEngine)
+		{
+			return -1;
+		}
 
-		m_taskQueueFactory = webrtc::CreateDefaultTaskQueueFactory();
-		DeviceBuffer = MakeUnique<webrtc::AudioDeviceBuffer>(m_taskQueueFactory.get());
+		FAudioDeviceHandle AudioDevice = GEngine->GetMainAudioDevice();
+		if (!AudioDevice)
+		{
+			return -1;
+		}
+
+		AudioDevice->UnregisterSubmixBufferListener(this);
+		bInitialized = false;
+
+		{
+			FScopeLock Lock(&DeviceBufferCS);
+			DeviceBuffer.Reset();
+		}
+
+		UE_LOG(LogPixelStreamingAudioCapturer, Verbose, TEXT("Terminate"));
+
+		return 0;
 	}
 
-	// subscribe to audio data
-	if (!GEngine)
+	bool FAudioCapturer::Initialized() const
 	{
-		return -1;
+		return bInitialized;
 	}
 
-	FAudioDeviceHandle AudioDevice = GEngine->GetMainAudioDevice();
-	if (!AudioDevice)
+	int16 FAudioCapturer::PlayoutDevices()
 	{
-		UE_LOG(LogPixelStreamingAudioCapturer, Warning, TEXT("No audio device"));
+		CHECKinitialized_();
 		return -1;
 	}
 
-	bInitialized = true;
-	AudioDevice->RegisterSubmixBufferListener(this);
-
-	UE_LOG(LogPixelStreamingAudioCapturer, Verbose, TEXT("Init"));
-
-	return 0;
-}
-
-int32 UE::PixelStreaming::FAudioCapturer::Terminate()
-{
-	if (!bInitialized)
-		return 0;
-
-	// unsubscribe from audio data
-	if (!GEngine)
+	int16 FAudioCapturer::RecordingDevices()
 	{
+		CHECKinitialized_();
 		return -1;
 	}
 
-	FAudioDeviceHandle AudioDevice = GEngine->GetMainAudioDevice();
-	if (!AudioDevice)
+	int32 FAudioCapturer::PlayoutDeviceName(
+		uint16 index, char name[webrtc::kAdmMaxDeviceNameSize], char guid[webrtc::kAdmMaxGuidSize])
 	{
+		CHECKinitialized_();
 		return -1;
 	}
 
-	AudioDevice->UnregisterSubmixBufferListener(this);
-	bInitialized = false;
-
+	int32 FAudioCapturer::RecordingDeviceName(
+		uint16 index, char name[webrtc::kAdmMaxDeviceNameSize], char guid[webrtc::kAdmMaxGuidSize])
 	{
-		FScopeLock Lock(&DeviceBufferCS);
-		DeviceBuffer.Reset();
+		CHECKinitialized_();
+		return -1;
 	}
 
-	UE_LOG(LogPixelStreamingAudioCapturer, Verbose, TEXT("Terminate"));
-
-	return 0;
-}
-
-bool UE::PixelStreaming::FAudioCapturer::Initialized() const
-{
-	return bInitialized;
-}
-
-int16 UE::PixelStreaming::FAudioCapturer::PlayoutDevices()
-{
-	CHECKinitialized_();
-	return -1;
-}
-
-int16 UE::PixelStreaming::FAudioCapturer::RecordingDevices()
-{
-	CHECKinitialized_();
-	return -1;
-}
-
-int32 UE::PixelStreaming::FAudioCapturer::PlayoutDeviceName(
-	uint16 index, char name[webrtc::kAdmMaxDeviceNameSize], char guid[webrtc::kAdmMaxGuidSize])
-{
-	CHECKinitialized_();
-	return -1;
-}
+	int32 FAudioCapturer::SetPlayoutDevice(uint16 index)
+	{
+		CHECKinitialized_();
+		return 0;
+	}
 
-int32 UE::PixelStreaming::FAudioCapturer::RecordingDeviceName(
-	uint16 index, char name[webrtc::kAdmMaxDeviceNameSize], char guid[webrtc::kAdmMaxGuidSize])
-{
-	CHECKinitialized_();
-	return -1;
-}
+	int32 FAudioCapturer::SetPlayoutDevice(WindowsDeviceType device)
+	{
+		CHECKinitialized_();
+		return 0;
+	}
 
-int32 UE::PixelStreaming::FAudioCapturer::SetPlayoutDevice(uint16 index)
-{
-	CHECKinitialized_();
-	return 0;
-}
+	int32 FAudioCapturer::SetRecordingDevice(uint16 index)
+	{
+		CHECKinitialized_();
+		return 0;
+	}
 
-int32 UE::PixelStreaming::FAudioCapturer::SetPlayoutDevice(WindowsDeviceType device)
-{
-	CHECKinitialized_();
-	return 0;
-}
+	int32 FAudioCapturer::SetRecordingDevice(WindowsDeviceType device)
+	{
+		CHECKinitialized_();
+		return 0;
+	}
 
-int32 UE::PixelStreaming::FAudioCapturer::SetRecordingDevice(uint16 index)
-{
-	CHECKinitialized_();
-	return 0;
-}
+	int32 FAudioCapturer::PlayoutIsAvailable(bool* available)
+	{
+		CHECKinitialized_();
+		return -1;
+	}
 
-int32 UE::PixelStreaming::FAudioCapturer::SetRecordingDevice(WindowsDeviceType device)
-{
-	CHECKinitialized_();
-	return 0;
-}
+	int32 FAudioCapturer::InitPlayout()
+	{
+		CHECKinitialized_();
+		return -1;
+	}
 
-int32 UE::PixelStreaming::FAudioCapturer::PlayoutIsAvailable(bool* available)
-{
-	CHECKinitialized_();
-	return -1;
-}
+	bool FAudioCapturer::PlayoutIsInitialized() const
+	{
+		CHECKinitialized__BOOL();
+		return false;
+	}
 
-int32 UE::PixelStreaming::FAudioCapturer::InitPlayout()
-{
-	CHECKinitialized_();
-	return -1;
-}
+	int32 FAudioCapturer::RecordingIsAvailable(bool* available)
+	{
+		CHECKinitialized_();
+		return -1;
+	}
 
-bool UE::PixelStreaming::FAudioCapturer::PlayoutIsInitialized() const
-{
-	CHECKinitialized__BOOL();
-	return false;
-}
+	int32 FAudioCapturer::InitRecording()
+	{
+		CHECKinitialized_();
 
-int32 UE::PixelStreaming::FAudioCapturer::RecordingIsAvailable(bool* available)
-{
-	CHECKinitialized_();
-	return -1;
-}
+		{
+			FScopeLock Lock(&DeviceBufferCS);
+			// #Audio : Allow dynamic values for samplerate and/or channels ,
+			// or receive those from UnrealEngine ?
+			DeviceBuffer->SetRecordingSampleRate(SampleRate);
+			DeviceBuffer->SetRecordingChannels(NumChannels);
+		}
 
-int32 UE::PixelStreaming::FAudioCapturer::InitRecording()
-{
-	CHECKinitialized_();
+		bRecordingInitialized = true;
+		return 0;
+	}
 
+	bool FAudioCapturer::RecordingIsInitialized() const
 	{
-		FScopeLock Lock(&DeviceBufferCS);
-		// #Audio : Allow dynamic values for samplerate and/or channels ,
-		// or receive those from UnrealEngine ?
-		DeviceBuffer->SetRecordingSampleRate(SampleRate);
-		DeviceBuffer->SetRecordingChannels(NumChannels);
+		CHECKinitialized__BOOL();
+		return bRecordingInitialized == true;
 	}
 
-	bRecordingInitialized = true;
-	return 0;
-}
-
-bool UE::PixelStreaming::FAudioCapturer::RecordingIsInitialized() const
-{
-	CHECKinitialized__BOOL();
-	return bRecordingInitialized == true;
-}
-
-int32 UE::PixelStreaming::FAudioCapturer::StartPlayout()
-{
-	CHECKinitialized_();
-	return -1;
-}
+	int32 FAudioCapturer::StartPlayout()
+	{
+		CHECKinitialized_();
+		return -1;
+	}
 
-int32 UE::PixelStreaming::FAudioCapturer::StopPlayout()
-{
-	CHECKinitialized_();
-	return -1;
-}
+	int32 FAudioCapturer::StopPlayout()
+	{
+		CHECKinitialized_();
+		return -1;
+	}
 
-bool UE::PixelStreaming::FAudioCapturer::Playing() const
-{
-	CHECKinitialized__BOOL();
-	return false;
-}
+	bool FAudioCapturer::Playing() const
+	{
+		CHECKinitialized__BOOL();
+		return false;
+	}
 
-int32 UE::PixelStreaming::FAudioCapturer::StartRecording()
-{
-	CHECKinitialized_();
-	return -1;
-}
+	int32 FAudioCapturer::StartRecording()
+	{
+		CHECKinitialized_();
+		return -1;
+	}
 
-int32 UE::PixelStreaming::FAudioCapturer::StopRecording()
-{
-	CHECKinitialized_();
-	return -1;
-}
+	int32 FAudioCapturer::StopRecording()
+	{
+		CHECKinitialized_();
+		return -1;
+	}
 
-bool UE::PixelStreaming::FAudioCapturer::Recording() const
-{
-	CHECKinitialized__BOOL();
-	return bRecordingInitialized;
-}
+	bool FAudioCapturer::Recording() const
+	{
+		CHECKinitialized__BOOL();
+		return bRecordingInitialized;
+	}
 
-int32 UE::PixelStreaming::FAudioCapturer::InitSpeaker()
-{
-	CHECKinitialized_();
-	return -1;
-}
+	int32 FAudioCapturer::InitSpeaker()
+	{
+		CHECKinitialized_();
+		return -1;
+	}
 
-bool UE::PixelStreaming::FAudioCapturer::SpeakerIsInitialized() const
-{
-	CHECKinitialized__BOOL();
-	return false;
-}
+	bool FAudioCapturer::SpeakerIsInitialized() const
+	{
+		CHECKinitialized__BOOL();
+		return false;
+	}
 
-int32 UE::PixelStreaming::FAudioCapturer::InitMicrophone()
-{
-	CHECKinitialized_();
-	return 0;
-}
+	int32 FAudioCapturer::InitMicrophone()
+	{
+		CHECKinitialized_();
+		return 0;
+	}
 
-bool UE::PixelStreaming::FAudioCapturer::MicrophoneIsInitialized() const
-{
-	CHECKinitialized__BOOL();
-	return true;
-}
+	bool FAudioCapturer::MicrophoneIsInitialized() const
+	{
+		CHECKinitialized__BOOL();
+		return true;
+	}
 
-int32 UE::PixelStreaming::FAudioCapturer::StereoPlayoutIsAvailable(bool* available) const
-{
-	CHECKinitialized_();
-	return -1;
-}
+	int32 FAudioCapturer::StereoPlayoutIsAvailable(bool* available) const
+	{
+		CHECKinitialized_();
+		return -1;
+	}
 
-int32 UE::PixelStreaming::FAudioCapturer::SetStereoPlayout(bool enable)
-{
-	CHECKinitialized_();
-	return -1;
-}
+	int32 FAudioCapturer::SetStereoPlayout(bool enable)
+	{
+		CHECKinitialized_();
+		return -1;
+	}
 
-int32 UE::PixelStreaming::FAudioCapturer::StereoPlayout(bool* enabled) const
-{
-	CHECKinitialized_();
-	return -1;
-}
+	int32 FAudioCapturer::StereoPlayout(bool* enabled) const
+	{
+		CHECKinitialized_();
+		return -1;
+	}
 
-int32 UE::PixelStreaming::FAudioCapturer::StereoRecordingIsAvailable(bool* available) const
-{
-	CHECKinitialized_();
-	*available = true;
-	return 0;
-}
+	int32 FAudioCapturer::StereoRecordingIsAvailable(bool* available) const
+	{
+		CHECKinitialized_();
+		*available = true;
+		return 0;
+	}
 
-int32 UE::PixelStreaming::FAudioCapturer::SetStereoRecording(bool enable)
-{
-	CHECKinitialized_();
-	return 0;
-}
+	int32 FAudioCapturer::SetStereoRecording(bool enable)
+	{
+		CHECKinitialized_();
+		return 0;
+	}
 
-int32 UE::PixelStreaming::FAudioCapturer::StereoRecording(bool* enabled) const
-{
-	CHECKinitialized_();
-	*enabled = true;
-	return 0;
-}
+	int32 FAudioCapturer::StereoRecording(bool* enabled) const
+	{
+		CHECKinitialized_();
+		*enabled = true;
+		return 0;
+	}
+} // namespace UE::PixelStreaming
\ No newline at end of file
diff --git a/Engine/Plugins/Media/PixelStreaming/Source/PixelStreaming/Private/AudioCapturer.h b/Engine/Plugins/Media/PixelStreaming/Source/PixelStreaming/Private/AudioCapturer.h
--- a/Engine/Plugins/Media/PixelStreaming/Source/PixelStreaming/Private/AudioCapturer.h
+++ b/Engine/Plugins/Media/PixelStreaming/Source/PixelStreaming/Private/AudioCapturer.h
@@ -6,191 +6,190 @@
 #include "AudioDevice.h"
 #include "HAL/CriticalSection.h"
 
-namespace UE {
-	namespace PixelStreaming {
-		// This class is only used if -PixelStreamingWebRTCUseLegacyAudioDevice or related CVar are used.
-		// This class will likely be removed in the future once the new ADM is confirmed stable.
-		class FAudioCapturer : public webrtc::AudioDeviceModule, public ISubmixBufferListener
-		{
-		private:
-			// ISubmixBufferListener interface
-			void OnNewSubmixBuffer(const USoundSubmix* OwningSubmix, float* AudioData, int32 NumSamples, int32 NumChannels, const int32 SampleRate, double AudioClock) override;
-
-			//
-			// webrtc::AudioDeviceModule interface
-			//
-			int32 ActiveAudioLayer(AudioLayer* audioLayer) const override;
-			int32 RegisterAudioCallback(webrtc::AudioTransport* audioCallback) override;
-
-			// Main initialization and termination
-			int32 Init() override;
-			int32 Terminate() override;
-			bool Initialized() const override;
-
-			// Device enumeration
-			int16 PlayoutDevices() override;
-			int16 RecordingDevices() override;
-			int32 PlayoutDeviceName(uint16 index, char name[webrtc::kAdmMaxDeviceNameSize], char guid[webrtc::kAdmMaxGuidSize]) override;
-			int32 RecordingDeviceName(uint16 index, char name[webrtc::kAdmMaxDeviceNameSize], char guid[webrtc::kAdmMaxGuidSize]) override;
-
-			// Device selection
-			int32 SetPlayoutDevice(uint16 index) override;
-			int32 SetPlayoutDevice(WindowsDeviceType device) override;
-			int32 SetRecordingDevice(uint16 index) override;
-			int32 SetRecordingDevice(WindowsDeviceType device) override;
-
-			// Audio transport initialization
-			int32 PlayoutIsAvailable(bool* available) override;
-			int32 InitPlayout() override;
-			bool PlayoutIsInitialized() const override;
-			int32 RecordingIsAvailable(bool* available) override;
-			int32 InitRecording() override;
-			bool RecordingIsInitialized() const override;
-
-			// Audio transport control
-			virtual int32 StartPlayout() override;
-			virtual int32 StopPlayout() override;
-			virtual bool Playing() const override;
-			virtual int32 StartRecording() override;
-			virtual int32 StopRecording() override;
-			virtual bool Recording() const override;
-
-			// Audio mixer initialization
-			virtual int32 InitSpeaker() override;
-			virtual bool SpeakerIsInitialized() const override;
-			virtual int32 InitMicrophone() override;
-			virtual bool MicrophoneIsInitialized() const override;
-
-			// Speaker volume controls
-			virtual int32 SpeakerVolumeIsAvailable(bool* available) override
-			{
-				return -1;
-			}
-			virtual int32 SetSpeakerVolume(uint32 volume) override
-			{
-				return -1;
-			}
-			virtual int32 SpeakerVolume(uint32* volume) const override
-			{
-				return -1;
-			}
-			virtual int32 MaxSpeakerVolume(uint32* maxVolume) const override
-			{
-				return -1;
-			}
-			virtual int32 MinSpeakerVolume(uint32* minVolume) const override
-			{
-				return -1;
-			}
-
-			// Microphone volume controls
-			virtual int32 MicrophoneVolumeIsAvailable(bool* available) override
-			{
-				return -1;
-			}
-			virtual int32 SetMicrophoneVolume(uint32 volume) override
-			{
-				return -1;
-			}
-			virtual int32 MicrophoneVolume(uint32* volume) const override
-			{
-				return -1;
-			}
-			virtual int32 MaxMicrophoneVolume(uint32* maxVolume) const override
-			{
-				return -1;
-			}
-			virtual int32 MinMicrophoneVolume(uint32* minVolume) const override
-			{
-				return -1;
-			}
-
-			// Speaker mute control
-			virtual int32 SpeakerMuteIsAvailable(bool* available) override
-			{
-				return -1;
-			}
-			virtual int32 SetSpeakerMute(bool enable) override
-			{
-				return -1;
-			}
-			virtual int32 SpeakerMute(bool* enabled) const override
-			{
-				return -1;
-			}
-
-			// Microphone mute control
-			virtual int32 MicrophoneMuteIsAvailable(bool* available) override
-			{
-				return -1;
-			}
-			virtual int32 SetMicrophoneMute(bool enable) override
-			{
-				return -1;
-			}
-			virtual int32 MicrophoneMute(bool* enabled) const override
-			{
-				return -1;
-			}
-
-			// Stereo support
-			virtual int32 StereoPlayoutIsAvailable(bool* available) const override;
-			virtual int32 SetStereoPlayout(bool enable) override;
-			virtual int32 StereoPlayout(bool* enabled) const override;
-			virtual int32 StereoRecordingIsAvailable(bool* available) const override;
-			virtual int32 SetStereoRecording(bool enable) override;
-			virtual int32 StereoRecording(bool* enabled) const override;
-
-			// Playout delay
-			virtual int32 PlayoutDelay(uint16* delayMS) const override
-			{
-				return -1;
-			}
-
-			// Only supported on Android.
-			virtual bool BuiltInAECIsAvailable() const override
-			{
-				return false;
-			}
-			virtual bool BuiltInAGCIsAvailable() const override
-			{
-				return false;
-			}
-			virtual bool BuiltInNSIsAvailable() const override
-			{
-				return false;
-			}
-
-			// Enables the built-in audio effects. Only supported on Android.
-			virtual int32 EnableBuiltInAEC(bool enable) override
-			{
-				return -1;
-			}
-			virtual int32 EnableBuiltInAGC(bool enable) override
-			{
-				return -1;
-			}
-			virtual int32 EnableBuiltInNS(bool enable) override
-			{
-				return -1;
-			}
-
-			FThreadSafeBool bInitialized = false;
-			TArray<int16> PCM16;
-
-			TUniquePtr<webrtc::AudioDeviceBuffer> DeviceBuffer;
-			FCriticalSection DeviceBufferCS; // protects capturing audio and sending it to WebRTC from simultaneous termination
-
-			TArray<uint8_t> RecordingBuffer;
-			int RecordingBufferSize = 0;
-
-			FThreadSafeBool bRecordingInitialized = false;
-			static constexpr int SampleRate = 48000;
-			static constexpr int NumChannels = 2;
-
-			bool bFormatChecked = false;
-
-			std::unique_ptr<webrtc::TaskQueueFactory> m_taskQueueFactory;
-		};
-	}
-}
+namespace UE::PixelStreaming 
+{
+	// This class is only used if -PixelStreamingWebRTCUseLegacyAudioDevice or related CVar are used.
+	// This class will likely be removed in the future once the new ADM is confirmed stable.
+	class FAudioCapturer : public webrtc::AudioDeviceModule, public ISubmixBufferListener
+	{
+	private:
+		// ISubmixBufferListener interface
+		void OnNewSubmixBuffer(const USoundSubmix* OwningSubmix, float* AudioData, int32 NumSamples, int32 NumChannels, const int32 SampleRate, double AudioClock) override;
+
+		//
+		// webrtc::AudioDeviceModule interface
+		//
+		int32 ActiveAudioLayer(AudioLayer* audioLayer) const override;
+		int32 RegisterAudioCallback(webrtc::AudioTransport* audioCallback) override;
+
+		// Main initialization and termination
+		int32 Init() override;
+		int32 Terminate() override;
+		bool Initialized() const override;
+
+		// Device enumeration
+		int16 PlayoutDevices() override;
+		int16 RecordingDevices() override;
+		int32 PlayoutDeviceName(uint16 index, char name[webrtc::kAdmMaxDeviceNameSize], char guid[webrtc::kAdmMaxGuidSize]) override;
+		int32 RecordingDeviceName(uint16 index, char name[webrtc::kAdmMaxDeviceNameSize], char guid[webrtc::kAdmMaxGuidSize]) override;
+
+		// Device selection
+		int32 SetPlayoutDevice(uint16 index) override;
+		int32 SetPlayoutDevice(WindowsDeviceType device) override;
+		int32 SetRecordingDevice(uint16 index) override;
+		int32 SetRecordingDevice(WindowsDeviceType device) override;
+
+		// Audio transport initialization
+		int32 PlayoutIsAvailable(bool* available) override;
+		int32 InitPlayout() override;
+		bool PlayoutIsInitialized() const override;
+		int32 RecordingIsAvailable(bool* available) override;
+		int32 InitRecording() override;
+		bool RecordingIsInitialized() const override;
+
+		// Audio transport control
+		virtual int32 StartPlayout() override;
+		virtual int32 StopPlayout() override;
+		virtual bool Playing() const override;
+		virtual int32 StartRecording() override;
+		virtual int32 StopRecording() override;
+		virtual bool Recording() const override;
+
+		// Audio mixer initialization
+		virtual int32 InitSpeaker() override;
+		virtual bool SpeakerIsInitialized() const override;
+		virtual int32 InitMicrophone() override;
+		virtual bool MicrophoneIsInitialized() const override;
+
+		// Speaker volume controls
+		virtual int32 SpeakerVolumeIsAvailable(bool* available) override
+		{
+			return -1;
+		}
+		virtual int32 SetSpeakerVolume(uint32 volume) override
+		{
+			return -1;
+		}
+		virtual int32 SpeakerVolume(uint32* volume) const override
+		{
+			return -1;
+		}
+		virtual int32 MaxSpeakerVolume(uint32* maxVolume) const override
+		{
+			return -1;
+		}
+		virtual int32 MinSpeakerVolume(uint32* minVolume) const override
+		{
+			return -1;
+		}
+
+		// Microphone volume controls
+		virtual int32 MicrophoneVolumeIsAvailable(bool* available) override
+		{
+			return -1;
+		}
+		virtual int32 SetMicrophoneVolume(uint32 volume) override
+		{
+			return -1;
+		}
+		virtual int32 MicrophoneVolume(uint32* volume) const override
+		{
+			return -1;
+		}
+		virtual int32 MaxMicrophoneVolume(uint32* maxVolume) const override
+		{
+			return -1;
+		}
+		virtual int32 MinMicrophoneVolume(uint32* minVolume) const override
+		{
+			return -1;
+		}
+
+		// Speaker mute control
+		virtual int32 SpeakerMuteIsAvailable(bool* available) override
+		{
+			return -1;
+		}
+		virtual int32 SetSpeakerMute(bool enable) override
+		{
+			return -1;
+		}
+		virtual int32 SpeakerMute(bool* enabled) const override
+		{
+			return -1;
+		}
+
+		// Microphone mute control
+		virtual int32 MicrophoneMuteIsAvailable(bool* available) override
+		{
+			return -1;
+		}
+		virtual int32 SetMicrophoneMute(bool enable) override
+		{
+			return -1;
+		}
+		virtual int32 MicrophoneMute(bool* enabled) const override
+		{
+			return -1;
+		}
+
+		// Stereo support
+		virtual int32 StereoPlayoutIsAvailable(bool* available) const override;
+		virtual int32 SetStereoPlayout(bool enable) override;
+		virtual int32 StereoPlayout(bool* enabled) const override;
+		virtual int32 StereoRecordingIsAvailable(bool* available) const override;
+		virtual int32 SetStereoRecording(bool enable) override;
+		virtual int32 StereoRecording(bool* enabled) const override;
+
+		// Playout delay
+		virtual int32 PlayoutDelay(uint16* delayMS) const override
+		{
+			return -1;
+		}
+
+		// Only supported on Android.
+		virtual bool BuiltInAECIsAvailable() const override
+		{
+			return false;
+		}
+		virtual bool BuiltInAGCIsAvailable() const override
+		{
+			return false;
+		}
+		virtual bool BuiltInNSIsAvailable() const override
+		{
+			return false;
+		}
+
+		// Enables the built-in audio effects. Only supported on Android.
+		virtual int32 EnableBuiltInAEC(bool enable) override
+		{
+			return -1;
+		}
+		virtual int32 EnableBuiltInAGC(bool enable) override
+		{
+			return -1;
+		}
+		virtual int32 EnableBuiltInNS(bool enable) override
+		{
+			return -1;
+		}
+
+		FThreadSafeBool bInitialized = false;
+		TArray<int16> PCM16;
+
+		TUniquePtr<webrtc::AudioDeviceBuffer> DeviceBuffer;
+		FCriticalSection DeviceBufferCS; // protects capturing audio and sending it to WebRTC from simultaneous termination
+
+		TArray<uint8_t> RecordingBuffer;
+		int RecordingBufferSize = 0;
+
+		FThreadSafeBool bRecordingInitialized = false;
+		static constexpr int SampleRate = 48000;
+		static constexpr int NumChannels = 2;
+
+		bool bFormatChecked = false;
+
+		std::unique_ptr<webrtc::TaskQueueFactory> m_taskQueueFactory;
+	};
+} // namespace UE::PixelStreaming
\ No newline at end of file
diff --git a/Engine/Plugins/Media/PixelStreaming/Source/PixelStreaming/Private/AudioDeviceModule.cpp b/Engine/Plugins/Media/PixelStreaming/Source/PixelStreaming/Private/AudioDeviceModule.cpp
--- a/Engine/Plugins/Media/PixelStreaming/Source/PixelStreaming/Private/AudioDeviceModule.cpp
+++ b/Engine/Plugins/Media/PixelStreaming/Source/PixelStreaming/Private/AudioDeviceModule.cpp
@@ -9,15 +9,18 @@
 DECLARE_LOG_CATEGORY_EXTERN(LogPixelStreamingAudioDeviceModule, Log, All);
 DEFINE_LOG_CATEGORY(LogPixelStreamingAudioDeviceModule);
 
-// These are copied from webrtc internals
-#define CHECKinitialized_() \
+
+namespace UE::PixelStreaming
+{
+	// These are copied from webrtc internals
+	#define CHECKinitialized_() \
 	{                       \
 		if (!bInitialized)  \
 		{                   \
 			return -1;      \
 		};                  \
 	}
-#define CHECKinitialized__BOOL() \
+	#define CHECKinitialized__BOOL() \
 	{                            \
 		if (!bInitialized)       \
 		{                        \
@@ -25,589 +28,590 @@ DEFINE_LOG_CATEGORY(LogPixelStreamingAudioDeviceModule);
 		};                       \
 	}
 
-int32 UE::PixelStreaming::FAudioDeviceModule::ActiveAudioLayer(AudioLayer* audioLayer) const
-{
-	*audioLayer = AudioDeviceModule::kDummyAudio;
-	return 0;
-}
-
-int32 UE::PixelStreaming::FAudioDeviceModule::RegisterAudioCallback(webrtc::AudioTransport* audioCallback)
-{
-	Requester.RegisterAudioCallback(audioCallback);
-	Capturer.RegisterAudioCallback(audioCallback);
-	return 0;
-}
-
-int32 UE::PixelStreaming::FAudioDeviceModule::Init()
-{
-	InitRecording();
-
-	UE_LOG(LogPixelStreamingAudioDeviceModule, Verbose, TEXT("Init PixelStreamingAudioDeviceModule"));
-	bInitialized = true;
-	return 0;
-}
-
-int32 UE::PixelStreaming::FAudioDeviceModule::Terminate()
-{
-	if (!bInitialized)
+	int32 FAudioDeviceModule::ActiveAudioLayer(AudioLayer* audioLayer) const
 	{
+		*audioLayer = AudioDeviceModule::kDummyAudio;
 		return 0;
 	}
 
-	Requester.Uninitialise();
-	Capturer.Uninitialise();
-
-	UE_LOG(LogPixelStreamingAudioDeviceModule, Verbose, TEXT("Terminate"));
-
-	return 0;
-}
-
-bool UE::PixelStreaming::FAudioDeviceModule::Initialized() const
-{
-	return bInitialized;
-}
-
-int16 UE::PixelStreaming::FAudioDeviceModule::PlayoutDevices()
-{
-	CHECKinitialized_();
-	return -1;
-}
-
-int16 UE::PixelStreaming::FAudioDeviceModule::RecordingDevices()
-{
-	CHECKinitialized_();
-	return -1;
-}
+	int32 FAudioDeviceModule::RegisterAudioCallback(webrtc::AudioTransport* audioCallback)
+	{
+		Requester.RegisterAudioCallback(audioCallback);
+		Capturer.RegisterAudioCallback(audioCallback);
+		return 0;
+	}
 
-int32 UE::PixelStreaming::FAudioDeviceModule::PlayoutDeviceName(uint16 index, char name[webrtc::kAdmMaxDeviceNameSize], char guid[webrtc::kAdmMaxGuidSize])
-{
-	CHECKinitialized_();
-	return -1;
-}
+	int32 FAudioDeviceModule::Init()
+	{
+		InitRecording();
 
-int32 UE::PixelStreaming::FAudioDeviceModule::RecordingDeviceName(uint16 index, char name[webrtc::kAdmMaxDeviceNameSize], char guid[webrtc::kAdmMaxGuidSize])
-{
-	CHECKinitialized_();
-	return -1;
-}
+		UE_LOG(LogPixelStreamingAudioDeviceModule, Verbose, TEXT("Init PixelStreamingAudioDeviceModule"));
+		bInitialized = true;
+		return 0;
+	}
 
-int32 UE::PixelStreaming::FAudioDeviceModule::SetPlayoutDevice(uint16 index)
-{
-	CHECKinitialized_();
-	return 0;
-}
+	int32 FAudioDeviceModule::Terminate()
+	{
+		if (!bInitialized)
+		{
+			return 0;
+		}
 
-int32 UE::PixelStreaming::FAudioDeviceModule::SetPlayoutDevice(WindowsDeviceType device)
-{
-	CHECKinitialized_();
-	return 0;
-}
+		Requester.Uninitialise();
+		Capturer.Uninitialise();
 
-int32 UE::PixelStreaming::FAudioDeviceModule::SetRecordingDevice(uint16 index)
-{
-	CHECKinitialized_();
-	return 0;
-}
+		UE_LOG(LogPixelStreamingAudioDeviceModule, Verbose, TEXT("Terminate"));
 
-int32 UE::PixelStreaming::FAudioDeviceModule::SetRecordingDevice(WindowsDeviceType device)
-{
-	CHECKinitialized_();
-	return 0;
-}
+		return 0;
+	}
 
-int32 UE::PixelStreaming::FAudioDeviceModule::PlayoutIsAvailable(bool* available)
-{
-	CHECKinitialized_();
-	*available = !UE::PixelStreaming::Settings::CVarPixelStreamingWebRTCDisableReceiveAudio.GetValueOnAnyThread();
-	return 0;
-}
+	bool FAudioDeviceModule::Initialized() const
+	{
+		return bInitialized;
+	}
 
-int32 UE::PixelStreaming::FAudioDeviceModule::InitPlayout()
-{
-	CHECKinitialized_();
-	bool bIsPlayoutAvailable = false;
-	PlayoutIsAvailable(&bIsPlayoutAvailable);
-	if (!bIsPlayoutAvailable)
+	int16 FAudioDeviceModule::PlayoutDevices()
 	{
+		CHECKinitialized_();
 		return -1;
 	}
 
-	Requester.InitPlayout();
-	return 0;
-}
-
-bool UE::PixelStreaming::FAudioDeviceModule::PlayoutIsInitialized() const
-{
-	CHECKinitialized__BOOL();
-	return Requester.PlayoutIsInitialized();
-}
-
-int32 UE::PixelStreaming::FAudioDeviceModule::StartPlayout()
-{
-	CHECKinitialized_();
-
-	if (!Requester.PlayoutIsInitialized())
+	int16 FAudioDeviceModule::RecordingDevices()
 	{
+		CHECKinitialized_();
 		return -1;
 	}
 
-	Requester.StartPlayout();
-	return 0;
-}
-
-int32 UE::PixelStreaming::FAudioDeviceModule::StopPlayout()
-{
-	CHECKinitialized_();
-	if (!Requester.PlayoutIsInitialized())
+	int32 FAudioDeviceModule::PlayoutDeviceName(uint16 index, char name[webrtc::kAdmMaxDeviceNameSize], char guid[webrtc::kAdmMaxGuidSize])
 	{
+		CHECKinitialized_();
 		return -1;
 	}
 
-	Requester.StopPlayout();
-	return 0;
-}
+	int32 FAudioDeviceModule::RecordingDeviceName(uint16 index, char name[webrtc::kAdmMaxDeviceNameSize], char guid[webrtc::kAdmMaxGuidSize])
+	{
+		CHECKinitialized_();
+		return -1;
+	}
 
-bool UE::PixelStreaming::FAudioDeviceModule::Playing() const
-{
-	CHECKinitialized__BOOL();
-	return Requester.Playing();
-}
+	int32 FAudioDeviceModule::SetPlayoutDevice(uint16 index)
+	{
+		CHECKinitialized_();
+		return 0;
+	}
 
-int32 UE::PixelStreaming::FAudioDeviceModule::RecordingIsAvailable(bool* available)
-{
-	CHECKinitialized_();
-	*available = !UE::PixelStreaming::Settings::CVarPixelStreamingWebRTCDisableTransmitAudio.GetValueOnAnyThread();
-	return 0;
-}
+	int32 FAudioDeviceModule::SetPlayoutDevice(WindowsDeviceType device)
+	{
+		CHECKinitialized_();
+		return 0;
+	}
 
-int32 UE::PixelStreaming::FAudioDeviceModule::InitRecording()
-{
-	CHECKinitialized_();
+	int32 FAudioDeviceModule::SetRecordingDevice(uint16 index)
+	{
+		CHECKinitialized_();
+		return 0;
+	}
 
-	bool bIsRecordingAvailable = false;
-	RecordingIsAvailable(&bIsRecordingAvailable);
-	if (!bIsRecordingAvailable)
+	int32 FAudioDeviceModule::SetRecordingDevice(WindowsDeviceType device)
 	{
-		return -1;
+		CHECKinitialized_();
+		return 0;
 	}
 
-	if (!Capturer.IsInitialised())
+	int32 FAudioDeviceModule::PlayoutIsAvailable(bool* available)
 	{
-		Capturer.Init();
+		CHECKinitialized_();
+		*available = !UE::PixelStreaming::Settings::CVarPixelStreamingWebRTCDisableReceiveAudio.GetValueOnAnyThread();
+		return 0;
 	}
 
-	return 0;
-}
+	int32 FAudioDeviceModule::InitPlayout()
+	{
+		CHECKinitialized_();
+		bool bIsPlayoutAvailable = false;
+		PlayoutIsAvailable(&bIsPlayoutAvailable);
+		if (!bIsPlayoutAvailable)
+		{
+			return -1;
+		}
 
-bool UE::PixelStreaming::FAudioDeviceModule::RecordingIsInitialized() const
-{
-	CHECKinitialized__BOOL();
-	return Capturer.IsInitialised();
-}
+		Requester.InitPlayout();
+		return 0;
+	}
 
-int32 UE::PixelStreaming::FAudioDeviceModule::StartRecording()
-{
-	CHECKinitialized_();
-	if (!Capturer.IsCapturing())
+	bool FAudioDeviceModule::PlayoutIsInitialized() const
 	{
-		Capturer.StartCapturing();
+		CHECKinitialized__BOOL();
+		return Requester.PlayoutIsInitialized();
 	}
-	return 0;
-}
 
-int32 UE::PixelStreaming::FAudioDeviceModule::StopRecording()
-{
-	CHECKinitialized_();
-	if (Capturer.IsCapturing())
+	int32 FAudioDeviceModule::StartPlayout()
 	{
-		Capturer.EndCapturing();
-	}
-	return 0;
-}
+		CHECKinitialized_();
 
-bool UE::PixelStreaming::FAudioDeviceModule::Recording() const
-{
-	CHECKinitialized__BOOL();
-	return Capturer.IsCapturing();
-}
+		if (!Requester.PlayoutIsInitialized())
+		{
+			return -1;
+		}
 
-int32 UE::PixelStreaming::FAudioDeviceModule::InitSpeaker()
-{
-	CHECKinitialized_();
-	return -1;
-}
+		Requester.StartPlayout();
+		return 0;
+	}
 
-bool UE::PixelStreaming::FAudioDeviceModule::SpeakerIsInitialized() const
-{
-	CHECKinitialized__BOOL();
-	return false;
-}
+	int32 FAudioDeviceModule::StopPlayout()
+	{
+		CHECKinitialized_();
+		if (!Requester.PlayoutIsInitialized())
+		{
+			return -1;
+		}
 
-int32 UE::PixelStreaming::FAudioDeviceModule::InitMicrophone()
-{
-	CHECKinitialized_();
-	return 0;
-}
+		Requester.StopPlayout();
+		return 0;
+	}
 
-bool UE::PixelStreaming::FAudioDeviceModule::MicrophoneIsInitialized() const
-{
-	CHECKinitialized__BOOL();
-	return true;
-}
+	bool FAudioDeviceModule::Playing() const
+	{
+		CHECKinitialized__BOOL();
+		return Requester.Playing();
+	}
 
-int32 UE::PixelStreaming::FAudioDeviceModule::StereoPlayoutIsAvailable(bool* available) const
-{
-	CHECKinitialized_();
-	*available = true;
-	return 0;
-}
+	int32 FAudioDeviceModule::RecordingIsAvailable(bool* available)
+	{
+		CHECKinitialized_();
+		*available = !UE::PixelStreaming::Settings::CVarPixelStreamingWebRTCDisableTransmitAudio.GetValueOnAnyThread();
+		return 0;
+	}
 
-int32 UE::PixelStreaming::FAudioDeviceModule::SetStereoPlayout(bool enable)
-{
-	CHECKinitialized_();
-	FString AudioChannelStr = enable ? TEXT("stereo") : TEXT("mono");
-	UE_LOG(LogPixelStreamingAudioDeviceModule, Verbose, TEXT("WebRTC has requested browser audio playout in UE be: %s"), *AudioChannelStr);
-	return 0;
-}
+	int32 FAudioDeviceModule::InitRecording()
+	{
+		CHECKinitialized_();
 
-int32 UE::PixelStreaming::FAudioDeviceModule::StereoPlayout(bool* enabled) const
-{
-	CHECKinitialized_();
-	*enabled = true;
-	return 0;
-}
+		bool bIsRecordingAvailable = false;
+		RecordingIsAvailable(&bIsRecordingAvailable);
+		if (!bIsRecordingAvailable)
+		{
+			return -1;
+		}
 
-int32 UE::PixelStreaming::FAudioDeviceModule::StereoRecordingIsAvailable(bool* available) const
-{
-	CHECKinitialized_();
-	*available = true;
-	return 0;
-}
+		if (!Capturer.IsInitialised())
+		{
+			Capturer.Init();
+		}
 
-int32 UE::PixelStreaming::FAudioDeviceModule::SetStereoRecording(bool enable)
-{
-	CHECKinitialized_();
-	return 0;
-}
+		return 0;
+	}
 
-int32 UE::PixelStreaming::FAudioDeviceModule::StereoRecording(bool* enabled) const
-{
-	CHECKinitialized_();
-	*enabled = true;
-	return 0;
-}
+	bool FAudioDeviceModule::RecordingIsInitialized() const
+	{
+		CHECKinitialized__BOOL();
+		return Capturer.IsInitialised();
+	}
 
-//--UE::PixelStreaming::FSubmixCapturer--
+	int32 FAudioDeviceModule::StartRecording()
+	{
+		CHECKinitialized_();
+		if (!Capturer.IsCapturing())
+		{
+			Capturer.StartCapturing();
+		}
+		return 0;
+	}
 
-uint32_t UE::PixelStreaming::FSubmixCapturer::GetVolume() const
-{
-	return VolumeLevel;
-}
+	int32 FAudioDeviceModule::StopRecording()
+	{
+		CHECKinitialized_();
+		if (Capturer.IsCapturing())
+		{
+			Capturer.EndCapturing();
+		}
+		return 0;
+	}
 
-void UE::PixelStreaming::FSubmixCapturer::SetVolume(uint32_t NewVolume)
-{
-	VolumeLevel = NewVolume;
-}
+	bool FAudioDeviceModule::Recording() const
+	{
+		CHECKinitialized__BOOL();
+		return Capturer.IsCapturing();
+	}
 
-bool UE::PixelStreaming::FSubmixCapturer::Init()
-{
-	FScopeLock Lock(&CriticalSection);
+	int32 FAudioDeviceModule::InitSpeaker()
+	{
+		CHECKinitialized_();
+		return -1;
+	}
 
-	// subscribe to audio data
-	if (!GEngine)
+	bool FAudioDeviceModule::SpeakerIsInitialized() const
 	{
-		bInitialised = false;
+		CHECKinitialized__BOOL();
 		return false;
 	}
 
-	// already initialised
-	if (bInitialised)
+	int32 FAudioDeviceModule::InitMicrophone()
+	{
+		CHECKinitialized_();
+		return 0;
+	}
+
+	bool FAudioDeviceModule::MicrophoneIsInitialized() const
 	{
+		CHECKinitialized__BOOL();
 		return true;
 	}
 
-	FAudioDeviceHandle AudioDevice = GEngine->GetMainAudioDevice();
-	if (!AudioDevice)
+	int32 FAudioDeviceModule::StereoPlayoutIsAvailable(bool* available) const
 	{
-		UE_LOG(LogPixelStreamingAudioDeviceModule, Warning, TEXT("No audio device"));
-		bInitialised = false;
-		return false;
+		CHECKinitialized_();
+		*available = true;
+		return 0;
 	}
 
-	AudioDevice->RegisterSubmixBufferListener(this);
-	bInitialised = true;
-	return true;
-}
+	int32 FAudioDeviceModule::SetStereoPlayout(bool enable)
+	{
+		CHECKinitialized_();
+		FString AudioChannelStr = enable ? TEXT("stereo") : TEXT("mono");
+		UE_LOG(LogPixelStreamingAudioDeviceModule, Verbose, TEXT("WebRTC has requested browser audio playout in UE be: %s"), *AudioChannelStr);
+		return 0;
+	}
 
-void UE::PixelStreaming::FSubmixCapturer::OnNewSubmixBuffer(const USoundSubmix* OwningSubmix, float* AudioData, int32 NumSamples, int32 NumChannels, const int32 SampleRate, double AudioClock)
-{
-	FScopeLock Lock(&CriticalSection);
+	int32 FAudioDeviceModule::StereoPlayout(bool* enabled) const
+	{
+		CHECKinitialized_();
+		*enabled = true;
+		return 0;
+	}
 
-	if (!bInitialised || !bCapturing)
+	int32 FAudioDeviceModule::StereoRecordingIsAvailable(bool* available) const
 	{
-		return;
+		CHECKinitialized_();
+		*available = true;
+		return 0;
 	}
 
-	// No point doing anything with UE audio if the callback from WebRTC has not been set yet.
-	if (AudioCallback == nullptr)
+	int32 FAudioDeviceModule::SetStereoRecording(bool enable)
 	{
-		return;
+		CHECKinitialized_();
+		return 0;
 	}
 
-	// Check if the sample rate from UE matches our target sample rate
-	if (TargetSampleRate != SampleRate)
+	int32 FAudioDeviceModule::StereoRecording(bool* enabled) const
 	{
-		// Only report the problem once
-		if (!bReportedSampleRateMismatch)
-		{
-			bReportedSampleRateMismatch = true;
-			UE_LOG(LogPixelStreamingAudioDeviceModule, Error, TEXT("Audio sample rate mismatch. Expected: %d | Actual: %d"), TargetSampleRate, SampleRate);
-		}
-		return;
+		CHECKinitialized_();
+		*enabled = true;
+		return 0;
 	}
 
-	UE_LOG(LogPixelStreamingAudioDeviceModule, VeryVerbose, TEXT("captured %d samples, %dc, %dHz"), NumSamples, NumChannels, SampleRate);
+	//--UE::PixelStreaming::FSubmixCapturer--
 
-	// Note: TSampleBuffer takes in AudioData as float* and internally converts to int16
-	Audio::TSampleBuffer<int16> Buffer(AudioData, NumSamples, NumChannels, SampleRate);
+	uint32_t FSubmixCapturer::GetVolume() const
+	{
+		return VolumeLevel;
+	}
 
-	// Mix to our target number of channels if the source does not already match.
-	if (Buffer.GetNumChannels() != TargetNumChannels)
+	void FSubmixCapturer::SetVolume(uint32_t NewVolume)
 	{
-		Buffer.MixBufferToChannels(TargetNumChannels);
+		VolumeLevel = NewVolume;
 	}
 
-	RecordingBuffer.Append(Buffer.GetData(), Buffer.GetNumSamples());
+	bool FSubmixCapturer::Init()
+	{
+		FScopeLock Lock(&CriticalSection);
 
-	const float ChunkDurationSecs = 0.01f; //10ms
-	const int32 SamplesPer10Ms = GetSamplesPerDurationSecs(ChunkDurationSecs);
+		// subscribe to audio data
+		if (!GEngine)
+		{
+			bInitialised = false;
+			return false;
+		}
 
-	// Feed in 10ms chunks
-	while (RecordingBuffer.Num() > SamplesPer10Ms)
-	{
+		// already initialised
+		if (bInitialised)
+		{
+			return true;
+		}
+
+		FAudioDeviceHandle AudioDevice = GEngine->GetMainAudioDevice();
+		if (!AudioDevice)
+		{
+			UE_LOG(LogPixelStreamingAudioDeviceModule, Warning, TEXT("No audio device"));
+			bInitialised = false;
+			return false;
+		}
+
+		AudioDevice->RegisterSubmixBufferListener(this);
+		bInitialised = true;
+		return true;
+	}
 
-		// Extract a 10ms chunk of samples from recording buffer
-		TArray<int16_t> SubmitBuffer(RecordingBuffer.GetData(), SamplesPer10Ms);
-		const size_t Frames = SubmitBuffer.Num() / TargetNumChannels;
-		const size_t BytesPerFrame = TargetNumChannels * sizeof(int16_t);
+	void FSubmixCapturer::OnNewSubmixBuffer(const USoundSubmix* OwningSubmix, float* AudioData, int32 NumSamples, int32 NumChannels, const int32 SampleRate, double AudioClock)
+	{
+		FScopeLock Lock(&CriticalSection);
 
-		uint32_t OutMicLevel = VolumeLevel;
+		if (!bInitialised || !bCapturing)
+		{
+			return;
+		}
 
-		int32_t WebRTCRes = AudioCallback->RecordedDataIsAvailable(
-			SubmitBuffer.GetData(),
-			Frames,
-			BytesPerFrame,
-			TargetNumChannels,
-			TargetSampleRate,
-			0,
-			0,
-			VolumeLevel,
-			false,
-			OutMicLevel);
+		// No point doing anything with UE audio if the callback from WebRTC has not been set yet.
+		if (AudioCallback == nullptr)
+		{
+			return;
+		}
 
-		SetVolume(OutMicLevel);
+		// Check if the sample rate from UE matches our target sample rate
+		if (TargetSampleRate != SampleRate)
+		{
+			// Only report the problem once
+			if (!bReportedSampleRateMismatch)
+			{
+				bReportedSampleRateMismatch = true;
+				UE_LOG(LogPixelStreamingAudioDeviceModule, Error, TEXT("Audio sample rate mismatch. Expected: %d | Actual: %d"), TargetSampleRate, SampleRate);
+			}
+			return;
+		}
 
-		// Remove 10ms of samples from the recording buffer now it is submitted
-		RecordingBuffer.RemoveAt(0, SamplesPer10Ms, false);
-	}
-}
+		UE_LOG(LogPixelStreamingAudioDeviceModule, VeryVerbose, TEXT("captured %d samples, %dc, %dHz"), NumSamples, NumChannels, SampleRate);
 
-int32 UE::PixelStreaming::FSubmixCapturer::GetSamplesPerDurationSecs(float InSeconds) const
-{
-	int32 SamplesPerSecond = TargetNumChannels * TargetSampleRate;
-	int32 NumSamplesPerDuration = (int32)(SamplesPerSecond * InSeconds);
-	return NumSamplesPerDuration;
-}
+		// Note: TSampleBuffer takes in AudioData as float* and internally converts to int16
+		Audio::TSampleBuffer<int16> Buffer(AudioData, NumSamples, NumChannels, SampleRate);
 
-void UE::PixelStreaming::FSubmixCapturer::RegisterAudioCallback(webrtc::AudioTransport* AudioTransportCallback)
-{
-	FScopeLock Lock(&CriticalSection);
-	AudioCallback = AudioTransportCallback;
-}
+		// Mix to our target number of channels if the source does not already match.
+		if (Buffer.GetNumChannels() != TargetNumChannels)
+		{
+			Buffer.MixBufferToChannels(TargetNumChannels);
+		}
 
-bool UE::PixelStreaming::FSubmixCapturer::IsInitialised() const
-{
-	return bInitialised;
-}
+		RecordingBuffer.Append(Buffer.GetData(), Buffer.GetNumSamples());
 
-void UE::PixelStreaming::FSubmixCapturer::Uninitialise()
-{
-	FScopeLock Lock(&CriticalSection);
+		const float ChunkDurationSecs = 0.01f; //10ms
+		const int32 SamplesPer10Ms = GetSamplesPerDurationSecs(ChunkDurationSecs);
 
-	if (GEngine)
-	{
-		FAudioDeviceHandle AudioDevice = GEngine->GetMainAudioDevice();
-		if (AudioDevice)
+		// Feed in 10ms chunks
+		while (RecordingBuffer.Num() > SamplesPer10Ms)
 		{
-			AudioDevice->UnregisterSubmixBufferListener(this);
+
+			// Extract a 10ms chunk of samples from recording buffer
+			TArray<int16_t> SubmitBuffer(RecordingBuffer.GetData(), SamplesPer10Ms);
+			const size_t Frames = SubmitBuffer.Num() / TargetNumChannels;
+			const size_t BytesPerFrame = TargetNumChannels * sizeof(int16_t);
+
+			uint32_t OutMicLevel = VolumeLevel;
+
+			int32_t WebRTCRes = AudioCallback->RecordedDataIsAvailable(
+				SubmitBuffer.GetData(),
+				Frames,
+				BytesPerFrame,
+				TargetNumChannels,
+				TargetSampleRate,
+				0,
+				0,
+				VolumeLevel,
+				false,
+				OutMicLevel);
+
+			SetVolume(OutMicLevel);
+
+			// Remove 10ms of samples from the recording buffer now it is submitted
+			RecordingBuffer.RemoveAt(0, SamplesPer10Ms, false);
 		}
 	}
 
-	RecordingBuffer.Empty();
-	bInitialised = false;
-	bCapturing = false;
-}
-
-bool UE::PixelStreaming::FSubmixCapturer::StartCapturing()
-{
-	if (!bInitialised)
+	int32 FSubmixCapturer::GetSamplesPerDurationSecs(float InSeconds) const
 	{
-		return false;
+		int32 SamplesPerSecond = TargetNumChannels * TargetSampleRate;
+		int32 NumSamplesPerDuration = (int32)(SamplesPerSecond * InSeconds);
+		return NumSamplesPerDuration;
 	}
-	bCapturing = true;
-	return true;
-}
 
-bool UE::PixelStreaming::FSubmixCapturer::EndCapturing()
-{
-	if (!bInitialised)
+	void FSubmixCapturer::RegisterAudioCallback(webrtc::AudioTransport* AudioTransportCallback)
 	{
-		return false;
+		FScopeLock Lock(&CriticalSection);
+		AudioCallback = AudioTransportCallback;
 	}
-	bCapturing = false;
-	return true;
-}
 
-bool UE::PixelStreaming::FSubmixCapturer::IsCapturing() const
-{
-	return bCapturing;
-}
-
-//--UE::PixelStreaming::FAudioPlayoutRequester--
-
-void UE::PixelStreaming::FAudioPlayoutRequester::Uninitialise()
-{
-	StopPlayout();
-	AudioCallback = nullptr;
-	bIsPlayoutInitialised = false;
-	PlayoutBuffer.Empty();
-}
-
-void UE::PixelStreaming::FAudioPlayoutRequester::InitPlayout()
-{
-	bIsPlayoutInitialised = true;
-}
+	bool FSubmixCapturer::IsInitialised() const
+	{
+		return bInitialised;
+	}
 
-void UE::PixelStreaming::FAudioPlayoutRequester::StartPlayout()
-{
-	if (PlayoutIsInitialized() && !Playing())
+	void FSubmixCapturer::Uninitialise()
 	{
-		TFunction<void()> RequesterFunc = [this]() {
-			FScopeLock Lock(&PlayoutCriticalSection);
+		FScopeLock Lock(&CriticalSection);
 
-			// Only request audio if audio callback is valid
-			if (!AudioCallback)
+		if (GEngine)
+		{
+			FAudioDeviceHandle AudioDevice = GEngine->GetMainAudioDevice();
+			if (AudioDevice)
 			{
-				return;
+				AudioDevice->UnregisterSubmixBufferListener(this);
 			}
+		}
+
+		RecordingBuffer.Empty();
+		bInitialised = false;
+		bCapturing = false;
+	}
 
-			// Our intention is to request samples at some fixed interval (i.e 10ms)
-			int32 NSamplesPerChannel = (SampleRate) / (1000 / UE::PixelStreaming::FAudioPlayoutRequester::RequestIntervalMs);
-			const size_t BytesPerFrame = NumChannels * sizeof(int16_t);
+	bool FSubmixCapturer::StartCapturing()
+	{
+		if (!bInitialised)
+		{
+			return false;
+		}
+		bCapturing = true;
+		return true;
+	}
 
-			// Ensure buffer has the total number of samples we need for all audio frames
-			PlayoutBuffer.Reserve(NSamplesPerChannel * NumChannels);
+	bool FSubmixCapturer::EndCapturing()
+	{
+		if (!bInitialised)
+		{
+			return false;
+		}
+		bCapturing = false;
+		return true;
+	}
 
-			size_t OutNSamples = 0;
-			int64_t ElapsedTimeMs = -1;
-			int64_t NtpTimeMs = -1;
+	bool FSubmixCapturer::IsCapturing() const
+	{
+		return bCapturing;
+	}
 
-			// Note this is mixed result of all audio sources, which in turn triggers the sinks for each audio source to be called.
-			// For example, if you had 3 audio sources in the browser sending 16kHz mono they would all be mixed down into the number
-			// of channels and sample rate specified below. However, for listening to each audio source prior to mixing refer to
-			// UE::PixelStreaming::FAudioSink.
-			uint32_t Result = AudioCallback->NeedMorePlayData(
-				NSamplesPerChannel,
-				BytesPerFrame,
-				NumChannels,
-				SampleRate,
-				PlayoutBuffer.GetData(),
-				OutNSamples,
-				&ElapsedTimeMs,
-				&NtpTimeMs);
-
-			if (Result != 0)
-			{
-				UE_LOG(LogPixelStreamingAudioDeviceModule, Error, TEXT("NeedMorePlayData return non-zero result indicating an error"));
-			}
-		};
+	//--UE::PixelStreaming::FAudioPlayoutRequester--
 
-		RequesterThread.Reset(nullptr);
-		RequesterRunnable = MakeUnique<UE::PixelStreaming::FAudioPlayoutRequester::Runnable>(RequesterFunc);
-		RequesterThread.Reset(FRunnableThread::Create(RequesterRunnable.Get(), TEXT("Pixel Streaming WebRTC Audio Requester")));
-		bIsPlaying = true;
+	void FAudioPlayoutRequester::Uninitialise()
+	{
+		StopPlayout();
+		AudioCallback = nullptr;
+		bIsPlayoutInitialised = false;
+		PlayoutBuffer.Empty();
 	}
-}
 
-void UE::PixelStreaming::FAudioPlayoutRequester::StopPlayout()
-{
-	if (PlayoutIsInitialized() && Playing())
+	void FAudioPlayoutRequester::InitPlayout()
 	{
-		RequesterRunnable->Stop();
-		bIsPlaying = false;
+		bIsPlayoutInitialised = true;
 	}
-}
 
-bool UE::PixelStreaming::FAudioPlayoutRequester::Playing() const
-{
-	return bIsPlaying;
-}
+	void FAudioPlayoutRequester::StartPlayout()
+	{
+		if (PlayoutIsInitialized() && !Playing())
+		{
+			TFunction<void()> RequesterFunc = [this]() {
+				FScopeLock Lock(&PlayoutCriticalSection);
+
+				// Only request audio if audio callback is valid
+				if (!AudioCallback)
+				{
+					return;
+				}
+
+				// Our intention is to request samples at some fixed interval (i.e 10ms)
+				int32 NSamplesPerChannel = (SampleRate) / (1000 / FAudioPlayoutRequester::RequestIntervalMs);
+				const size_t BytesPerFrame = NumChannels * sizeof(int16_t);
+
+				// Ensure buffer has the total number of samples we need for all audio frames
+				PlayoutBuffer.Reserve(NSamplesPerChannel * NumChannels);
+
+				size_t OutNSamples = 0;
+				int64_t ElapsedTimeMs = -1;
+				int64_t NtpTimeMs = -1;
+
+				// Note this is mixed result of all audio sources, which in turn triggers the sinks for each audio source to be called.
+				// For example, if you had 3 audio sources in the browser sending 16kHz mono they would all be mixed down into the number
+				// of channels and sample rate specified below. However, for listening to each audio source prior to mixing refer to
+				// UE::PixelStreaming::FAudioSink.
+				uint32_t Result = AudioCallback->NeedMorePlayData(
+					NSamplesPerChannel,
+					BytesPerFrame,
+					NumChannels,
+					SampleRate,
+					PlayoutBuffer.GetData(),
+					OutNSamples,
+					&ElapsedTimeMs,
+					&NtpTimeMs);
+
+				if (Result != 0)
+				{
+					UE_LOG(LogPixelStreamingAudioDeviceModule, Error, TEXT("NeedMorePlayData return non-zero result indicating an error"));
+				}
+			};
+
+			RequesterThread.Reset(nullptr);
+			RequesterRunnable = MakeUnique<FAudioPlayoutRequester::Runnable>(RequesterFunc);
+			RequesterThread.Reset(FRunnableThread::Create(RequesterRunnable.Get(), TEXT("Pixel Streaming WebRTC Audio Requester")));
+			bIsPlaying = true;
+		}
+	}
 
-bool UE::PixelStreaming::FAudioPlayoutRequester::PlayoutIsInitialized() const
-{
-	return bIsPlayoutInitialised;
-}
+	void FAudioPlayoutRequester::StopPlayout()
+	{
+		if (PlayoutIsInitialized() && Playing())
+		{
+			RequesterRunnable->Stop();
+			bIsPlaying = false;
+		}
+	}
 
-void UE::PixelStreaming::FAudioPlayoutRequester::RegisterAudioCallback(webrtc::AudioTransport* AudioCb)
-{
-	FScopeLock Lock(&PlayoutCriticalSection);
-	AudioCallback = AudioCb;
-}
+	bool FAudioPlayoutRequester::Playing() const
+	{
+		return bIsPlaying;
+	}
 
-//--FAudioRequester::Runnable--
+	bool FAudioPlayoutRequester::PlayoutIsInitialized() const
+	{
+		return bIsPlayoutInitialised;
+	}
 
-bool UE::PixelStreaming::FAudioPlayoutRequester::Runnable::Init()
-{
-	return RequestPlayoutFunc != nullptr;
-}
+	void FAudioPlayoutRequester::RegisterAudioCallback(webrtc::AudioTransport* AudioCb)
+	{
+		FScopeLock Lock(&PlayoutCriticalSection);
+		AudioCallback = AudioCb;
+	}
 
-uint32 UE::PixelStreaming::FAudioPlayoutRequester::Runnable::Run()
-{
-	LastAudioRequestTimeMs = rtc::TimeMillis();
-	bIsRunning = true;
+	//--FAudioRequester::Runnable--
 
-	// Request audio in a loop until this boolean is toggled off.
-	while (bIsRunning)
+	bool FAudioPlayoutRequester::Runnable::Init()
 	{
-		int64_t Now = rtc::TimeMillis();
-		int64_t DeltaMs = Now - LastAudioRequestTimeMs;
+		return RequestPlayoutFunc != nullptr;
+	}
 
-		// Check if the 10ms delta has elapsed, if it has not, then sleep the remaining
-		if (DeltaMs < UE::PixelStreaming::FAudioPlayoutRequester::RequestIntervalMs)
+	uint32 FAudioPlayoutRequester::Runnable::Run()
+	{
+		LastAudioRequestTimeMs = rtc::TimeMillis();
+		bIsRunning = true;
+
+		// Request audio in a loop until this boolean is toggled off.
+		while (bIsRunning)
 		{
-			int64_t SleepTimeMs = UE::PixelStreaming::FAudioPlayoutRequester::RequestIntervalMs - DeltaMs;
-			float SleepTimeSecs = (float)SleepTimeMs / 1000.0f;
-			FPlatformProcess::Sleep(SleepTimeSecs);
-			continue;
-		}
+			int64_t Now = rtc::TimeMillis();
+			int64_t DeltaMs = Now - LastAudioRequestTimeMs;
 
-		// Update request time to now seeing as the 10ms delta has elapsed
-		LastAudioRequestTimeMs = Now;
+			// Check if the 10ms delta has elapsed, if it has not, then sleep the remaining
+			if (DeltaMs < FAudioPlayoutRequester::RequestIntervalMs)
+			{
+				int64_t SleepTimeMs = FAudioPlayoutRequester::RequestIntervalMs - DeltaMs;
+				float SleepTimeSecs = (float)SleepTimeMs / 1000.0f;
+				FPlatformProcess::Sleep(SleepTimeSecs);
+				continue;
+			}
 
-		// Actually request playout
-		RequestPlayoutFunc();
-	}
+			// Update request time to now seeing as the 10ms delta has elapsed
+			LastAudioRequestTimeMs = Now;
 
-	return 0;
-}
+			// Actually request playout
+			RequestPlayoutFunc();
+		}
 
-void UE::PixelStreaming::FAudioPlayoutRequester::Runnable::Stop()
-{
-	bIsRunning = false;
-}
+		return 0;
+	}
 
-void UE::PixelStreaming::FAudioPlayoutRequester::Runnable::Exit()
-{
-	bIsRunning = false;
-	RequestPlayoutFunc = nullptr;
-}
\ No newline at end of file
+	void FAudioPlayoutRequester::Runnable::Stop()
+	{
+		bIsRunning = false;
+	}
+
+	void FAudioPlayoutRequester::Runnable::Exit()
+	{
+		bIsRunning = false;
+		RequestPlayoutFunc = nullptr;
+	}
+} // namespace UE::PixelStreaming
\ No newline at end of file
diff --git a/Engine/Plugins/Media/PixelStreaming/Source/PixelStreaming/Private/AudioDeviceModule.h b/Engine/Plugins/Media/PixelStreaming/Source/PixelStreaming/Private/AudioDeviceModule.h
--- a/Engine/Plugins/Media/PixelStreaming/Source/PixelStreaming/Private/AudioDeviceModule.h
+++ b/Engine/Plugins/Media/PixelStreaming/Source/PixelStreaming/Private/AudioDeviceModule.h
@@ -10,305 +10,306 @@
 #include "Templates/Function.h"
 #include "Templates/SharedPointer.h"
 
-namespace UE {
-	namespace PixelStreaming {
-		// Captures audio from UE and passes it along to WebRTC.
-		class FSubmixCapturer : public ISubmixBufferListener
+namespace UE::PixelStreaming 
+{
+	// Captures audio from UE and passes it along to WebRTC.
+	class FSubmixCapturer : public ISubmixBufferListener
+	{
+
+	public:
+		// This magic number is the max volume used in webrtc fake audio device module.
+		static const uint32_t MaxVolumeLevel = 14392;
+
+		FSubmixCapturer()
+			: bInitialised(false)
+			, bCapturing(false)
+			, TargetSampleRate(48000)
+			, TargetNumChannels(2)
+			, bReportedSampleRateMismatch(false)
+			, AudioCallback(nullptr)
+			, VolumeLevel(FSubmixCapturer::MaxVolumeLevel)
+			, RecordingBuffer()
+			, CriticalSection()
+		{
+		}
+
+		virtual ~FSubmixCapturer() = default;
+
+		bool Init();
+		bool IsInitialised() const;
+		bool IsCapturing() const;
+		void Uninitialise();
+		bool StartCapturing();
+		bool EndCapturing();
+		uint32_t GetVolume() const;
+		void SetVolume(uint32_t NewVolume);
+
+		// ISubmixBufferListener interface
+		void OnNewSubmixBuffer(const USoundSubmix* OwningSubmix, float* AudioData, int32 NumSamples, int32 NumChannels, const int32 SampleRate, double AudioClock) override;
+
+		void RegisterAudioCallback(webrtc::AudioTransport* AudioCb);
+
+	private:
+		int32 GetSamplesPerDurationSecs(float InSeconds) const;
+
+	private:
+		FThreadSafeBool bInitialised;
+		FThreadSafeBool bCapturing;
+		uint32 TargetSampleRate;
+		uint8 TargetNumChannels;
+		bool bReportedSampleRateMismatch;
+		webrtc::AudioTransport* AudioCallback;
+		uint32_t VolumeLevel;
+		TArray<int16_t> RecordingBuffer;
+		FCriticalSection CriticalSection; // One thread captures audio from UE, the other controls the state of the capturer.
+	};
+
+	// Requests audio from WebRTC at a regular interval (10ms)
+	// This is required so that WebRTC audio sinks actually have
+	// some audio data for their sinks. Without this WebRTC assumes
+	// there is no demand for audio and does not populate the sinks.
+	class FAudioPlayoutRequester
+	{
+	public:
+		class Runnable : public FRunnable
 		{
-
 		public:
-			// This magic number is the max volume used in webrtc fake audio device module.
-			static const uint32_t MaxVolumeLevel = 14392;
-
-			FSubmixCapturer()
-				: bInitialised(false)
-				, bCapturing(false)
-				, TargetSampleRate(48000)
-				, TargetNumChannels(2)
-				, bReportedSampleRateMismatch(false)
-				, AudioCallback(nullptr)
-				, VolumeLevel(FSubmixCapturer::MaxVolumeLevel)
-				, RecordingBuffer()
-				, CriticalSection()
-			{
-			}
-
-			virtual ~FSubmixCapturer() = default;
-
-			bool Init();
-			bool IsInitialised() const;
-			bool IsCapturing() const;
-			void Uninitialise();
-			bool StartCapturing();
-			bool EndCapturing();
-			uint32_t GetVolume() const;
-			void SetVolume(uint32_t NewVolume);
-
-			// ISubmixBufferListener interface
-			void OnNewSubmixBuffer(const USoundSubmix* OwningSubmix, float* AudioData, int32 NumSamples, int32 NumChannels, const int32 SampleRate, double AudioClock) override;
-
-			void RegisterAudioCallback(webrtc::AudioTransport* AudioCb);
-
-		private:
-			int32 GetSamplesPerDurationSecs(float InSeconds) const;
-
+			Runnable(TFunction<void()> RequestPlayoutFunc)
+				: bIsRunning(false)
+				, LastAudioRequestTimeMs(0)
+				, RequestPlayoutFunc(RequestPlayoutFunc) {}
+			// Begin FRunnable interface.
+			virtual bool Init() override;
+			virtual uint32 Run() override;
+			virtual void Stop() override;
+			virtual void Exit() override;
+			// End FRunnable interface
 		private:
-			FThreadSafeBool bInitialised;
-			FThreadSafeBool bCapturing;
-			uint32 TargetSampleRate;
-			uint8 TargetNumChannels;
-			bool bReportedSampleRateMismatch;
-			webrtc::AudioTransport* AudioCallback;
-			uint32_t VolumeLevel;
-			TArray<int16_t> RecordingBuffer;
-			FCriticalSection CriticalSection; // One thread captures audio from UE, the other controls the state of the capturer.
+			bool bIsRunning;
+			int64_t LastAudioRequestTimeMs;
+			TFunction<void()> RequestPlayoutFunc;
 		};
 
-		// Requests audio from WebRTC at a regular interval (10ms)
-		// This is required so that WebRTC audio sinks actually have
-		// some audio data for their sinks. Without this WebRTC assumes
-		// there is no demand for audio and does not populate the sinks.
-		class FAudioPlayoutRequester
+		FAudioPlayoutRequester()
+			: bIsPlayoutInitialised(false)
+			, bIsPlaying(false)
+			, SampleRate(48000) //webrtc will mix all sources into one big buffer with this sample rate
+			, NumChannels(2)	//webrtc will mix all sources into one big buffer with this many channels
+			, RequesterRunnable()
+			, RequesterThread()
+			, AudioCallback(nullptr)
+			, PlayoutCriticalSection()
+			, PlayoutBuffer()
 		{
-		public:
-			class Runnable : public FRunnable
-			{
-			public:
-				Runnable(TFunction<void()> RequestPlayoutFunc)
-					: bIsRunning(false)
-					, LastAudioRequestTimeMs(0)
-					, RequestPlayoutFunc(RequestPlayoutFunc) {}
-				// Begin FRunnable interface.
-				virtual bool Init() override;
-				virtual uint32 Run() override;
-				virtual void Stop() override;
-				virtual void Exit() override;
-				// End FRunnable interface
-			private:
-				bool bIsRunning;
-				int64_t LastAudioRequestTimeMs;
-				TFunction<void()> RequestPlayoutFunc;
-			};
-
-			FAudioPlayoutRequester()
-				: bIsPlayoutInitialised(false)
-				, bIsPlaying(false)
-				, SampleRate(48000) //webrtc will mix all sources into one big buffer with this sample rate
-				, NumChannels(2)	//webrtc will mix all sources into one big buffer with this many channels
-				, RequesterRunnable()
-				, RequesterThread()
-				, AudioCallback(nullptr)
-				, PlayoutCriticalSection()
-				, PlayoutBuffer(){};
-
-			void InitPlayout();
-			void StartPlayout();
-			void StopPlayout();
-			bool Playing() const;
-			bool PlayoutIsInitialized() const;
-			void Uninitialise();
-			void RegisterAudioCallback(webrtc::AudioTransport* AudioCallback);
-
-		public:
-			static int16_t const RequestIntervalMs = 10;
-
-		private:
-			FThreadSafeBool bIsPlayoutInitialised;
-			FThreadSafeBool bIsPlaying;
-			uint32 SampleRate;
-			uint8 NumChannels;
-			TUniquePtr<FAudioPlayoutRequester::Runnable> RequesterRunnable;
-			TUniquePtr<FRunnableThread> RequesterThread;
-			webrtc::AudioTransport* AudioCallback;
-			FCriticalSection PlayoutCriticalSection;
-			TArray<int16_t> PlayoutBuffer;
 		};
 
-		// A custom audio device module for WebRTC.
-		// It lets us receive WebRTC audio in UE
-		// and also transmit UE audio to WebRTC.
-		class FAudioDeviceModule : public webrtc::AudioDeviceModule
+		void InitPlayout();
+		void StartPlayout();
+		void StopPlayout();
+		bool Playing() const;
+		bool PlayoutIsInitialized() const;
+		void Uninitialise();
+		void RegisterAudioCallback(webrtc::AudioTransport* AudioCallback);
+
+	public:
+		static int16_t const RequestIntervalMs = 10;
+
+	private:
+		FThreadSafeBool bIsPlayoutInitialised;
+		FThreadSafeBool bIsPlaying;
+		uint32 SampleRate;
+		uint8 NumChannels;
+		TUniquePtr<FAudioPlayoutRequester::Runnable> RequesterRunnable;
+		TUniquePtr<FRunnableThread> RequesterThread;
+		webrtc::AudioTransport* AudioCallback;
+		FCriticalSection PlayoutCriticalSection;
+		TArray<int16_t> PlayoutBuffer;
+	};
+
+	// A custom audio device module for WebRTC.
+	// It lets us receive WebRTC audio in UE
+	// and also transmit UE audio to WebRTC.
+	class FAudioDeviceModule : public webrtc::AudioDeviceModule
+	{
+
+	public:
+		FAudioDeviceModule()
+			: bInitialized(false)
+			, Capturer()
+			, Requester(){};
+
+		virtual ~FAudioDeviceModule() = default;
+
+	private:
+		FThreadSafeBool bInitialized; // True when we setup capturer/playout requester.
+		FSubmixCapturer Capturer;
+		FAudioPlayoutRequester Requester;
+
+	private:
+		//
+		// webrtc::AudioDeviceModule interface
+		//
+		int32 ActiveAudioLayer(AudioLayer* audioLayer) const override;
+		int32 RegisterAudioCallback(webrtc::AudioTransport* audioCallback) override;
+
+		// Main initialization and termination
+		int32 Init() override;
+		int32 Terminate() override;
+		bool Initialized() const override;
+
+		// Device enumeration
+		int16 PlayoutDevices() override;
+		int16 RecordingDevices() override;
+		int32 PlayoutDeviceName(uint16 index, char name[webrtc::kAdmMaxDeviceNameSize], char guid[webrtc::kAdmMaxGuidSize]) override;
+		int32 RecordingDeviceName(uint16 index, char name[webrtc::kAdmMaxDeviceNameSize], char guid[webrtc::kAdmMaxGuidSize]) override;
+
+		// Device selection
+		int32 SetPlayoutDevice(uint16 index) override;
+		int32 SetPlayoutDevice(WindowsDeviceType device) override;
+		int32 SetRecordingDevice(uint16 index) override;
+		int32 SetRecordingDevice(WindowsDeviceType device) override;
+
+		// Audio transport initialization
+		int32 PlayoutIsAvailable(bool* available) override;
+		int32 InitPlayout() override;
+		bool PlayoutIsInitialized() const override;
+		int32 RecordingIsAvailable(bool* available) override;
+		int32 InitRecording() override;
+		bool RecordingIsInitialized() const override;
+
+		// Audio transport control
+		virtual int32 StartPlayout() override;
+		virtual int32 StopPlayout() override;
+
+		// True when audio is being pulled by the instance.
+		virtual bool Playing() const override;
+
+		virtual int32 StartRecording() override;
+		virtual int32 StopRecording() override;
+		virtual bool Recording() const override;
+
+		// Audio mixer initialization
+		virtual int32 InitSpeaker() override;
+		virtual bool SpeakerIsInitialized() const override;
+		virtual int32 InitMicrophone() override;
+		virtual bool MicrophoneIsInitialized() const override;
+
+		// Speaker volume controls
+		virtual int32 SpeakerVolumeIsAvailable(bool* available) override
 		{
+			return -1;
+		}
+		virtual int32 SetSpeakerVolume(uint32 volume) override
+		{
+			return -1;
+		}
+		virtual int32 SpeakerVolume(uint32* volume) const override
+		{
+			return -1;
+		}
+		virtual int32 MaxSpeakerVolume(uint32* maxVolume) const override
+		{
+			return -1;
+		}
+		virtual int32 MinSpeakerVolume(uint32* minVolume) const override
+		{
+			return -1;
+		}
 
-		public:
-			FAudioDeviceModule()
-				: bInitialized(false)
-				, Capturer()
-				, Requester(){};
+		// Microphone volume controls
+		virtual int32 MicrophoneVolumeIsAvailable(bool* available) override
+		{
+			return 0;
+		}
+		virtual int32 SetMicrophoneVolume(uint32 volume) override
+		{
+			//Capturer.SetVolume(volume);
+			return 0;
+		}
+		virtual int32 MicrophoneVolume(uint32* volume) const override
+		{
+			//*volume = Capturer.GetVolume();
+			return 0;
+		}
+		virtual int32 MaxMicrophoneVolume(uint32* maxVolume) const override
+		{
+			*maxVolume = FSubmixCapturer::MaxVolumeLevel;
+			return 0;
+		}
+		virtual int32 MinMicrophoneVolume(uint32* minVolume) const override
+		{
+			return 0;
+		}
 
-			virtual ~FAudioDeviceModule() = default;
+		// Speaker mute control
+		virtual int32 SpeakerMuteIsAvailable(bool* available) override
+		{
+			return -1;
+		}
+		virtual int32 SetSpeakerMute(bool enable) override
+		{
+			return -1;
+		}
+		virtual int32 SpeakerMute(bool* enabled) const override
+		{
+			return -1;
+		}
 
-		private:
-			FThreadSafeBool bInitialized; // True when we setup capturer/playout requester.
-			FSubmixCapturer Capturer;
-			FAudioPlayoutRequester Requester;
+		// Microphone mute control
+		virtual int32 MicrophoneMuteIsAvailable(bool* available) override
+		{
+			*available = false;
+			return -1;
+		}
+		virtual int32 SetMicrophoneMute(bool enable) override
+		{
+			return -1;
+		}
+		virtual int32 MicrophoneMute(bool* enabled) const override
+		{
+			return -1;
+		}
+
+		// Stereo support
+		virtual int32 StereoPlayoutIsAvailable(bool* available) const override;
+		virtual int32 SetStereoPlayout(bool enable) override;
+		virtual int32 StereoPlayout(bool* enabled) const override;
+		virtual int32 StereoRecordingIsAvailable(bool* available) const override;
+		virtual int32 SetStereoRecording(bool enable) override;
+		virtual int32 StereoRecording(bool* enabled) const override;
+
+		// Playout delay
+		virtual int32 PlayoutDelay(uint16* delayMS) const override
+		{
+			*delayMS = 0;
+			return 0;
+		}
 
-		private:
-			//
-			// webrtc::AudioDeviceModule interface
-			//
-			int32 ActiveAudioLayer(AudioLayer* audioLayer) const override;
-			int32 RegisterAudioCallback(webrtc::AudioTransport* audioCallback) override;
-
-			// Main initialization and termination
-			int32 Init() override;
-			int32 Terminate() override;
-			bool Initialized() const override;
-
-			// Device enumeration
-			int16 PlayoutDevices() override;
-			int16 RecordingDevices() override;
-			int32 PlayoutDeviceName(uint16 index, char name[webrtc::kAdmMaxDeviceNameSize], char guid[webrtc::kAdmMaxGuidSize]) override;
-			int32 RecordingDeviceName(uint16 index, char name[webrtc::kAdmMaxDeviceNameSize], char guid[webrtc::kAdmMaxGuidSize]) override;
-
-			// Device selection
-			int32 SetPlayoutDevice(uint16 index) override;
-			int32 SetPlayoutDevice(WindowsDeviceType device) override;
-			int32 SetRecordingDevice(uint16 index) override;
-			int32 SetRecordingDevice(WindowsDeviceType device) override;
-
-			// Audio transport initialization
-			int32 PlayoutIsAvailable(bool* available) override;
-			int32 InitPlayout() override;
-			bool PlayoutIsInitialized() const override;
-			int32 RecordingIsAvailable(bool* available) override;
-			int32 InitRecording() override;
-			bool RecordingIsInitialized() const override;
-
-			// Audio transport control
-			virtual int32 StartPlayout() override;
-			virtual int32 StopPlayout() override;
-
-			// True when audio is being pulled by the instance.
-			virtual bool Playing() const override;
-
-			virtual int32 StartRecording() override;
-			virtual int32 StopRecording() override;
-			virtual bool Recording() const override;
-
-			// Audio mixer initialization
-			virtual int32 InitSpeaker() override;
-			virtual bool SpeakerIsInitialized() const override;
-			virtual int32 InitMicrophone() override;
-			virtual bool MicrophoneIsInitialized() const override;
-
-			// Speaker volume controls
-			virtual int32 SpeakerVolumeIsAvailable(bool* available) override
-			{
-				return -1;
-			}
-			virtual int32 SetSpeakerVolume(uint32 volume) override
-			{
-				return -1;
-			}
-			virtual int32 SpeakerVolume(uint32* volume) const override
-			{
-				return -1;
-			}
-			virtual int32 MaxSpeakerVolume(uint32* maxVolume) const override
-			{
-				return -1;
-			}
-			virtual int32 MinSpeakerVolume(uint32* minVolume) const override
-			{
-				return -1;
-			}
-
-			// Microphone volume controls
-			virtual int32 MicrophoneVolumeIsAvailable(bool* available) override
-			{
-				return 0;
-			}
-			virtual int32 SetMicrophoneVolume(uint32 volume) override
-			{
-				//Capturer.SetVolume(volume);
-				return 0;
-			}
-			virtual int32 MicrophoneVolume(uint32* volume) const override
-			{
-				//*volume = Capturer.GetVolume();
-				return 0;
-			}
-			virtual int32 MaxMicrophoneVolume(uint32* maxVolume) const override
-			{
-				*maxVolume = FSubmixCapturer::MaxVolumeLevel;
-				return 0;
-			}
-			virtual int32 MinMicrophoneVolume(uint32* minVolume) const override
-			{
-				return 0;
-			}
-
-			// Speaker mute control
-			virtual int32 SpeakerMuteIsAvailable(bool* available) override
-			{
-				return -1;
-			}
-			virtual int32 SetSpeakerMute(bool enable) override
-			{
-				return -1;
-			}
-			virtual int32 SpeakerMute(bool* enabled) const override
-			{
-				return -1;
-			}
-
-			// Microphone mute control
-			virtual int32 MicrophoneMuteIsAvailable(bool* available) override
-			{
-				*available = false;
-				return -1;
-			}
-			virtual int32 SetMicrophoneMute(bool enable) override
-			{
-				return -1;
-			}
-			virtual int32 MicrophoneMute(bool* enabled) const override
-			{
-				return -1;
-			}
-
-			// Stereo support
-			virtual int32 StereoPlayoutIsAvailable(bool* available) const override;
-			virtual int32 SetStereoPlayout(bool enable) override;
-			virtual int32 StereoPlayout(bool* enabled) const override;
-			virtual int32 StereoRecordingIsAvailable(bool* available) const override;
-			virtual int32 SetStereoRecording(bool enable) override;
-			virtual int32 StereoRecording(bool* enabled) const override;
-
-			// Playout delay
-			virtual int32 PlayoutDelay(uint16* delayMS) const override
-			{
-				*delayMS = 0;
-				return 0;
-			}
-
-			virtual bool BuiltInAECIsAvailable() const override
-			{
-				return false;
-			}
-			virtual bool BuiltInAGCIsAvailable() const override
-			{
-				return false;
-			}
-			virtual bool BuiltInNSIsAvailable() const override
-			{
-				return false;
-			}
-
-			// Enables the built-in audio effects. Only supported on Android.
-			virtual int32 EnableBuiltInAEC(bool enable) override
-			{
-				return -1;
-			}
-			virtual int32 EnableBuiltInAGC(bool enable) override
-			{
-				return -1;
-			}
-			virtual int32 EnableBuiltInNS(bool enable) override
-			{
-				return -1;
-			}
-		};
-	}
-}
+		virtual bool BuiltInAECIsAvailable() const override
+		{
+			return false;
+		}
+		virtual bool BuiltInAGCIsAvailable() const override
+		{
+			return false;
+		}
+		virtual bool BuiltInNSIsAvailable() const override
+		{
+			return false;
+		}
+
+		// Enables the built-in audio effects. Only supported on Android.
+		virtual int32 EnableBuiltInAEC(bool enable) override
+		{
+			return -1;
+		}
+		virtual int32 EnableBuiltInAGC(bool enable) override
+		{
+			return -1;
+		}
+		virtual int32 EnableBuiltInNS(bool enable) override
+		{
+			return -1;
+		}
+	};
+} // namespace UE::PixelStreaming
\ No newline at end of file
diff --git a/Engine/Plugins/Media/PixelStreaming/Source/PixelStreaming/Private/AudioSink.cpp b/Engine/Plugins/Media/PixelStreaming/Source/PixelStreaming/Private/AudioSink.cpp
--- a/Engine/Plugins/Media/PixelStreaming/Source/PixelStreaming/Private/AudioSink.cpp
+++ b/Engine/Plugins/Media/PixelStreaming/Source/PixelStreaming/Private/AudioSink.cpp
@@ -2,61 +2,64 @@
 
 #include "AudioSink.h"
 
-UE::PixelStreaming::FAudioSink::~FAudioSink()
+namespace UE::PixelStreaming
 {
-	for (auto Iter = AudioConsumers.CreateIterator(); Iter; ++Iter)
+	FAudioSink::~FAudioSink()
 	{
-		IPixelStreamingAudioConsumer* AudioConsumer = Iter.ElementIt->Value;
-		Iter.RemoveCurrent();
-		if (AudioConsumer != nullptr)
+		for (auto Iter = AudioConsumers.CreateIterator(); Iter; ++Iter)
 		{
-			AudioConsumer->OnConsumerRemoved();
+			IPixelStreamingAudioConsumer* AudioConsumer = Iter.ElementIt->Value;
+			Iter.RemoveCurrent();
+			if (AudioConsumer != nullptr)
+			{
+				AudioConsumer->OnConsumerRemoved();
+			}
 		}
 	}
-}
 
-void UE::PixelStreaming::FAudioSink::OnData(const void* audio_data, int bits_per_sample, int sample_rate, size_t number_of_channels, size_t number_of_frames, absl::optional<int64_t> absolute_capture_timestamp_ms)
-{
+	void FAudioSink::OnData(const void* audio_data, int bits_per_sample, int sample_rate, size_t number_of_channels, size_t number_of_frames, absl::optional<int64_t> absolute_capture_timestamp_ms)
+	{
 
-	// This data is populated from the internals of WebRTC, basically each audio track sent from the browser has its RTP audio source received and decoded.
-	// The sample rate and number of channels here has absolutely no relationship with PixelStreamingAudioDeviceModule.
-	// The sample rate and number of channels here is determined adaptively by WebRTC's NetEQ class that selects sample rate/number of channels
-	// based on network conditions and other factors.
+		// This data is populated from the internals of WebRTC, basically each audio track sent from the browser has its RTP audio source received and decoded.
+		// The sample rate and number of channels here has absolutely no relationship with PixelStreamingAudioDeviceModule.
+		// The sample rate and number of channels here is determined adaptively by WebRTC's NetEQ class that selects sample rate/number of channels
+		// based on network conditions and other factors.
+
+		if (!HasAudioConsumers())
+		{
+			return;
+		}
+
+		// Iterate audio consumers and pass this data to their buffers
+		for (IPixelStreamingAudioConsumer* AudioConsumer : AudioConsumers)
+		{
+			AudioConsumer->ConsumeRawPCM(static_cast<const int16_t*>(audio_data), sample_rate, number_of_channels, number_of_frames);
+		}
 
-	if (!HasAudioConsumers())
-	{
 		return;
 	}
 
-	// Iterate audio consumers and pass this data to their buffers
-	for (IPixelStreamingAudioConsumer* AudioConsumer : AudioConsumers)
+	void FAudioSink::AddAudioConsumer(IPixelStreamingAudioConsumer* AudioConsumer)
 	{
-		AudioConsumer->ConsumeRawPCM(static_cast<const int16_t*>(audio_data), sample_rate, number_of_channels, number_of_frames);
+		bool bAlreadyInSet = false;
+		AudioConsumers.Add(AudioConsumer, &bAlreadyInSet);
+		if (!bAlreadyInSet)
+		{
+			AudioConsumer->OnConsumerAdded();
+		}
 	}
 
-	return;
-}
-
-void UE::PixelStreaming::FAudioSink::AddAudioConsumer(IPixelStreamingAudioConsumer* AudioConsumer)
-{
-	bool bAlreadyInSet = false;
-	AudioConsumers.Add(AudioConsumer, &bAlreadyInSet);
-	if (!bAlreadyInSet)
+	void FAudioSink::RemoveAudioConsumer(IPixelStreamingAudioConsumer* AudioConsumer)
 	{
-		AudioConsumer->OnConsumerAdded();
+		if (AudioConsumers.Contains(AudioConsumer))
+		{
+			AudioConsumers.Remove(AudioConsumer);
+			AudioConsumer->OnConsumerRemoved();
+		}
 	}
-}
 
-void UE::PixelStreaming::FAudioSink::RemoveAudioConsumer(IPixelStreamingAudioConsumer* AudioConsumer)
-{
-	if (AudioConsumers.Contains(AudioConsumer))
+	bool FAudioSink::HasAudioConsumers()
 	{
-		AudioConsumers.Remove(AudioConsumer);
-		AudioConsumer->OnConsumerRemoved();
+		return AudioConsumers.Num() > 0;
 	}
-}
-
-bool UE::PixelStreaming::FAudioSink::HasAudioConsumers()
-{
-	return AudioConsumers.Num() > 0;
-}
\ No newline at end of file
+} // namespace UE::PixelStreaming
\ No newline at end of file
diff --git a/Engine/Plugins/Media/PixelStreaming/Source/PixelStreaming/Private/AudioSink.h b/Engine/Plugins/Media/PixelStreaming/Source/PixelStreaming/Private/AudioSink.h
--- a/Engine/Plugins/Media/PixelStreaming/Source/PixelStreaming/Private/AudioSink.h
+++ b/Engine/Plugins/Media/PixelStreaming/Source/PixelStreaming/Private/AudioSink.h
@@ -6,34 +6,33 @@
 #include "IPixelStreamingAudioConsumer.h"
 #include "Containers/Set.h"
 
-namespace UE {
-	namespace PixelStreaming {
-		// This is a PixelStreaming sink AND a WebRTC sink.
-		// It collects audio coming in from a WebRTC audio source and passes into into UE's audio system.
-		class FAudioSink : public webrtc::AudioTrackSinkInterface, public IPixelStreamingAudioSink
-		{
-		public:
-			FAudioSink()
-				: AudioConsumers(){};
+namespace UE::PixelStreaming 
+{
+	// This is a PixelStreaming sink AND a WebRTC sink.
+	// It collects audio coming in from a WebRTC audio source and passes into into UE's audio system.
+	class FAudioSink : public webrtc::AudioTrackSinkInterface, public IPixelStreamingAudioSink
+	{
+	public:
+		FAudioSink()
+			: AudioConsumers(){};
 
-			// Note: destructor will call destroy on any attached audio consumers
-			virtual ~FAudioSink();
+		// Note: destructor will call destroy on any attached audio consumers
+		virtual ~FAudioSink();
 
-			// Begin AudioTrackSinkInterface
-			void OnData(const void* audio_data, int bits_per_sample, int sample_rate, size_t number_of_channels, size_t number_of_frames, absl::optional<int64_t> absolute_capture_timestamp_ms) override;
-			// End AudioTrackSinkInterface
+		// Begin AudioTrackSinkInterface
+		void OnData(const void* audio_data, int bits_per_sample, int sample_rate, size_t number_of_channels, size_t number_of_frames, absl::optional<int64_t> absolute_capture_timestamp_ms) override;
+		// End AudioTrackSinkInterface
 
-			// Begin IPixelStreamingAudioSink
-			void AddAudioConsumer(IPixelStreamingAudioConsumer* AudioConsumer) override;
-			void RemoveAudioConsumer(IPixelStreamingAudioConsumer* AudioConsumer) override;
-			bool HasAudioConsumers() override;
-			// End IPixelStreamingAudioSink
+		// Begin IPixelStreamingAudioSink
+		void AddAudioConsumer(IPixelStreamingAudioConsumer* AudioConsumer) override;
+		void RemoveAudioConsumer(IPixelStreamingAudioConsumer* AudioConsumer) override;
+		bool HasAudioConsumers() override;
+		// End IPixelStreamingAudioSink
 
-		private:
-			void UpdateAudioSettings(int InNumChannels, int InSampleRate);
+	private:
+		void UpdateAudioSettings(int InNumChannels, int InSampleRate);
 
-		private:
-			TSet<IPixelStreamingAudioConsumer*> AudioConsumers;
-		};
-	}
-}
+	private:
+		TSet<IPixelStreamingAudioConsumer*> AudioConsumers;
+	};
+} // namespace UE::PixelStreaming
\ No newline at end of file
diff --git a/Engine/Plugins/Media/PixelStreaming/Source/PixelStreaming/Private/ClientRTC.cpp b/Engine/Plugins/Media/PixelStreaming/Source/PixelStreaming/Private/ClientRTC.cpp
--- a/Engine/Plugins/Media/PixelStreaming/Source/PixelStreaming/Private/ClientRTC.cpp
+++ b/Engine/Plugins/Media/PixelStreaming/Source/PixelStreaming/Private/ClientRTC.cpp
@@ -35,7 +35,7 @@ namespace
 namespace UE::PixelStreaming
 {
 	FClientRTC::FClientRTC()
-		:State(EState::Disconnected)
+		: State(EState::Disconnected)
 	{
 		FModuleManager::LoadModuleChecked<FWebSocketsModule>("WebSockets");
 		rtc::InitializeSSL();
@@ -57,7 +57,8 @@ namespace UE::PixelStreaming
 			nullptr);													   // audio_processing
 		check(PeerConnectionFactory);
 
-		SignallingServerConnection = MakeUnique<FSignallingServerConnection>(*this, TEXT("UE_INSTANCE"));
+		FSignallingServerConnection::FWebSocketFactory WebSocketFactory = [](const FString& Url) { return FWebSocketsModule::Get().CreateWebSocket(Url, TEXT("")); };
+		SignallingServerConnection = MakeUnique<FSignallingServerConnection>(WebSocketFactory, *this);
 	}
 
 	FClientRTC::~FClientRTC()
@@ -70,6 +71,11 @@ namespace UE::PixelStreaming
 		State = EState::Connecting;
 	}
 
+	void FClientRTC::Disconnect()
+	{
+		PeerConnection = nullptr;
+	}
+
 	bool FClientRTC::SendMessage(Protocol::EToStreamerMsg Type, const FString& Descriptor) const
 	{
 		if (!DataChannel)
@@ -185,6 +191,8 @@ namespace UE::PixelStreaming
 
 		DataChannel = Channel;
 		DataChannel->RegisterObserver(this);
+
+		OnDataChannelOpen.Broadcast(*this);
 	}
 
 	void FClientRTC::OnRenegotiationNeeded()
diff --git a/Engine/Plugins/Media/PixelStreaming/Source/PixelStreaming/Private/ClientRTC.h b/Engine/Plugins/Media/PixelStreaming/Source/PixelStreaming/Private/ClientRTC.h
--- a/Engine/Plugins/Media/PixelStreaming/Source/PixelStreaming/Private/ClientRTC.h
+++ b/Engine/Plugins/Media/PixelStreaming/Source/PixelStreaming/Private/ClientRTC.h
@@ -23,6 +23,7 @@ namespace UE::PixelStreaming
 		};
 
 		void Connect(const FString& Url);
+		void Disconnect();
 		EState GetState() const { return State; }
 
 		bool SendMessage(Protocol::EToStreamerMsg Type, const FString& Descriptor) const;
@@ -33,6 +34,9 @@ namespace UE::PixelStreaming
 		DECLARE_EVENT_OneParam(FClientRTC, FDisconnectedEvent, FClientRTC&);
 		FDisconnectedEvent OnDisconnected;
 
+		DECLARE_EVENT_OneParam(FClientRTC, FDataChannelOpenEvent, FClientRTC&);
+		FDataChannelOpenEvent OnDataChannelOpen;
+
 		DECLARE_EVENT_ThreeParams(FClientRTC, FDataMessageEvent, FClientRTC&, uint8, const FString&);
 		FDataMessageEvent OnDataMessage;
 
