This commit is a large collection (from various authors) of Pixel Streaming 5.1 work

diff --git a/Engine/Plugins/Media/PixelStreaming/Source/PixelStreaming/Private/EncoderFactory.cpp b/Engine/Plugins/Media/PixelStreaming/Source/PixelStreaming/Private/EncoderFactory.cpp
--- a/Engine/Plugins/Media/PixelStreaming/Source/PixelStreaming/Private/EncoderFactory.cpp
+++ b/Engine/Plugins/Media/PixelStreaming/Source/PixelStreaming/Private/EncoderFactory.cpp
@@ -13,191 +13,205 @@
 #include "PixelStreamingPrivate.h"
 #include "Utils.h"
 
-/*
-* ------------- UE::PixelStreaming::FVideoEncoderFactory ---------------
-*/
-
-UE::PixelStreaming::FVideoEncoderFactory::FVideoEncoderFactory()
-{
-}
-
-UE::PixelStreaming::FVideoEncoderFactory::~FVideoEncoderFactory()
+namespace UE::PixelStreaming
 {
-}
+	/*
+	* ------------- UE::PixelStreaming::FVideoEncoderFactory ---------------
+	*/
 
-std::vector<webrtc::SdpVideoFormat> UE::PixelStreaming::FVideoEncoderFactory::GetSupportedFormats() const
-{
-	std::vector<webrtc::SdpVideoFormat> video_formats;
+	FVideoEncoderFactory::FVideoEncoderFactory()
+	{
+	}
 
-	switch (UE::PixelStreaming::Settings::GetSelectedCodec())
+	FVideoEncoderFactory::~FVideoEncoderFactory()
 	{
-		case UE::PixelStreaming::Settings::ECodec::VP8:
-			video_formats.push_back(webrtc::SdpVideoFormat(cricket::kVp8CodecName));
-		case UE::PixelStreaming::Settings::ECodec::VP9:
-			video_formats.push_back(webrtc::SdpVideoFormat(cricket::kVp9CodecName));
-		case UE::PixelStreaming::Settings::ECodec::H264:
-		default:
-			video_formats.push_back(UE::PixelStreaming::CreateH264Format(webrtc::H264::kProfileConstrainedBaseline, webrtc::H264::kLevel3_1));
-			video_formats.push_back(UE::PixelStreaming::CreateH264Format(webrtc::H264::kProfileBaseline, webrtc::H264::kLevel3_1));
 	}
-	return video_formats;
-}
 
-UE::PixelStreaming::FVideoEncoderFactory::CodecInfo UE::PixelStreaming::FVideoEncoderFactory::QueryVideoEncoder(const webrtc::SdpVideoFormat& format) const
-{
-	webrtc::VideoEncoderFactory::CodecInfo CodecInfo = { false, false };
-	CodecInfo.is_hardware_accelerated = true;
-	CodecInfo.has_internal_source = false;
-	return CodecInfo;
-}
+	std::vector<webrtc::SdpVideoFormat> FVideoEncoderFactory::GetSupportedFormats() const
+	{
+		std::vector<webrtc::SdpVideoFormat> video_formats;
 
-std::unique_ptr<webrtc::VideoEncoder> UE::PixelStreaming::FVideoEncoderFactory::CreateVideoEncoder(const webrtc::SdpVideoFormat& format)
-{
-	if (absl::EqualsIgnoreCase(format.name, cricket::kVp8CodecName))
+		switch (UE::PixelStreaming::Settings::GetSelectedCodec())
+		{
+			case UE::PixelStreaming::Settings::ECodec::VP8:
+				video_formats.push_back(webrtc::SdpVideoFormat(cricket::kVp8CodecName));
+			case UE::PixelStreaming::Settings::ECodec::VP9:
+				video_formats.push_back(webrtc::SdpVideoFormat(cricket::kVp9CodecName));
+			case UE::PixelStreaming::Settings::ECodec::H264:
+			default:
+				video_formats.push_back(UE::PixelStreaming::CreateH264Format(webrtc::H264::kProfileConstrainedBaseline, webrtc::H264::kLevel3_1));
+				video_formats.push_back(UE::PixelStreaming::CreateH264Format(webrtc::H264::kProfileBaseline, webrtc::H264::kLevel3_1));
+		}
+		return video_formats;
+	}
+
+	FVideoEncoderFactory::CodecInfo FVideoEncoderFactory::QueryVideoEncoder(const webrtc::SdpVideoFormat& format) const
 	{
-		return webrtc::VP8Encoder::Create();
+		webrtc::VideoEncoderFactory::CodecInfo CodecInfo = { false, false };
+		CodecInfo.is_hardware_accelerated = true;
+		CodecInfo.has_internal_source = false;
+		return CodecInfo;
 	}
-	else if (absl::EqualsIgnoreCase(format.name, cricket::kVp9CodecName))
+
+	std::unique_ptr<webrtc::VideoEncoder> FVideoEncoderFactory::CreateVideoEncoder(const webrtc::SdpVideoFormat& format)
 	{
-		return webrtc::VP9Encoder::Create();
+		if (absl::EqualsIgnoreCase(format.name, cricket::kVp8CodecName))
+		{
+			return webrtc::VP8Encoder::Create();
+		}
+		else if (absl::EqualsIgnoreCase(format.name, cricket::kVp9CodecName))
+		{
+			return webrtc::VP9Encoder::Create();
+		}
+		else
+		{
+			// Lock during encoder creation
+			FScopeLock Lock(&ActiveEncodersGuard);
+			auto VideoEncoder = std::make_unique<FVideoEncoderRTC>(*this);
+			ActiveEncoders.Add(VideoEncoder.get());
+			return VideoEncoder;
+		}
 	}
-	else
+
+	void FVideoEncoderFactory::OnEncodedImage(const webrtc::EncodedImage& encoded_image, const webrtc::CodecSpecificInfo* codec_specific_info, const webrtc::RTPFragmentationHeader* fragmentation)
 	{
-		// Lock during encoder creation
+		// Lock as we send encoded image to each encoder.
 		FScopeLock Lock(&ActiveEncodersGuard);
-		auto VideoEncoder = std::make_unique<UE::PixelStreaming::FVideoEncoderRTC>(*this);
-		ActiveEncoders.Add(VideoEncoder.get());
-		return VideoEncoder;
-	}
-}
 
-void UE::PixelStreaming::FVideoEncoderFactory::OnEncodedImage(const webrtc::EncodedImage& encoded_image, const webrtc::CodecSpecificInfo* codec_specific_info, const webrtc::RTPFragmentationHeader* fragmentation)
-{
-	// Lock as we send encoded image to each encoder.
-	FScopeLock Lock(&ActiveEncodersGuard);
+		// Go through each encoder and send our encoded image to its callback
+		for (FVideoEncoderRTC* Encoder : ActiveEncoders)
+		{
+			Encoder->SendEncodedImage(encoded_image, codec_specific_info, fragmentation);
+		}
+	}
 
-	// Go through each encoder and send our encoded image to its callback
-	for (UE::PixelStreaming::FVideoEncoderRTC* Encoder : ActiveEncoders)
+	void FVideoEncoderFactory::ReleaseVideoEncoder(FVideoEncoderRTC* Encoder)
 	{
-		Encoder->SendEncodedImage(encoded_image, codec_specific_info, fragmentation);
+		// Lock during deleting an encoder
+		FScopeLock Lock(&ActiveEncodersGuard);
+		ActiveEncoders.Remove(Encoder);
 	}
-}
 
-void UE::PixelStreaming::FVideoEncoderFactory::ReleaseVideoEncoder(UE::PixelStreaming::FVideoEncoderRTC* Encoder)
-{
-	// Lock during deleting an encoder
-	FScopeLock Lock(&ActiveEncodersGuard);
-	ActiveEncoders.Remove(Encoder);
-}
-
-void UE::PixelStreaming::FVideoEncoderFactory::ForceKeyFrame()
-{
-	FScopeLock Lock(&ActiveEncodersGuard);
-	HardwareEncoder->SetForceNextKeyframe();
-}
+	void FVideoEncoderFactory::ForceKeyFrame()
+	{
+		FScopeLock Lock(&ActiveEncodersGuard);
+		HardwareEncoder->SetForceNextKeyframe();
+	}
 
-UE::PixelStreaming::FVideoEncoderH264Wrapper* UE::PixelStreaming::FVideoEncoderFactory::GetOrCreateHardwareEncoder(int Width, int Height, int MaxBitrate, int TargetBitrate, int MaxFramerate)
-{
-	if (HardwareEncoder == nullptr)
-	{
-		// Make AVEncoder frame factory.
-		TUniquePtr<FEncoderFrameFactory> FrameFactory = MakeUnique<FEncoderFrameFactory>();
-
-		// Make the encoder config
-		AVEncoder::FVideoEncoder::FLayerConfig EncoderConfig;
-		EncoderConfig.Width = Width;
-		EncoderConfig.Height = Height;
-		EncoderConfig.MaxFramerate = MaxFramerate;
-		EncoderConfig.TargetBitrate = TargetBitrate;
-		EncoderConfig.MaxBitrate = MaxBitrate;
-
-		// Make the actual AVEncoder encoder.
-		const TArray<AVEncoder::FVideoEncoderInfo>& Available = AVEncoder::FVideoEncoderFactory::Get().GetAvailable();
-		TUniquePtr<AVEncoder::FVideoEncoder> Encoder = AVEncoder::FVideoEncoderFactory::Get().Create(Available[0].ID, FrameFactory->GetOrCreateVideoEncoderInput(), EncoderConfig);
-		if (Encoder.IsValid())
+	FVideoEncoderH264Wrapper* FVideoEncoderFactory::GetOrCreateHardwareEncoder(int Width, int Height, int MaxBitrate, int TargetBitrate, int MaxFramerate)
+	{
+		if (HardwareEncoder == nullptr)
 		{
-			Encoder->SetOnEncodedPacket([this](uint32 InLayerIndex, const TSharedPtr<AVEncoder::FVideoEncoderInputFrame> InFrame, const AVEncoder::FCodecPacket& InPacket) {
-				// Note: this is a static method call.
-				FVideoEncoderH264Wrapper::OnEncodedPacket(this, InLayerIndex, InFrame, InPacket);
-			});
-
-			// Make the hardware encoder wrapper
-			HardwareEncoder = MakeUnique<FVideoEncoderH264Wrapper>(MoveTemp(FrameFactory), MoveTemp(Encoder));
-			return HardwareEncoder.Get();
+			// Make AVEncoder frame factory.
+			TUniquePtr<FEncoderFrameFactory> FrameFactory = MakeUnique<FEncoderFrameFactory>();
+
+			// Make the encoder config
+			AVEncoder::FVideoEncoder::FLayerConfig EncoderConfig;
+			EncoderConfig.Width = Width;
+			EncoderConfig.Height = Height;
+			EncoderConfig.MaxFramerate = MaxFramerate;
+			EncoderConfig.TargetBitrate = TargetBitrate;
+			EncoderConfig.MaxBitrate = MaxBitrate;
+
+			// Make the actual AVEncoder encoder.
+			const TArray<AVEncoder::FVideoEncoderInfo>& Available = AVEncoder::FVideoEncoderFactory::Get().GetAvailable();
+			TUniquePtr<AVEncoder::FVideoEncoder> Encoder = AVEncoder::FVideoEncoderFactory::Get().Create(Available[0].ID, FrameFactory->GetOrCreateVideoEncoderInput(), EncoderConfig);
+			if (Encoder.IsValid())
+			{
+				Encoder->SetOnEncodedPacket([this](uint32 InLayerIndex, const TSharedPtr<AVEncoder::FVideoEncoderInputFrame> InFrame, const AVEncoder::FCodecPacket& InPacket) {
+					// Note: this is a static method call.
+					FVideoEncoderH264Wrapper::OnEncodedPacket(this, InLayerIndex, InFrame, InPacket);
+				});
+
+				// Make the hardware encoder wrapper
+				HardwareEncoder = MakeUnique<FVideoEncoderH264Wrapper>(MoveTemp(FrameFactory), MoveTemp(Encoder));
+				return HardwareEncoder.Get();
+			}
+			else
+			{
+				UE_LOG(LogPixelStreaming, Error, TEXT("Could not create encoder. Check encoder config or perhaps you used up all your HW encoders."));
+				// We could not make the encoder, so indicate the id was not set successfully.
+				return nullptr;
+			}
 		}
 		else
 		{
-			UE_LOG(LogPixelStreaming, Error, TEXT("Could not create encoder. Check encoder config or perhaps you used up all your HW encoders."));
-			// We could not make the encoder, so indicate the id was not set successfully.
-			return nullptr;
+			return HardwareEncoder.Get();
 		}
 	}
-	else
+
+	FVideoEncoderH264Wrapper* FVideoEncoderFactory::GetHardwareEncoder()
 	{
 		return HardwareEncoder.Get();
 	}
-}
-
-UE::PixelStreaming::FVideoEncoderH264Wrapper* UE::PixelStreaming::FVideoEncoderFactory::GetHardwareEncoder()
-{
-	return HardwareEncoder.Get();
-}
 
-/*
-* ------------- UE::PixelStreaming::FSimulcastEncoderFactory ---------------
-*/
+	/*
+	* ------------- UE::PixelStreaming::FSimulcastEncoderFactory ---------------
+	*/
 
-UE::PixelStreaming::FSimulcastEncoderFactory::FSimulcastEncoderFactory()
-	: PrimaryEncoderFactory(MakeUnique<UE::PixelStreaming::FVideoEncoderFactory>())
-{
-	// Make a copy of simulcast settings and sort them based on scaling.
-	for (FEncoderFactoryId i = 0; i < Settings::SimulcastParameters.Layers.Num(); i++)
+	FSimulcastEncoderFactory::FSimulcastEncoderFactory()
+		: PrimaryEncoderFactory(MakeUnique<FVideoEncoderFactory>())
 	{
-		GetOrCreateEncoderFactory(i);
+		// Make a copy of simulcast settings and sort them based on scaling.
+		for (FEncoderFactoryId i = 0; i < Settings::SimulcastParameters.Layers.Num(); i++)
+		{
+			GetOrCreateEncoderFactory(i);
+		}
 	}
-}
-
-UE::PixelStreaming::FSimulcastEncoderFactory::~FSimulcastEncoderFactory()
-{
-}
 
-std::unique_ptr<webrtc::VideoEncoder> UE::PixelStreaming::FSimulcastEncoderFactory::CreateVideoEncoder(const webrtc::SdpVideoFormat& format)
-{
-	return std::make_unique<UE::PixelStreaming::FSimulcastEncoderAdapter>(*this, format);
-}
+	FSimulcastEncoderFactory::~FSimulcastEncoderFactory()
+	{
+	}
 
-UE::PixelStreaming::FVideoEncoderFactory* UE::PixelStreaming::FSimulcastEncoderFactory::GetEncoderFactory(FEncoderFactoryId Id)
-{
-	FScopeLock Lock(&EncoderFactoriesGuard);
-	if (auto&& Existing = EncoderFactories.Find(Id))
+	std::unique_ptr<webrtc::VideoEncoder> FSimulcastEncoderFactory::CreateVideoEncoder(const webrtc::SdpVideoFormat& format)
 	{
-		return Existing->Get();
+		if (absl::EqualsIgnoreCase(format.name, cricket::kVp8CodecName))
+		{
+			return webrtc::VP8Encoder::Create();
+		}
+		else if (absl::EqualsIgnoreCase(format.name, cricket::kVp9CodecName))
+		{
+			return webrtc::VP9Encoder::Create();
+		}
+		else
+		{
+			return std::make_unique<FSimulcastEncoderAdapter>(*this, format);
+		}
 	}
-	return nullptr;
-}
 
-UE::PixelStreaming::FVideoEncoderFactory* UE::PixelStreaming::FSimulcastEncoderFactory::GetOrCreateEncoderFactory(FEncoderFactoryId Id)
-{
-	UE::PixelStreaming::FVideoEncoderFactory* Existing = GetEncoderFactory(Id);
-	if (Existing == nullptr)
+	FVideoEncoderFactory* FSimulcastEncoderFactory::GetEncoderFactory(FEncoderFactoryId Id)
 	{
 		FScopeLock Lock(&EncoderFactoriesGuard);
-		TUniquePtr<FVideoEncoderFactory> EncoderFactory = MakeUnique<UE::PixelStreaming::FVideoEncoderFactory>();
-		EncoderFactories.Add(Id, MoveTemp(EncoderFactory));
-		return EncoderFactory.Get();
+		if (auto&& Existing = EncoderFactories.Find(Id))
+		{
+			return Existing->Get();
+		}
+		return nullptr;
 	}
-	else
+
+	FVideoEncoderFactory* FSimulcastEncoderFactory::GetOrCreateEncoderFactory(FEncoderFactoryId Id)
 	{
-		return Existing;
+		FVideoEncoderFactory* Existing = GetEncoderFactory(Id);
+		if (Existing == nullptr)
+		{
+			FScopeLock Lock(&EncoderFactoriesGuard);
+			TUniquePtr<FVideoEncoderFactory> EncoderFactory = MakeUnique<FVideoEncoderFactory>();
+			EncoderFactories.Add(Id, MoveTemp(EncoderFactory));
+			return EncoderFactory.Get();
+		}
+		else
+		{
+			return Existing;
+		}
 	}
-}
 
-std::vector<webrtc::SdpVideoFormat> UE::PixelStreaming::FSimulcastEncoderFactory::GetSupportedFormats() const
-{
-	return PrimaryEncoderFactory->GetSupportedFormats();
-}
+	std::vector<webrtc::SdpVideoFormat> FSimulcastEncoderFactory::GetSupportedFormats() const
+	{
+		return PrimaryEncoderFactory->GetSupportedFormats();
+	}
 
-UE::PixelStreaming::FSimulcastEncoderFactory::CodecInfo UE::PixelStreaming::FSimulcastEncoderFactory::QueryVideoEncoder(const webrtc::SdpVideoFormat& format) const
-{
-	return PrimaryEncoderFactory->QueryVideoEncoder(format);
-}
+	FSimulcastEncoderFactory::CodecInfo FSimulcastEncoderFactory::QueryVideoEncoder(const webrtc::SdpVideoFormat& format) const
+	{
+		return PrimaryEncoderFactory->QueryVideoEncoder(format);
+	}
+}
\ No newline at end of file
diff --git a/Engine/Plugins/Media/PixelStreaming/Source/PixelStreaming/Private/EncoderFactory.h b/Engine/Plugins/Media/PixelStreaming/Source/PixelStreaming/Private/EncoderFactory.h
--- a/Engine/Plugins/Media/PixelStreaming/Source/PixelStreaming/Private/EncoderFactory.h
+++ b/Engine/Plugins/Media/PixelStreaming/Source/PixelStreaming/Private/EncoderFactory.h
@@ -8,64 +8,60 @@
 #include "HAL/CriticalSection.h"
 #include "Misc/Optional.h"
 
-namespace UE {
-	namespace PixelStreaming {
-		class FVideoEncoderFactory;
-		class FVideoEncoderRTC;
-		class FVideoEncoderH264Wrapper;
-
-		class FSimulcastEncoderFactory : public webrtc::VideoEncoderFactory
-		{
-		public:
-			FSimulcastEncoderFactory();
-			virtual ~FSimulcastEncoderFactory();
-			virtual std::vector<webrtc::SdpVideoFormat> GetSupportedFormats() const ;
-			virtual CodecInfo QueryVideoEncoder(const webrtc::SdpVideoFormat& format) const override;
-			virtual std::unique_ptr<webrtc::VideoEncoder> CreateVideoEncoder(const webrtc::SdpVideoFormat& format);
-
-
-			using FEncoderFactoryId = uint64;
-			FVideoEncoderFactory* GetOrCreateEncoderFactory(FEncoderFactoryId Id);
-			FVideoEncoderFactory* GetEncoderFactory(FEncoderFactoryId Id);
-		private:
-			TMap<FEncoderFactoryId, TUniquePtr<FVideoEncoderFactory>> EncoderFactories;
-			FCriticalSection EncoderFactoriesGuard;
-
-			TUniquePtr<FVideoEncoderFactory> PrimaryEncoderFactory;
-		};
-
-		class FVideoEncoderFactory : public webrtc::VideoEncoderFactory
-		{
-		public:
-			FVideoEncoderFactory();
-			virtual ~FVideoEncoderFactory() override;
-
-			virtual std::vector<webrtc::SdpVideoFormat> GetSupportedFormats() const override;
-
-			// Always returns our H264 hardware encoders codec_info for now/
-			virtual CodecInfo QueryVideoEncoder(const webrtc::SdpVideoFormat& Format) const override;
-
-			// Always returns our H264 hardware encoder for now.
-			virtual std::unique_ptr<webrtc::VideoEncoder> CreateVideoEncoder(const webrtc::SdpVideoFormat& Format) override;
-
-			void ReleaseVideoEncoder(FVideoEncoderRTC* Encoder);
-			void ForceKeyFrame();
-
+namespace UE::PixelStreaming 
+{
+	class FVideoEncoderFactory;
+	class FVideoEncoderRTC;
+	class FVideoEncoderH264Wrapper;
+
+	class FSimulcastEncoderFactory : public webrtc::VideoEncoderFactory
+	{
+	public:
+		FSimulcastEncoderFactory();
+		virtual ~FSimulcastEncoderFactory();
+		virtual std::vector<webrtc::SdpVideoFormat> GetSupportedFormats() const ;
+		virtual CodecInfo QueryVideoEncoder(const webrtc::SdpVideoFormat& format) const override;
+		virtual std::unique_ptr<webrtc::VideoEncoder> CreateVideoEncoder(const webrtc::SdpVideoFormat& format);
+
+		using FEncoderFactoryId = uint64;
+		FVideoEncoderFactory* GetOrCreateEncoderFactory(FEncoderFactoryId Id);
+		FVideoEncoderFactory* GetEncoderFactory(FEncoderFactoryId Id);
+	private:
+		TMap<FEncoderFactoryId, TUniquePtr<FVideoEncoderFactory>> EncoderFactories;
+		FCriticalSection EncoderFactoriesGuard;
+
+		TUniquePtr<FVideoEncoderFactory> PrimaryEncoderFactory;
+	};
+
+	class FVideoEncoderFactory : public webrtc::VideoEncoderFactory
+	{
+	public:
+		FVideoEncoderFactory();
+		virtual ~FVideoEncoderFactory() override;
+
+		virtual std::vector<webrtc::SdpVideoFormat> GetSupportedFormats() const override;
+
+		// Always returns our H264 hardware encoders codec_info for now/
+		virtual CodecInfo QueryVideoEncoder(const webrtc::SdpVideoFormat& Format) const override;
+
+		// Always returns our H264 hardware encoder for now.
+		virtual std::unique_ptr<webrtc::VideoEncoder> CreateVideoEncoder(const webrtc::SdpVideoFormat& Format) override;
+
+		void ReleaseVideoEncoder(FVideoEncoderRTC* Encoder);
+		void ForceKeyFrame();
 			
-			FVideoEncoderH264Wrapper* GetOrCreateHardwareEncoder(int Width, int Height, int MaxBitrate, int TargetBitrate, int MaxFramerate);
-			FVideoEncoderH264Wrapper* GetHardwareEncoder();
-
-			void OnEncodedImage(const webrtc::EncodedImage& Encoded_image, const webrtc::CodecSpecificInfo* CodecSpecificInfo, const webrtc::RTPFragmentationHeader* Fragmentation);
+		FVideoEncoderH264Wrapper* GetOrCreateHardwareEncoder(int Width, int Height, int MaxBitrate, int TargetBitrate, int MaxFramerate);
+		FVideoEncoderH264Wrapper* GetHardwareEncoder();
 
-		private:
+		void OnEncodedImage(const webrtc::EncodedImage& Encoded_image, const webrtc::CodecSpecificInfo* CodecSpecificInfo, const webrtc::RTPFragmentationHeader* Fragmentation);
 
-			static webrtc::SdpVideoFormat CreateH264Format(webrtc::H264::Profile Profile, webrtc::H264::Level Level);
+	private:
+		static webrtc::SdpVideoFormat CreateH264Format(webrtc::H264::Profile Profile, webrtc::H264::Level Level);
 
-			TUniquePtr<FVideoEncoderH264Wrapper> HardwareEncoder;
+		TUniquePtr<FVideoEncoderH264Wrapper> HardwareEncoder;
 
-			// Encoders assigned to each peer. Each one of these will be assigned to one of the hardware encoders
-			TArray<FVideoEncoderRTC*> ActiveEncoders;
-			FCriticalSection ActiveEncodersGuard;
-		};
-	}
-}
+		// Encoders assigned to each peer. Each one of these will be assigned to one of the hardware encoders
+		TArray<FVideoEncoderRTC*> ActiveEncoders;
+		FCriticalSection ActiveEncodersGuard;
+	};
+} // namespace UE::PixelStreaming
\ No newline at end of file
diff --git a/Engine/Plugins/Media/PixelStreaming/Source/PixelStreaming/Private/EncoderFrameFactory.cpp b/Engine/Plugins/Media/PixelStreaming/Source/PixelStreaming/Private/EncoderFrameFactory.cpp
--- a/Engine/Plugins/Media/PixelStreaming/Source/PixelStreaming/Private/EncoderFrameFactory.cpp
+++ b/Engine/Plugins/Media/PixelStreaming/Source/PixelStreaming/Private/EncoderFrameFactory.cpp
@@ -11,6 +11,7 @@
 
 #if PLATFORM_WINDOWS
 	#include "VideoCommon.h"
+	#include "ShaderCore.h"
 	#include "D3D11State.h"
 	#include "D3D11Resources.h"
 	#include "ID3D12DynamicRHI.h"
@@ -21,22 +22,24 @@ THIRD_PARTY_INCLUDES_END
 	#include "Windows/HideWindowsPlatformTypes.h"
 #endif
 
-UE::PixelStreaming::FEncoderFrameFactory::FEncoderFrameFactory()
+namespace UE::PixelStreaming
 {
-}
+	FEncoderFrameFactory::FEncoderFrameFactory()
+	{
+	}
 
-UE::PixelStreaming::FEncoderFrameFactory::~FEncoderFrameFactory()
-{
+	FEncoderFrameFactory::~FEncoderFrameFactory()
+	{
 	FlushFrames();
-}
+	}
 
-void UE::PixelStreaming::FEncoderFrameFactory::FlushFrames()
-{
+	void FEncoderFrameFactory::FlushFrames()
+	{
 	TextureToFrameMapping.Empty();
-}
+	}
 
-void UE::PixelStreaming::FEncoderFrameFactory::RemoveStaleTextures()
-{
+	void FEncoderFrameFactory::RemoveStaleTextures()
+	{
 	// Remove any textures whose only reference is the one held by this class
 
 	TMap<FTextureRHIRef, TSharedPtr<AVEncoder::FVideoEncoderInputFrame>>::TIterator Iter = TextureToFrameMapping.CreateIterator();
@@ -48,10 +51,10 @@ void UE::PixelStreaming::FEncoderFrameFactory::RemoveStaleTextures()
 			Iter.RemoveCurrent();
 		}
 	}
-}
+	}
 
-TSharedPtr<AVEncoder::FVideoEncoderInputFrame> UE::PixelStreaming::FEncoderFrameFactory::GetOrCreateFrame(const FTextureRHIRef InTexture)
-{
+	TSharedPtr<AVEncoder::FVideoEncoderInputFrame> FEncoderFrameFactory::GetOrCreateFrame(const FTextureRHIRef InTexture)
+	{
 	check(EncoderInput.IsValid());
 
 	RemoveStaleTextures();
@@ -76,33 +79,33 @@ TSharedPtr<AVEncoder::FVideoEncoderInputFrame> UE::PixelStreaming::FEncoderFrame
 		TextureToFrameMapping.Add(InTexture, OutFrame);
 	}
 
-	OutFrame->SetWidth(InTexture->GetSizeX());
-	OutFrame->SetHeight(InTexture->GetSizeY());
+	OutFrame->SetWidth(InTexture->GetDesc().Extent.X);
+	OutFrame->SetHeight(InTexture->GetDesc().Extent.Y);
 	OutFrame->SetFrameID(++FrameId);
 	return OutFrame;
-}
+	}
 
-TSharedPtr<AVEncoder::FVideoEncoderInputFrame> UE::PixelStreaming::FEncoderFrameFactory::GetFrameAndSetTexture(FTextureRHIRef InTexture)
-{
+	TSharedPtr<AVEncoder::FVideoEncoderInputFrame> FEncoderFrameFactory::GetFrameAndSetTexture(FTextureRHIRef InTexture)
+	{
 	check(EncoderInput.IsValid());
 
 	TSharedPtr<AVEncoder::FVideoEncoderInputFrame> Frame = GetOrCreateFrame(InTexture);
 
 	return Frame;
-}
+	}
 
-TSharedPtr<AVEncoder::FVideoEncoderInput> UE::PixelStreaming::FEncoderFrameFactory::GetOrCreateVideoEncoderInput()
-{
+	TSharedPtr<AVEncoder::FVideoEncoderInput> FEncoderFrameFactory::GetOrCreateVideoEncoderInput()
+	{
 	if (!EncoderInput.IsValid())
 	{
 		EncoderInput = CreateVideoEncoderInput();
 	}
 
 	return EncoderInput;
-}
+	}
 
-TSharedPtr<AVEncoder::FVideoEncoderInput> UE::PixelStreaming::FEncoderFrameFactory::CreateVideoEncoderInput() const
-{
+	TSharedPtr<AVEncoder::FVideoEncoderInput> FEncoderFrameFactory::CreateVideoEncoderInput() const
+	{
 	if (!GDynamicRHI)
 	{
 		UE_LOG(LogPixelStreaming, Error, TEXT("GDynamicRHI not valid for some reason."));
@@ -136,7 +139,7 @@ TSharedPtr<AVEncoder::FVideoEncoderInput> UE::PixelStreaming::FEncoderFrameFacto
 			}
 		}
 	}
-#if PLATFORM_WINDOWS
+	#if PLATFORM_WINDOWS
 	else if (RHIType == ERHIInterfaceType::D3D11)
 	{
 		if (IsRHIDeviceAMD())
@@ -159,14 +162,14 @@ TSharedPtr<AVEncoder::FVideoEncoderInput> UE::PixelStreaming::FEncoderFrameFacto
 			return AVEncoder::FVideoEncoderInput::CreateForCUDA(FModuleManager::GetModuleChecked<FCUDAModule>("CUDA").GetCudaContext(), bIsResizable);
 		}
 	}
-#endif
+	#endif
 
 	UE_LOG(LogPixelStreaming, Error, TEXT("Current RHI %s is not supported in Pixel Streaming"), GDynamicRHI->GetName());
 	return nullptr;
-}
+	}
 
-void UE::PixelStreaming::FEncoderFrameFactory::SetTexture(TSharedPtr<AVEncoder::FVideoEncoderInputFrame> InputFrame, const FTextureRHIRef& Texture)
-{
+	void FEncoderFrameFactory::SetTexture(TSharedPtr<AVEncoder::FVideoEncoderInputFrame> InputFrame, const FTextureRHIRef& Texture)
+	{
 	const ERHIInterfaceType RHIType = RHIGetInterfaceType();
 
 	// VULKAN
@@ -186,7 +189,7 @@ void UE::PixelStreaming::FEncoderFrameFactory::SetTexture(TSharedPtr<AVEncoder::
 			UE_LOG(LogPixelStreaming, Error, TEXT("Pixel Streaming only supports AMD and NVIDIA devices, this device is neither of those."));
 		}
 	}
-#if PLATFORM_WINDOWS
+	#if PLATFORM_WINDOWS
 	// TODO: Fix CUDA DX11 (Using CUDA as a bridge between DX11 and NVENC currently produces garbled results)
 	// DX11
 	else if (RHIType == ERHIInterfaceType::D3D11)
@@ -209,16 +212,17 @@ void UE::PixelStreaming::FEncoderFrameFactory::SetTexture(TSharedPtr<AVEncoder::
 			UE_LOG(LogPixelStreaming, Error, TEXT("Pixel Streaming only supports AMD and NVIDIA devices, this device is neither of those."));
 		}
 	}
-#endif // PLATFORM_WINDOWS
+	#endif // PLATFORM_WINDOWS
 	else
 	{
 		UE_LOG(LogPixelStreaming, Error, TEXT("Pixel Streaming does not support this RHI - %s"), GDynamicRHI->GetName());
 	}
-}
+	}
 
-void UE::PixelStreaming::FEncoderFrameFactory::SetTextureCUDAVulkan(TSharedPtr<AVEncoder::FVideoEncoderInputFrame> InputFrame, const FTextureRHIRef& Texture)
-{
+	void FEncoderFrameFactory::SetTextureCUDAVulkan(TSharedPtr<AVEncoder::FVideoEncoderInputFrame> InputFrame, const FTextureRHIRef& Texture)
+	{
 	IVulkanDynamicRHI* VulkanRHI = GetIVulkanDynamicRHI();
+
 	const FVulkanRHIAllocationInfo TextureAllocationInfo = VulkanRHI->RHIGetAllocationInfo(Texture.GetReference());
 
 	VkDevice Device = VulkanRHI->RHIGetVkDevice();
@@ -264,7 +268,7 @@ void UE::PixelStreaming::FEncoderFrameFactory::SetTextureCUDAVulkan(TSharedPtr<A
 
 	// Only store handle to be closed on frame destruction if it is an NT handle
 	Handle = bUseNTHandle ? Handle : NULL;
-#else
+	#else
 	void* Handle = nullptr;
 
 	// Get the CUarray to that textures memory making sure the clear it when done
@@ -301,7 +305,7 @@ void UE::PixelStreaming::FEncoderFrameFactory::SetTextureCUDAVulkan(TSharedPtr<A
 		}
 	}
 
-#endif
+	#endif
 
 	CUmipmappedArray MappedMipArray = nullptr;
 	CUarray MappedArray = nullptr;
@@ -310,8 +314,8 @@ void UE::PixelStreaming::FEncoderFrameFactory::SetTextureCUDAVulkan(TSharedPtr<A
 		CUDA_EXTERNAL_MEMORY_MIPMAPPED_ARRAY_DESC MipmapDesc = {};
 		MipmapDesc.numLevels = 1;
 		MipmapDesc.offset = TextureAllocationInfo.Offset;
-		MipmapDesc.arrayDesc.Width = Texture->GetSizeX();
-		MipmapDesc.arrayDesc.Height = Texture->GetSizeY();
+		MipmapDesc.arrayDesc.Width = Texture->GetDesc().Extent.X;
+		MipmapDesc.arrayDesc.Height = Texture->GetDesc().Extent.Y;
 		MipmapDesc.arrayDesc.Depth = 0;
 		MipmapDesc.arrayDesc.NumChannels = 4;
 		MipmapDesc.arrayDesc.Format = CU_AD_FORMAT_UNSIGNED_INT8;
@@ -367,11 +371,11 @@ void UE::PixelStreaming::FEncoderFrameFactory::SetTextureCUDAVulkan(TSharedPtr<A
 
 		FCUDAModule::CUDA().cuCtxPopCurrent(NULL);
 	});
-}
+	}
 
-#if PLATFORM_WINDOWS
-void UE::PixelStreaming::FEncoderFrameFactory::SetTextureCUDAD3D11(TSharedPtr<AVEncoder::FVideoEncoderInputFrame> InputFrame, const FTextureRHIRef& Texture)
-{
+	#if PLATFORM_WINDOWS
+	void FEncoderFrameFactory::SetTextureCUDAD3D11(TSharedPtr<AVEncoder::FVideoEncoderInputFrame> InputFrame, const FTextureRHIRef& Texture)
+	{
 	FD3D11Texture* D3D11Texture = GetD3D11TextureFromRHITexture(Texture);
 	unsigned long long TextureMemorySize = D3D11Texture->GetMemorySize();
 
@@ -417,8 +421,8 @@ void UE::PixelStreaming::FEncoderFrameFactory::SetTextureCUDAD3D11(TSharedPtr<AV
 		CUDA_EXTERNAL_MEMORY_MIPMAPPED_ARRAY_DESC MipmapDesc = {};
 		MipmapDesc.numLevels = 1;
 		MipmapDesc.offset = 0;
-		MipmapDesc.arrayDesc.Width = Texture->GetSizeX();
-		MipmapDesc.arrayDesc.Height = Texture->GetSizeY();
+		MipmapDesc.arrayDesc.Width = Texture->GetDesc().Extent.X;
+		MipmapDesc.arrayDesc.Height = Texture->GetDesc().Extent.Y;
 		MipmapDesc.arrayDesc.Depth = 1;
 		MipmapDesc.arrayDesc.NumChannels = 4;
 		MipmapDesc.arrayDesc.Format = CU_AD_FORMAT_UNSIGNED_INT8;
@@ -474,10 +478,10 @@ void UE::PixelStreaming::FEncoderFrameFactory::SetTextureCUDAD3D11(TSharedPtr<AV
 
 		FCUDAModule::CUDA().cuCtxPopCurrent(NULL);
 	});
-}
+	}
 
-void UE::PixelStreaming::FEncoderFrameFactory::SetTextureCUDAD3D12(TSharedPtr<AVEncoder::FVideoEncoderInputFrame> InputFrame, const FTextureRHIRef& Texture)
-{
+	void FEncoderFrameFactory::SetTextureCUDAD3D12(TSharedPtr<AVEncoder::FVideoEncoderInputFrame> InputFrame, const FTextureRHIRef& Texture)
+	{
 	ID3D12Resource* NativeD3D12Resource = GetID3D12DynamicRHI()->RHIGetResource(Texture);
 	const int64 TextureMemorySize = GetID3D12DynamicRHI()->RHIGetResourceMemorySize(Texture);
 
@@ -530,8 +534,8 @@ void UE::PixelStreaming::FEncoderFrameFactory::SetTextureCUDAD3D12(TSharedPtr<AV
 		CUDA_EXTERNAL_MEMORY_MIPMAPPED_ARRAY_DESC MipmapDesc = {};
 		MipmapDesc.numLevels = 1;
 		MipmapDesc.offset = 0;
-		MipmapDesc.arrayDesc.Width = Texture->GetSizeX();
-		MipmapDesc.arrayDesc.Height = Texture->GetSizeY();
+		MipmapDesc.arrayDesc.Width = Texture->GetDesc().Extent.X;
+		MipmapDesc.arrayDesc.Height = Texture->GetDesc().Extent.Y;
 		MipmapDesc.arrayDesc.Depth = 1;
 		MipmapDesc.arrayDesc.NumChannels = 4;
 		MipmapDesc.arrayDesc.Format = CU_AD_FORMAT_UNSIGNED_INT8;
@@ -587,5 +591,6 @@ void UE::PixelStreaming::FEncoderFrameFactory::SetTextureCUDAD3D12(TSharedPtr<AV
 
 		FCUDAModule::CUDA().cuCtxPopCurrent(NULL);
 	});
-}
-#endif //PLATFORM_WINDOWS
\ No newline at end of file
+	}
+	#endif //PLATFORM_WINDOWS
+} // namespace UE::PixelStreaming
\ No newline at end of file
diff --git a/Engine/Plugins/Media/PixelStreaming/Source/PixelStreaming/Private/EncoderFrameFactory.h b/Engine/Plugins/Media/PixelStreaming/Source/PixelStreaming/Private/EncoderFrameFactory.h
--- a/Engine/Plugins/Media/PixelStreaming/Source/PixelStreaming/Private/EncoderFrameFactory.h
+++ b/Engine/Plugins/Media/PixelStreaming/Source/PixelStreaming/Private/EncoderFrameFactory.h
@@ -8,41 +8,38 @@
 
 class FRHITexture;
 
-namespace UE
+namespace UE::PixelStreaming
 {
-	namespace PixelStreaming
+	/*
+	* Factory to create `AVEncoder::FVideoEncoderInputFrame` for use in Pixel Streaming encoders.
+	* This class is responsible for creating/managing the underlying `AVEncoder::FVideoEncoderInput` required to create
+	* `AVEncoder::FVideoEncoderInputFrame`.
+	* This class is not threadsafe and should only be accessed from a single thread.
+	*/
+	class FEncoderFrameFactory
 	{
-		/*
-		* Factory to create `AVEncoder::FVideoEncoderInputFrame` for use in Pixel Streaming encoders.
-		* This class is responsible for creating/managing the underlying `AVEncoder::FVideoEncoderInput` required to create
-		* `AVEncoder::FVideoEncoderInputFrame`.
-		* This class is not threadsafe and should only be accessed from a single thread.
-		*/
-		class FEncoderFrameFactory
-		{
-		public:
-			FEncoderFrameFactory();
-			~FEncoderFrameFactory();
-			TSharedPtr<AVEncoder::FVideoEncoderInputFrame> GetFrameAndSetTexture(TRefCountPtr<FRHITexture> InTexture);
-			TSharedPtr<AVEncoder::FVideoEncoderInput> GetOrCreateVideoEncoderInput();
+	public:
+		FEncoderFrameFactory();
+		~FEncoderFrameFactory();
+		TSharedPtr<AVEncoder::FVideoEncoderInputFrame> GetFrameAndSetTexture(TRefCountPtr<FRHITexture> InTexture);
+		TSharedPtr<AVEncoder::FVideoEncoderInput> GetOrCreateVideoEncoderInput();
 
-		private:
-			TSharedPtr<AVEncoder::FVideoEncoderInputFrame> GetOrCreateFrame(const TRefCountPtr<FRHITexture> InTexture);
-			void RemoveStaleTextures();
-			void FlushFrames();
-			TSharedPtr<AVEncoder::FVideoEncoderInput> CreateVideoEncoderInput() const;
-			void SetTexture(TSharedPtr<AVEncoder::FVideoEncoderInputFrame> Frame, const TRefCountPtr<FRHITexture>& Texture);
-			void SetTextureCUDAVulkan(TSharedPtr<AVEncoder::FVideoEncoderInputFrame> Frame, const TRefCountPtr<FRHITexture>& Texture);
+	private:
+		TSharedPtr<AVEncoder::FVideoEncoderInputFrame> GetOrCreateFrame(const TRefCountPtr<FRHITexture> InTexture);
+		void RemoveStaleTextures();
+		void FlushFrames();
+		TSharedPtr<AVEncoder::FVideoEncoderInput> CreateVideoEncoderInput() const;
+		void SetTexture(TSharedPtr<AVEncoder::FVideoEncoderInputFrame> Frame, const TRefCountPtr<FRHITexture>& Texture);
+		void SetTextureCUDAVulkan(TSharedPtr<AVEncoder::FVideoEncoderInputFrame> Frame, const TRefCountPtr<FRHITexture>& Texture);
 #if PLATFORM_WINDOWS
-			void SetTextureCUDAD3D11(TSharedPtr<AVEncoder::FVideoEncoderInputFrame> Frame, const TRefCountPtr<FRHITexture>& Texture);
-			void SetTextureCUDAD3D12(TSharedPtr<AVEncoder::FVideoEncoderInputFrame> Frame, const TRefCountPtr<FRHITexture>& Texture);
+		void SetTextureCUDAD3D11(TSharedPtr<AVEncoder::FVideoEncoderInputFrame> Frame, const TRefCountPtr<FRHITexture>& Texture);
+		void SetTextureCUDAD3D12(TSharedPtr<AVEncoder::FVideoEncoderInputFrame> Frame, const TRefCountPtr<FRHITexture>& Texture);
 #endif // PLATFORM_WINDOWS
 
-		private:
-			uint64 FrameId = 0;
-			TSharedPtr<AVEncoder::FVideoEncoderInput> EncoderInput;
-			// Store a mapping between raw textures and the FVideoEncoderInputFrames that wrap them
-			TMap<TRefCountPtr<FRHITexture>, TSharedPtr<AVEncoder::FVideoEncoderInputFrame>> TextureToFrameMapping;
-		};
-	} // namespace PixelStreaming
-} // namespace UE
+	private:
+		uint64 FrameId = 0;
+		TSharedPtr<AVEncoder::FVideoEncoderInput> EncoderInput;
+		// Store a mapping between raw textures and the FVideoEncoderInputFrames that wrap them
+		TMap<TRefCountPtr<FRHITexture>, TSharedPtr<AVEncoder::FVideoEncoderInputFrame>> TextureToFrameMapping;
+	};
+} // namespace UE::PixelStreaming
diff --git a/Engine/Plugins/Media/PixelStreaming/Source/PixelStreaming/Private/FixedFPSPump.cpp b/Engine/Plugins/Media/PixelStreaming/Source/PixelStreaming/Private/FixedFPSPump.cpp
--- a/Engine/Plugins/Media/PixelStreaming/Source/PixelStreaming/Private/FixedFPSPump.cpp
+++ b/Engine/Plugins/Media/PixelStreaming/Source/PixelStreaming/Private/FixedFPSPump.cpp
@@ -3,98 +3,149 @@
 #include "FixedFPSPump.h"
 #include "VideoSource.h"
 #include "Settings.h"
-#include "PixelStreamingTextureSource.h"
 #include "PixelStreamingFrameBuffer.h"
 #include "PlayerSession.h"
 
-/*
-* ---------- FFixedFPSPump ----------------
-*/
-
-UE::PixelStreaming::FFixedFPSPump* UE::PixelStreaming::FFixedFPSPump::Instance = nullptr;
-
-UE::PixelStreaming::FFixedFPSPump* UE::PixelStreaming::FFixedFPSPump::Get()
+namespace UE::PixelStreaming
 {
-	checkf(Instance, TEXT("You should not try to Get() and instance of the pump before it has been constructed somewhere yet."));
-	return UE::PixelStreaming::FFixedFPSPump::Instance;
-}
+	/*
+	* ---------- FFixedFPSPump ----------------
+	*/
+	FFixedFPSPump::FFixedFPSPump()
+		: Runnable()
+	{
+		if (Settings::CVarPixelStreamingDecoupleFrameRate.GetValueOnAnyThread())
+		{
+			PumpThread = FRunnableThread::Create(&Runnable, TEXT("Pixel Streaming Frame Pump"));
+		}
+	}
 
-UE::PixelStreaming::FFixedFPSPump::FFixedFPSPump()
-	: NextPumpEvent(FPlatformProcess::GetSynchEventFromPool(false))
-{
-	PumpThread = MakeUnique<FThread>(TEXT("PumpThread"), [this]() { PumpLoop(); });
-	UE::PixelStreaming::FFixedFPSPump::Instance = this;
-}
+	FFixedFPSPump::~FFixedFPSPump()
+	{
+		Shutdown();
+		if (Settings::CVarPixelStreamingDecoupleFrameRate.GetValueOnAnyThread())
+		{
+			PumpThread->Kill(true);
+		}
+	}
 
-UE::PixelStreaming::FFixedFPSPump::~FFixedFPSPump()
-{
-	Shutdown();
-	PumpThread->Join();
-}
+	void FFixedFPSPump::Shutdown()
+	{
+		Runnable.Stop();
+	}
 
-void UE::PixelStreaming::FFixedFPSPump::Shutdown()
-{
-	bThreadRunning = false;
-	NextPumpEvent->Trigger();
-}
+	void FFixedFPSPump::RegisterPumpable(rtc::scoped_refptr<FPixelStreamingPumpable> Pumpable)
+	{
+		checkf(Pumpable, TEXT("Cannot register a nullptr."));
+		FScopeLock Guard(&Runnable.SourcesGuard);
+		Runnable.Pumpables.Add(Pumpable.get(), Pumpable);
+		Runnable.NextPumpEvent->Trigger();
+	}
 
-void UE::PixelStreaming::FFixedFPSPump::UnregisterVideoSource(FPixelStreamingPlayerId PlayerId)
-{
-	FScopeLock Guard(&SourcesGuard);
-	VideoSources.Remove(PlayerId);
-	NextPumpEvent->Trigger();
-}
+	void FFixedFPSPump::UnregisterPumpable(rtc::scoped_refptr<FPixelStreamingPumpable> Pumpable)
+	{
+		FScopeLock Guard(&Runnable.SourcesGuard);
+		Runnable.Pumpables.Remove(Pumpable.get());
+		Runnable.NextPumpEvent->Trigger();
+	}
 
-void UE::PixelStreaming::FFixedFPSPump::RegisterVideoSource(FPixelStreamingPlayerId PlayerId, IPumpedVideoSource* Source)
-{
-	checkf(Source, TEXT("Cannot register a nullptr VideoSource."));
-	FScopeLock Guard(&SourcesGuard);
-	VideoSources.Add(PlayerId, Source);
-	NextPumpEvent->Trigger();
-}
+	void FFixedFPSPump::PumpAll()
+	{
+		Runnable.PumpAll();
+	}
 
-void UE::PixelStreaming::FFixedFPSPump::PumpLoop()
-{
-	uint64 LastCycles = FPlatformTime::Cycles64();
+	/*
+	* ------------- FPumpRunnable ---------------------
+	*/
+	FFixedFPSPump::FPumpRunnable::FPumpRunnable()
+		: bIsRunning(false)
+		, NextPumpEvent(FPlatformProcess::GetSynchEventFromPool(false))
+	{
+	}
 
-	while (bThreadRunning)
+	bool FFixedFPSPump::FPumpRunnable::Init()
 	{
-		// No sources, so just wait.
-		if (VideoSources.Num() == 0)
-		{
-			NextPumpEvent->Wait();
-		}
+		return true;
+	}
 
+	void FFixedFPSPump::FPumpRunnable::PumpAll()
+	{
+		FScopeLock Guard(&SourcesGuard);
+		const int32 FrameId = NextFrameId++;
+		// Pump each video source
+		TMap<FPixelStreamingPumpable*, rtc::scoped_refptr<FPixelStreamingPumpable>>::TIterator Iter = Pumpables.CreateIterator();
+		for (; Iter; ++Iter)
 		{
-			FScopeLock Guard(&SourcesGuard);
+			FPixelStreamingPumpable* Pumpable = Iter.Value();
+			if (Pumpable->HasOneRef())
+			{
+				Iter.RemoveCurrent();
+				continue;
+			}
 
-			const int32 FrameId = NextFrameId++;
+			verifyf(Pumpable, TEXT("Pumpable was null"))
+			if (Pumpable->IsReadyForPump())
+			{
+				Pumpable->OnPump(FrameId);
+			}
+		}
+	}
 
-			// Pump each video source
-			TMap<FPixelStreamingPlayerId, IPumpedVideoSource*>::TIterator Iter = VideoSources.CreateIterator();
-			for (; Iter; ++Iter)
+	uint32 FFixedFPSPump::FPumpRunnable::Run()
+	{
+		LastPumpCycles = FPlatformTime::Cycles64();
+		bIsRunning = true;
+
+		while (bIsRunning)
+		{
+			// No sources, so just wait.
+			if (Pumpables.Num() == 0)
 			{
+				NextPumpEvent->Wait();
+			}
 
-				IPumpedVideoSource* VideoSource = Iter.Value();
+			PumpAll();
 
-				if (VideoSource->IsReadyForPump())
-				{
-					VideoSource->OnPump(FrameId);
-				}
+			// Sleep as long as we need for a constant FPS
+			const uint64 EndCycles = FPlatformTime::Cycles64();
+			const double DeltaMs = FPlatformTime::ToMilliseconds64(EndCycles - LastPumpCycles);
+			const int32 FPS = Settings::CVarPixelStreamingWebRTCFps.GetValueOnAnyThread();
+			const double FrameDeltaMs = 1000.0 / FPS;
+			const double SleepMs = FrameDeltaMs - DeltaMs;
+			LastPumpCycles = EndCycles;
+			if (SleepMs > 0)
+			{
+				NextPumpEvent->Wait(SleepMs, false);
 			}
 		}
+		return 0;
+	}
 
-		// Sleep as long as we need for a constant FPS
-		const uint64 EndCycles = FPlatformTime::Cycles64();
-		const double DeltaMs = FPlatformTime::ToMilliseconds64(EndCycles - LastCycles);
-		const int32 FPS = UE::PixelStreaming::Settings::CVarPixelStreamingWebRTCFps.GetValueOnAnyThread();
-		const double FrameDeltaMs = 1000.0 / FPS;
-		const double SleepMs = FrameDeltaMs - DeltaMs;
-		LastCycles = EndCycles;
-
-		if (SleepMs > 0)
+	/* Tick() only occurs when engine is running in single threaded mode */
+	void FFixedFPSPump::FPumpRunnable::Tick()
+	{
+		NextPumpEvent->Trigger();
+		/*
+		NOTE: When running single threaded, we will never be able to pump faster than the engine is rendering so the logic to do so is removed.
+		*/
+		if (Pumpables.Num() == 0)
 		{
-			NextPumpEvent->Wait(SleepMs, false);
+			// No sources, so just wait.
+			return;
 		}
+
+		PumpAll();
+	}
+
+	void FFixedFPSPump::FPumpRunnable::Stop()
+	{
+		bIsRunning = false;
+		NextPumpEvent->Trigger();
+	}
+
+	void FFixedFPSPump::FPumpRunnable::Exit()
+	{
+		bIsRunning = false;
+		NextPumpEvent->Trigger();
 	}
-}
+} // namespace UE::PixelStreaming
\ No newline at end of file
diff --git a/Engine/Plugins/Media/PixelStreaming/Source/PixelStreaming/Private/FixedFPSPump.h b/Engine/Plugins/Media/PixelStreaming/Source/PixelStreaming/Private/FixedFPSPump.h
--- a/Engine/Plugins/Media/PixelStreaming/Source/PixelStreaming/Private/FixedFPSPump.h
+++ b/Engine/Plugins/Media/PixelStreaming/Source/PixelStreaming/Private/FixedFPSPump.h
@@ -1,46 +1,64 @@
 // Copyright Epic Games, Inc. All Rights Reserved.
 #pragma once
 
-#include "PixelStreamingPlayerId.h"
-#include "HAL/Thread.h"
+#include "HAL/Runnable.h"
+#include "HAL/RunnableThread.h"
+#include "HAL/Event.h"
+#include "HAL/CriticalSection.h"
+#include "Misc/SingleThreadRunnable.h"
 #include "WebRTCIncludes.h"
-#include "IPixelStreamingPumpedVideoSource.h"
+#include "PixelStreamingPumpable.h"
 
-namespace UE
+namespace UE::PixelStreaming
 {
-	namespace PixelStreaming
+	/*
+	* Runs a seperate thread that "pumps" at a fixed FPS interval. "Pumpables" may add themselves to be "pumped" 
+	* at a fixed interval - indepedent of render FPS or game thread FPS. This is useful in many cases where a poorly 
+	* performing applications should still do a thing at a fixed interval. For example, we want a constant amount of 
+	* frames submitted to WebRTC so we are not penalized with a bitrate reduction.
+	*/
+	class FFixedFPSPump
 	{
-		/*
-		* Runs a seperate thread that "pumps" at a fixed FPS interval. WebRTC video sources may add themselves to be "pumped" 
-		* at a fixed interval indepedent of render FPS or game thread FPS. This is useful so that poorly 
-		* performing applications can still submit a constant amount of frames to WebRTC and are not penalized with a 
-		* bitrate reduction.
-		*/
-		class FFixedFPSPump
-		{
+	public:
+		FFixedFPSPump();
+		~FFixedFPSPump();
+		void Shutdown();
+		void RegisterPumpable(rtc::scoped_refptr<FPixelStreamingPumpable> Pumpable);
+		void UnregisterPumpable(rtc::scoped_refptr<FPixelStreamingPumpable> Pumpable);
+		void PumpAll();
 
+	private:
+		class FPumpRunnable : public FRunnable, public FSingleThreadRunnable
+		{
 		public:
-			FFixedFPSPump();
-			~FFixedFPSPump();
-			void Shutdown();
-			void RegisterVideoSource(FPixelStreamingPlayerId PlayerId, IPumpedVideoSource* Source);
-			void UnregisterVideoSource(FPixelStreamingPlayerId PlayerId);
-			static FFixedFPSPump* Get();
+			FPumpRunnable();
+			void PumpAll();
 
-		private:
-			void PumpLoop();
+			// Begin FRunnable interface.
+			virtual bool Init() override;
+			virtual uint32 Run() override;
+			virtual void Stop() override;
+			virtual void Exit() override;
+			virtual FSingleThreadRunnable* GetSingleThreadInterface() override
+			{
+				bIsRunning = true;
+				return this;
+			};
+			// End FRunnable interface
 
-		private:
-			FCriticalSection SourcesGuard;
-			TMap<FPixelStreamingPlayerId, IPumpedVideoSource*> VideoSources;
+			// Begin FSingleThreadRunnable interface.
+			virtual void Tick() override;
+			// End FSingleThreadRunnable interface
 
-			TUniquePtr<FThread> PumpThread;
-			bool bThreadRunning = true;
-			//FEvent* PlayersChangedEvent;
+			bool bIsRunning = false;
 			FEvent* NextPumpEvent;
+			FCriticalSection SourcesGuard;
+			TMap<FPixelStreamingPumpable*, rtc::scoped_refptr<FPixelStreamingPumpable>> Pumpables;
 			int32 NextFrameId = 0;
-
-			static FFixedFPSPump* Instance;
+			uint64 LastPumpCycles;
+			uint64 NextPumpCycles;
 		};
-	} // namespace PixelStreaming
-} // namespace UE
+		FPumpRunnable Runnable;
+		FRunnableThread* PumpThread;
+	};
+} // namespace UE::PixelStreaming
diff --git a/Engine/Plugins/Media/PixelStreaming/Source/PixelStreaming/Private/FrameBuffer.cpp b/Engine/Plugins/Media/PixelStreaming/Source/PixelStreaming/Private/FrameBuffer.cpp
--- a/Engine/Plugins/Media/PixelStreaming/Source/PixelStreaming/Private/FrameBuffer.cpp
+++ b/Engine/Plugins/Media/PixelStreaming/Source/PixelStreaming/Private/FrameBuffer.cpp
@@ -1,183 +1,150 @@
 // Copyright Epic Games, Inc. All Rights Reserved.
 
-#include "PixelStreamingFrameBuffer.h"
-#include "PixelStreamingTextureSource.h"
+#include "FrameBuffer.h"
 #include "Utils.h"
+#include "Stats.h"
 
-/*
-* ----------------- FInitializeFrameBuffer -----------------
-*/
-
-FInitializeFrameBuffer::FInitializeFrameBuffer(TSharedPtr<IPixelStreamingTextureSource> InTextureSource)
-	: TextureSource(InTextureSource)
-{
-}
-
-FInitializeFrameBuffer::~FInitializeFrameBuffer()
-{
-}
-
-int FInitializeFrameBuffer::width() const
-{
-	return TextureSource->GetSourceWidth();
-}
-
-int FInitializeFrameBuffer::height() const
-{
-	return TextureSource->GetSourceHeight();
-}
-
-/*
-* ----------------- FSimulcastFrameBuffer -----------------
-*/
-
-FSimulcastFrameBuffer::FSimulcastFrameBuffer(TArray<TSharedPtr<IPixelStreamingTextureSource>>& InTextureSources)
-	: TextureSources(InTextureSources)
+namespace UE::PixelStreaming
 {
-}
+	/*
+	* ----------------- FInitializeFrameBuffer -----------------
+	*/
 
-FSimulcastFrameBuffer::~FSimulcastFrameBuffer()
-{
-}
-
-int FSimulcastFrameBuffer::GetNumLayers() const
-{
-	return TextureSources.Num();
-}
+	FInitializeFrameBuffer::FInitializeFrameBuffer(TSharedPtr<FTextureTripleBuffer> InTextureBuffer)
+		: TextureBuffer(InTextureBuffer)
+	{
+	}
 
-TSharedPtr<IPixelStreamingTextureSource> FSimulcastFrameBuffer::GetLayerFrameSource(int LayerIndex) const
-{
-	checkf(LayerIndex >= 0 && LayerIndex < TextureSources.Num(), TEXT("Requested layer index was out of bounds."));
-	return TextureSources[LayerIndex];
-}
+	FInitializeFrameBuffer::~FInitializeFrameBuffer()
+	{
+	}
 
-int FSimulcastFrameBuffer::width() const
-{
-	checkf(TextureSources.Num() > 0, TEXT("Must be at least one texture source to get the width from."));
-	return FMath::Max(TextureSources[0]->GetSourceWidth(), TextureSources[TextureSources.Num() - 1]->GetSourceWidth());
-}
+	int FInitializeFrameBuffer::width() const
+	{
+		return TextureBuffer->GetSourceWidth();
+	}
 
-int FSimulcastFrameBuffer::height() const
-{
-	checkf(TextureSources.Num() > 0, TEXT("Must be at least one texture source to get the height from."));
-	return FMath::Max(TextureSources[0]->GetSourceHeight(), TextureSources[TextureSources.Num() - 1]->GetSourceHeight());
-}
+	int FInitializeFrameBuffer::height() const
+	{
+		return TextureBuffer->GetSourceHeight();
+	}
 
-/*
-* ----------------- FLayerFrameBuffer -----------------
-*/
+	/*
+	* ----------------- FSimulcastFrameBuffer -----------------
+	*/
 
-FLayerFrameBuffer::FLayerFrameBuffer(TSharedPtr<IPixelStreamingTextureSource> InTextureSource)
-	: TextureSource(InTextureSource)
-{
-}
+	FSimulcastFrameBuffer::FSimulcastFrameBuffer(TArray<TSharedPtr<FTextureTripleBuffer>>& InTextureBuffers)
+		: TextureBuffers(InTextureBuffers)
+	{
+	}
 
-FLayerFrameBuffer::~FLayerFrameBuffer()
-{
-}
+	FSimulcastFrameBuffer::~FSimulcastFrameBuffer()
+	{
+	}
 
-FTexture2DRHIRef FLayerFrameBuffer::GetFrame() const
-{
-	return TextureSource->GetTexture();
-}
+	int FSimulcastFrameBuffer::GetNumLayers() const
+	{
+		return TextureBuffers.Num();
+	}
 
-int FLayerFrameBuffer::width() const
-{
-	return TextureSource->GetSourceWidth();
-}
+	TSharedPtr<FTextureTripleBuffer> FSimulcastFrameBuffer::GetLayerFrameSource(int LayerIndex) const
+	{
+		checkf(LayerIndex >= 0 && LayerIndex < TextureBuffers.Num(), TEXT("Requested layer index was out of bounds."));
+		return TextureBuffers[LayerIndex];
+	}
 
-int FLayerFrameBuffer::height() const
-{
-	return TextureSource->GetSourceHeight();
-}
+	int FSimulcastFrameBuffer::width() const
+	{
+		checkf(TextureBuffers.Num() > 0, TEXT("Must be at least one texture source to get the width from."));
+		return FMath::Max(TextureBuffers[0]->GetSourceWidth(), TextureBuffers[TextureBuffers.Num() - 1]->GetSourceWidth());
+	}
 
-/*
-* ----------------- FFrameBufferI420 -----------------
-*/
+	int FSimulcastFrameBuffer::height() const
+	{
+		checkf(TextureBuffers.Num() > 0, TEXT("Must be at least one texture source to get the height from."));
+		return FMath::Max(TextureBuffers[0]->GetSourceHeight(), TextureBuffers[TextureBuffers.Num() - 1]->GetSourceHeight());
+	}
 
-FFrameBufferI420::FFrameBufferI420(TSharedPtr<IPixelStreamingTextureSource> InTextureSource)
-	: TextureSource(InTextureSource)
-{
-}
+	/*
+	* ----------------- FLayerFrameBuffer -----------------
+	*/
 
-FFrameBufferI420::~FFrameBufferI420()
-{
-}
+	FLayerFrameBuffer::FLayerFrameBuffer(TSharedPtr<FTextureTripleBuffer> InTextureBuffer)
+		: TextureBuffer(InTextureBuffer)
+	{
+	}
 
-FTexture2DRHIRef FFrameBufferI420::GetFrame() const
-{
-	return TextureSource->GetTexture();
-}
+	FLayerFrameBuffer::~FLayerFrameBuffer()
+	{
+	}
 
-int FFrameBufferI420::width() const
-{
-	return TextureSource->GetSourceWidth();
-}
+	TSharedPtr<FPixelStreamingTextureWrapper> FLayerFrameBuffer::GetFrame() const
+	{
+		return TextureBuffer->GetCurrent();
+	}
 
-int FFrameBufferI420::height() const
-{
-	return TextureSource->GetSourceHeight();
-}
+	int FLayerFrameBuffer::width() const
+	{
+		return TextureBuffer->GetSourceWidth();
+	}
 
-const webrtc::I420BufferInterface* FFrameBufferI420::GetI420() const
-{
-	return Buffer;
-}
+	int FLayerFrameBuffer::height() const
+	{
+		return TextureBuffer->GetSourceHeight();
+	}
 
-/*
-* NOTE: Only used for non-hardware encoders (e.g. VP8) - H264 does NOT hit this function.
-*/
-rtc::scoped_refptr<webrtc::I420BufferInterface> FFrameBufferI420::ToI420()
-{
-	IBackBufferTextureSource* TextureSourcePtr = static_cast<IBackBufferTextureSource*>(TextureSource.Get());
-	checkf(TextureSourcePtr && TextureSourcePtr->GetName() == FString(TEXT("FBackBufferToCPUTextureSource")), TEXT("To use VPX encoder texture source must be FBackBufferToCPUTextureSource"));
+	/*
+	* ----------------- FFrameBufferI420 -----------------
+	*/
 
-	TRefCountPtr<FPixelStreamingCPUReadableBackbufferTexture> Current = static_cast<FPixelStreamingCPUReadableBackbufferTexture*>(TextureSourcePtr->GetCurrent().GetReference());
+	FFrameBufferI420::FFrameBufferI420(TSharedPtr<FTextureTripleBuffer> InTextureBuffer)
+		: TextureBuffer(InTextureBuffer)
+	{
+	}
 
-	uint32 TextureWidth = Current->GetTexture()->GetSizeX();
-	uint32 TextureHeight = Current->GetTexture()->GetSizeY();
+	FFrameBufferI420::~FFrameBufferI420()
+	{
+	}
 
-	uint32 i = 0;
-	uint32 NumPixels = TextureWidth * TextureHeight;
-	uint32 ui = NumPixels;
-	uint32 vi = NumPixels + NumPixels / 4;
-	uint32 s = 0;
+	TSharedPtr<FPixelStreamingTextureWrapper> FFrameBufferI420::GetFrame() const
+	{
+		return TextureBuffer->GetCurrent();
+	}
 
-// data is coming in as B,G,R,A
-#define sR ((uint32)(rgbIn[s + 2]))
-#define sG ((uint32)(rgbIn[s + 1]))
-#define sB ((uint32)(rgbIn[s + 0]))
+	int FFrameBufferI420::width() const
+	{
+		return TextureBuffer->GetSourceWidth();
+	}
 
-	Buffer = webrtc::I420Buffer::Create(TextureWidth, TextureHeight, TextureWidth, TextureWidth / 2, TextureWidth / 2);
+	int FFrameBufferI420::height() const
+	{
+		return TextureBuffer->GetSourceHeight();
+	}
 
-	// If we have no raw pixels then early exit this swizzle operation
-	if (Current->GetRawPixels().Num() == 0)
+	const webrtc::I420BufferInterface* FFrameBufferI420::GetI420() const
 	{
-		webrtc::I420Buffer::SetBlack(Buffer);
 		return Buffer;
 	}
 
-	uint8* yuv420p = Buffer->MutableDataY();
-	const uint8* rgbIn = (const uint8*)Current->GetRawPixels().GetData();
-	for (uint32 j = 0; j < TextureHeight; j++)
+	/*
+	* NOTE: Only used for non-hardware encoders (e.g. VP8) - H264 does NOT hit this function.
+	*/
+	rtc::scoped_refptr<webrtc::I420BufferInterface> FFrameBufferI420::ToI420()
 	{
-		for (uint32 k = 0; k < TextureWidth; k++)
+		TSharedPtr<FPixelStreamingTextureWrapper> Current = TextureBuffer->GetCurrent();
+
+		uint64 PreConversion = FPlatformTime::Cycles64();
+
+		Buffer = TextureBuffer->ToWebRTCI420Buffer(Current);
+
+		uint64 PostConversion = FPlatformTime::Cycles64();
+		UE::PixelStreaming::FStats* Stats = UE::PixelStreaming::FStats::Get();
+		if (Stats)
 		{
-			yuv420p[i] = (uint8)((66 * sR + 129 * sG + 25 * sB + 128) >> 8) + 16;
-
-			if (0 == j % 2 && 0 == k % 2)
-			{
-				yuv420p[ui++] = (uint8)((-38 * sR - 74 * sG + 112 * sB + 128) >> 8) + 128;
-				yuv420p[vi++] = (uint8)((112 * sR - 94 * sG - 18 * sB + 128) >> 8) + 128;
-			}
-			i++;
-			s += 4;
+			double ConversionMs = FPlatformTime::ToMilliseconds64(PostConversion - PreConversion);
+			Stats->StoreApplicationStat(UE::PixelStreaming::FStatData(FName(*FString::Printf(TEXT("BGRA->YUV420 (ms)"))), ConversionMs, 2, true));
 		}
-	}
-
-#undef sR
-#undef sG
-#undef sB
 
-	return Buffer;
-}
+		return Buffer;
+	}
+}
\ No newline at end of file
diff --git a/Engine/Plugins/Media/PixelStreaming/Source/PixelStreaming/Private/FrameBuffer.h b/Engine/Plugins/Media/PixelStreaming/Source/PixelStreaming/Private/FrameBuffer.h
new file mode 100644
--- /dev/null
+++ b/Engine/Plugins/Media/PixelStreaming/Source/PixelStreaming/Private/FrameBuffer.h
@@ -0,0 +1,92 @@
+// Copyright Epic Games, Inc. All Rights Reserved.
+
+#pragma once
+#include "PixelStreamingFrameBuffer.h"
+#include "TextureTripleBuffer.h"
+
+namespace UE::PixelStreaming
+{
+	/*
+	* ----------------- FInitializeFrameBuffer -----------------
+	* We use this frame to force webrtc to create encoders that will siphon off other streams
+	* but never get real frames directly pumped to them.
+	*/
+	class FInitializeFrameBuffer : public FPixelStreamingFrameBuffer
+	{
+	public:
+		FInitializeFrameBuffer(TSharedPtr<FTextureTripleBuffer> InTextureBuffer);
+		virtual ~FInitializeFrameBuffer();
+
+		virtual EPixelStreamingFrameBufferType GetFrameBufferType() const { return EPixelStreamingFrameBufferType::Initialize; }
+
+		virtual int width() const override;
+		virtual int height() const override;
+
+	private:
+		TSharedPtr<FTextureTripleBuffer> TextureBuffer;
+	};
+
+	/*
+	* ----------------- FSimulcastFrameBuffer -----------------
+	* Holds a number of textures (called layers) and each is assigned an index.
+	*/
+	class FSimulcastFrameBuffer : public FPixelStreamingFrameBuffer
+	{
+	public:
+		FSimulcastFrameBuffer(TArray<TSharedPtr<FTextureTripleBuffer>>& InTextureBuffer);
+		virtual ~FSimulcastFrameBuffer();
+
+		virtual EPixelStreamingFrameBufferType GetFrameBufferType() const { return EPixelStreamingFrameBufferType::Simulcast; }
+		TSharedPtr<FTextureTripleBuffer> GetLayerFrameSource(int LayerIndex) const;
+		int GetNumLayers() const;
+
+		virtual int width() const override;
+		virtual int height() const override;
+
+	private:
+		TArray<TSharedPtr<FTextureTripleBuffer>>& TextureBuffers;
+	};
+
+	/*
+	* ----------------- FLayerFrameBuffer -----------------
+	* Holds a genuine single texture for encoding.
+	*/
+	class FLayerFrameBuffer : public FPixelStreamingFrameBuffer
+	{
+	public:
+		FLayerFrameBuffer(TSharedPtr<FTextureTripleBuffer> InTextureBuffer);
+		virtual ~FLayerFrameBuffer();
+
+		virtual EPixelStreamingFrameBufferType GetFrameBufferType() const { return EPixelStreamingFrameBufferType::Layer; }
+		TSharedPtr<FPixelStreamingTextureWrapper> GetFrame() const;
+
+		virtual int width() const override;
+		virtual int height() const override;
+
+	private:
+		TSharedPtr<FTextureTripleBuffer> TextureBuffer;
+	};
+
+	/*
+	* ----------------- FFrameBufferI420 -----------------
+	* Holds a texture source that is capable of being converted to I420 for encoding.
+	*/
+	class FFrameBufferI420 : public FPixelStreamingFrameBuffer
+	{
+	public:
+		FFrameBufferI420(TSharedPtr<FTextureTripleBuffer> InTextureBuffer);
+		virtual ~FFrameBufferI420();
+
+		virtual EPixelStreamingFrameBufferType GetFrameBufferType() const { return EPixelStreamingFrameBufferType::Layer; }
+		TSharedPtr<FPixelStreamingTextureWrapper> GetFrame() const;
+
+		virtual int width() const override;
+		virtual int height() const override;
+		virtual rtc::scoped_refptr<webrtc::I420BufferInterface> ToI420() override;
+		virtual const webrtc::I420BufferInterface* GetI420() const override;
+
+	private:
+		TSharedPtr<FTextureTripleBuffer> TextureBuffer;
+		rtc::scoped_refptr<webrtc::I420Buffer> Buffer;
+	};
+} // namespace UE::PixelStreaming
\ No newline at end of file
diff --git a/Engine/Plugins/Media/PixelStreaming/Source/PixelStreaming/Private/IPlayerSession.h b/Engine/Plugins/Media/PixelStreaming/Source/PixelStreaming/Private/IPlayerSession.h
--- a/Engine/Plugins/Media/PixelStreaming/Source/PixelStreaming/Private/IPlayerSession.h
+++ b/Engine/Plugins/Media/PixelStreaming/Source/PixelStreaming/Private/IPlayerSession.h
@@ -2,6 +2,7 @@
 
 #pragma once
 
+#include "WebRTCIncludes.h"
 #include "PixelStreamingProtocolDefs.h"
 #include "PixelStreamingPlayerId.h"
 
@@ -26,6 +27,7 @@ namespace UE::PixelStreaming
 		virtual void OnRemoteIceCandidate(const FString& SdpMid, int SdpMLineIndex, const FString& Sdp) = 0;
 		virtual void DisconnectPlayer(const FString& Reason) = 0;
 
+		virtual FName GetSessionType() const = 0;
 		virtual FPixelStreamingPlayerId GetPlayerId() const = 0;
 		virtual IPixelStreamingAudioSink* GetAudioSink() = 0;
 		virtual FDataChannelObserver* GetDataChannelObserver() = 0;
