This commit is a large collection (from various authors) of Pixel Streaming 5.1 work

diff --git a/Engine/Plugins/Media/PixelStreaming/Source/PixelStreaming/Private/InputDevice.cpp b/Engine/Plugins/Media/PixelStreaming/Source/PixelStreaming/Private/InputDevice.cpp
--- a/Engine/Plugins/Media/PixelStreaming/Source/PixelStreaming/Private/InputDevice.cpp
+++ b/Engine/Plugins/Media/PixelStreaming/Source/PixelStreaming/Private/InputDevice.cpp
@@ -24,1010 +24,1006 @@ DEFINE_LOG_CATEGORY(LogPixelStreamingInputDevice);
 DECLARE_LOG_CATEGORY_EXTERN(LogPixelStreamingInput, Log, VeryVerbose);
 DEFINE_LOG_CATEGORY(LogPixelStreamingInput);
 
-namespace UE
+namespace UE::PixelStreaming
 {
-	namespace PixelStreaming
+	template <typename T>
+	static const T& ParseBuffer(const uint8*& Data, uint32& Size)
 	{
-		template <typename T>
-		static const T& ParseBuffer(const uint8*& Data, uint32& Size)
+		checkf(sizeof(T) <= Size, TEXT("%d - %d"), sizeof(T), Size);
+		const T& Value = *reinterpret_cast<const T*>(Data);
+		Data += sizeof(T);
+		Size -= sizeof(T);
+		return Value;
+	}
+
+	static FString ParseString(const webrtc::DataBuffer& Buffer, const size_t Offset = 0)
+	{
+		FString Res;
+		if (Buffer.data.size() > Offset)
 		{
-			checkf(sizeof(T) <= Size, TEXT("%d - %d"), sizeof(T), Size);
-			const T& Value = *reinterpret_cast<const T*>(Data);
-			Data += sizeof(T);
-			Size -= sizeof(T);
-			return Value;
+			size_t StringLength = (Buffer.data.size() - Offset) / sizeof(TCHAR);
+			Res.GetCharArray().SetNumUninitialized(StringLength + 1);
+			FMemory::Memcpy(Res.GetCharArray().GetData(), Buffer.data.data() + Offset, StringLength * sizeof(TCHAR));
+			Res.GetCharArray()[StringLength] = 0;
 		}
+		return Res;
+	}
+
+	/**
+	 * When reading input from a browser then the cursor position will be sent
+	 * across with mouse events. We want to use this position and avoid getting the
+	 * cursor position from the operating system. This is not relevant to touch
+	 * events.
+	 */
+	class FCursor : public ICursor
+	{
+	public:
+		FCursor() {}
+		virtual ~FCursor() = default;
+		virtual FVector2D GetPosition() const override { return Position; }
+		virtual void SetPosition(const int32 X, const int32 Y) override { Position = FVector2D(X, Y); };
+		virtual void SetType(const EMouseCursor::Type InNewCursor) override{};
+		virtual EMouseCursor::Type GetType() const override { return EMouseCursor::Type::Default; };
+		virtual void GetSize(int32& Width, int32& Height) const override{};
+		virtual void Show(bool bShow) override{};
+		virtual void Lock(const RECT* const Bounds) override{};
+		virtual void SetTypeShape(EMouseCursor::Type InCursorType, void* CursorHandle) override{};
+
+	private:
+		/** The cursor position sent across with mouse events. */
+		FVector2D Position;
+	};
 
-		static FString ParseString(const webrtc::DataBuffer& Buffer, const size_t Offset = 0)
+	/**
+	* Wrap the GenericApplication layer so we can replace the cursor and override
+	* certain behavior.
+	*/
+	class FApplicationWrapper : public GenericApplication
+	{
+	public:
+		FApplicationWrapper(TSharedPtr<GenericApplication> InWrappedApplication)
+			: GenericApplication(MakeShareable(new FCursor()))
+			, WrappedApplication(InWrappedApplication)
 		{
-			FString Res;
-			if (Buffer.data.size() > Offset)
-			{
-				size_t StringLength = (Buffer.data.size() - Offset) / sizeof(TCHAR);
-				Res.GetCharArray().SetNumUninitialized(StringLength + 1);
-				FMemory::Memcpy(Res.GetCharArray().GetData(), Buffer.data.data() + Offset, StringLength * sizeof(TCHAR));
-				Res.GetCharArray()[StringLength] = 0;
-			}
-			return Res;
+			// Whether we want to always consider the mouse as attached. This allow
+			// us to run Pixel Streaming on a machine which has no physical mouse
+			// and just let the browser supply mouse positions.
+			const UPixelStreamingSettings* Settings = GetDefault<UPixelStreamingSettings>();
+			check(Settings);
+			bMouseAlwaysAttached = Settings->bMouseAlwaysAttached;
 		}
 
 		/**
-		 * When reading input from a browser then the cursor position will be sent
-		 * across with mouse events. We want to use this position and avoid getting the
-		 * cursor position from the operating system. This is not relevant to touch
-		 * events.
+		 * Functions passed directly to the wrapped application.
 		 */
-		class FCursor : public ICursor
-		{
-		public:
-			FCursor() {}
-			virtual ~FCursor() = default;
-			virtual FVector2D GetPosition() const override { return Position; }
-			virtual void SetPosition(const int32 X, const int32 Y) override { Position = FVector2D(X, Y); };
-			virtual void SetType(const EMouseCursor::Type InNewCursor) override{};
-			virtual EMouseCursor::Type GetType() const override { return EMouseCursor::Type::Default; };
-			virtual void GetSize(int32& Width, int32& Height) const override{};
-			virtual void Show(bool bShow) override{};
-			virtual void Lock(const RECT* const Bounds) override{};
-			virtual void SetTypeShape(EMouseCursor::Type InCursorType, void* CursorHandle) override{};
-
-		private:
-			/** The cursor position sent across with mouse events. */
-			FVector2D Position;
-		};
+		virtual void SetMessageHandler(const TSharedRef<FGenericApplicationMessageHandler>& InMessageHandler) { WrappedApplication->SetMessageHandler(InMessageHandler); }
+		virtual void PollGameDeviceState(const float TimeDelta) { WrappedApplication->PollGameDeviceState(TimeDelta); }
+		virtual void PumpMessages(const float TimeDelta) { WrappedApplication->PumpMessages(TimeDelta); }
+		virtual void ProcessDeferredEvents(const float TimeDelta) { WrappedApplication->ProcessDeferredEvents(TimeDelta); }
+		virtual void Tick(const float TimeDelta) { WrappedApplication->Tick(TimeDelta); }
+		virtual TSharedRef<FGenericWindow> MakeWindow() { return WrappedApplication->MakeWindow(); }
+		virtual void InitializeWindow(const TSharedRef<FGenericWindow>& Window, const TSharedRef<FGenericWindowDefinition>& InDefinition, const TSharedPtr<FGenericWindow>& InParent, const bool bShowImmediately) { WrappedApplication->InitializeWindow(Window, InDefinition, InParent, bShowImmediately); }
+		virtual void SetCapture(const TSharedPtr<FGenericWindow>& InWindow) { WrappedApplication->SetCapture(InWindow); }
+		virtual void* GetCapture(void) const { return WrappedApplication->GetCapture(); }
+		virtual FModifierKeysState GetModifierKeys() const { return WrappedApplication->GetModifierKeys(); }
+		virtual TSharedPtr<FGenericWindow> GetWindowUnderCursor() { return WrappedApplication->GetWindowUnderCursor(); }
+		virtual void SetHighPrecisionMouseMode(const bool Enable, const TSharedPtr<FGenericWindow>& InWindow) { WrappedApplication->SetHighPrecisionMouseMode(Enable, InWindow); };
+		virtual bool IsUsingHighPrecisionMouseMode() const { return WrappedApplication->IsUsingHighPrecisionMouseMode(); }
+		virtual bool IsUsingTrackpad() const { return WrappedApplication->IsUsingTrackpad(); }
+		virtual bool IsGamepadAttached() const { return WrappedApplication->IsGamepadAttached(); }
+		virtual void RegisterConsoleCommandListener(const FOnConsoleCommandListener& InListener) { WrappedApplication->RegisterConsoleCommandListener(InListener); }
+		virtual void AddPendingConsoleCommand(const FString& InCommand) { WrappedApplication->AddPendingConsoleCommand(InCommand); }
+		virtual FPlatformRect GetWorkArea(const FPlatformRect& CurrentWindow) const { return WrappedApplication->GetWorkArea(CurrentWindow); }
+		virtual bool TryCalculatePopupWindowPosition(const FPlatformRect& InAnchor, const FVector2D& InSize, const FVector2D& ProposedPlacement, const EPopUpOrientation::Type Orientation, /*OUT*/ FVector2D* const CalculatedPopUpPosition) const { return WrappedApplication->TryCalculatePopupWindowPosition(InAnchor, InSize, ProposedPlacement, Orientation, CalculatedPopUpPosition); }
+		virtual void GetInitialDisplayMetrics(FDisplayMetrics& OutDisplayMetrics) const { WrappedApplication->GetInitialDisplayMetrics(OutDisplayMetrics); }
+		virtual EWindowTitleAlignment::Type GetWindowTitleAlignment() const { return WrappedApplication->GetWindowTitleAlignment(); }
+		virtual EWindowTransparency GetWindowTransparencySupport() const { return WrappedApplication->GetWindowTransparencySupport(); }
+		virtual void DestroyApplication() { WrappedApplication->DestroyApplication(); }
+		virtual IInputInterface* GetInputInterface() { return WrappedApplication->GetInputInterface(); }
+		virtual ITextInputMethodSystem* GetTextInputMethodSystem() { return WrappedApplication->GetTextInputMethodSystem(); }
+		virtual void SendAnalytics(IAnalyticsProvider* Provider) { WrappedApplication->SendAnalytics(Provider); }
+		virtual bool SupportsSystemHelp() const { return WrappedApplication->SupportsSystemHelp(); }
+		virtual void ShowSystemHelp() { WrappedApplication->ShowSystemHelp(); }
+		virtual bool ApplicationLicenseValid(FPlatformUserId PlatformUser = PLATFORMUSERID_NONE) { return WrappedApplication->ApplicationLicenseValid(PlatformUser); }
 
 		/**
-		* Wrap the GenericApplication layer so we can replace the cursor and override
-		* certain behavior.
-		*/
-		class FApplicationWrapper : public GenericApplication
+		 * Functions with overridden behavior.
+		 */
+		virtual bool IsCursorDirectlyOverSlateWindow() const { return true; }
+		virtual bool IsMouseAttached() const { return bMouseAlwaysAttached ? true : WrappedApplication->IsMouseAttached(); }
+
+		TSharedPtr<GenericApplication> WrappedApplication;
+		bool bMouseAlwaysAttached;
+	};
+
+	const FVector2D FInputDevice::UnfocusedPos(-1.0f, -1.0f);
+	const size_t FInputDevice::MessageHeaderOffset = 3;
+
+	FInputDevice::FInputDevice(const TSharedRef<FGenericApplicationMessageHandler>& InMessageHandler)
+		: PixelStreamerApplicationWrapper(MakeShareable(new FApplicationWrapper(FSlateApplication::Get().GetPlatformApplication())))
+		, MessageHandler(InMessageHandler)
+		, bFakingTouchEvents(FSlateApplication::Get().IsFakingTouchEvents())
+		, FocusedPos(UnfocusedPos)
+		, PixelStreamingModule(FPixelStreamingModule::GetModule())
+	{
+		if (GEngine->GameViewport && !GEngine->GameViewport->HasSoftwareCursor(EMouseCursor::Default))
 		{
-		public:
-			FApplicationWrapper(TSharedPtr<GenericApplication> InWrappedApplication)
-				: GenericApplication(MakeShareable(new FCursor()))
-				, WrappedApplication(InWrappedApplication)
+			// Pixel streaming always requires a default software cursor as it needs
+			// to be shown on the browser to allow the user to click UI elements.
+			const UPixelStreamingSettings* Settings = GetDefault<UPixelStreamingSettings>();
+			check(Settings);
+
+			// Check to see if we want to hide the cursor (by making it invisible).
+			// This is required if we want to make the cursor client-side (displayed
+			// only by the the browser).
+			bool bHideCursor = UE::PixelStreaming::Settings::IsPixelStreamingHideCursor();
+			if (bHideCursor)
 			{
-				// Whether we want to always consider the mouse as attached. This allow
-				// us to run Pixel Streaming on a machine which has no physical mouse
-				// and just let the browser supply mouse positions.
-				const UPixelStreamingSettings* Settings = GetDefault<UPixelStreamingSettings>();
-				check(Settings);
-				bMouseAlwaysAttached = Settings->bMouseAlwaysAttached;
+				GEngine->GameViewport->AddSoftwareCursor(EMouseCursor::Default, Settings->HiddenCursorClassName);
+				GEngine->GameViewport->AddSoftwareCursor(EMouseCursor::TextEditBeam, Settings->HiddenCursorClassName);
 			}
+			else
+			{
+				GEngine->GameViewport->AddSoftwareCursor(EMouseCursor::Default, Settings->DefaultCursorClassName);
+				GEngine->GameViewport->AddSoftwareCursor(EMouseCursor::TextEditBeam, Settings->TextEditBeamCursorClassName);
+			}
+		}
+	}
 
-			/**
-			 * Functions passed directly to the wrapped application.
-			 */
-
-			virtual void SetMessageHandler(const TSharedRef<FGenericApplicationMessageHandler>& InMessageHandler) { WrappedApplication->SetMessageHandler(InMessageHandler); }
-			virtual void PollGameDeviceState(const float TimeDelta) { WrappedApplication->PollGameDeviceState(TimeDelta); }
-			virtual void PumpMessages(const float TimeDelta) { WrappedApplication->PumpMessages(TimeDelta); }
-			virtual void ProcessDeferredEvents(const float TimeDelta) { WrappedApplication->ProcessDeferredEvents(TimeDelta); }
-			virtual void Tick(const float TimeDelta) { WrappedApplication->Tick(TimeDelta); }
-			virtual TSharedRef<FGenericWindow> MakeWindow() { return WrappedApplication->MakeWindow(); }
-			virtual void InitializeWindow(const TSharedRef<FGenericWindow>& Window, const TSharedRef<FGenericWindowDefinition>& InDefinition, const TSharedPtr<FGenericWindow>& InParent, const bool bShowImmediately) { WrappedApplication->InitializeWindow(Window, InDefinition, InParent, bShowImmediately); }
-			virtual void SetCapture(const TSharedPtr<FGenericWindow>& InWindow) { WrappedApplication->SetCapture(InWindow); }
-			virtual void* GetCapture(void) const { return WrappedApplication->GetCapture(); }
-			virtual FModifierKeysState GetModifierKeys() const { return WrappedApplication->GetModifierKeys(); }
-			virtual TSharedPtr<FGenericWindow> GetWindowUnderCursor() { return WrappedApplication->GetWindowUnderCursor(); }
-			virtual void SetHighPrecisionMouseMode(const bool Enable, const TSharedPtr<FGenericWindow>& InWindow) { WrappedApplication->SetHighPrecisionMouseMode(Enable, InWindow); };
-			virtual bool IsUsingHighPrecisionMouseMode() const { return WrappedApplication->IsUsingHighPrecisionMouseMode(); }
-			virtual bool IsUsingTrackpad() const { return WrappedApplication->IsUsingTrackpad(); }
-			virtual bool IsGamepadAttached() const { return WrappedApplication->IsGamepadAttached(); }
-			virtual void RegisterConsoleCommandListener(const FOnConsoleCommandListener& InListener) { WrappedApplication->RegisterConsoleCommandListener(InListener); }
-			virtual void AddPendingConsoleCommand(const FString& InCommand) { WrappedApplication->AddPendingConsoleCommand(InCommand); }
-			virtual FPlatformRect GetWorkArea(const FPlatformRect& CurrentWindow) const { return WrappedApplication->GetWorkArea(CurrentWindow); }
-			virtual bool TryCalculatePopupWindowPosition(const FPlatformRect& InAnchor, const FVector2D& InSize, const FVector2D& ProposedPlacement, const EPopUpOrientation::Type Orientation, /*OUT*/ FVector2D* const CalculatedPopUpPosition) const { return WrappedApplication->TryCalculatePopupWindowPosition(InAnchor, InSize, ProposedPlacement, Orientation, CalculatedPopUpPosition); }
-			virtual void GetInitialDisplayMetrics(FDisplayMetrics& OutDisplayMetrics) const { WrappedApplication->GetInitialDisplayMetrics(OutDisplayMetrics); }
-			virtual EWindowTitleAlignment::Type GetWindowTitleAlignment() const { return WrappedApplication->GetWindowTitleAlignment(); }
-			virtual EWindowTransparency GetWindowTransparencySupport() const { return WrappedApplication->GetWindowTransparencySupport(); }
-			virtual void DestroyApplication() { WrappedApplication->DestroyApplication(); }
-			virtual IInputInterface* GetInputInterface() { return WrappedApplication->GetInputInterface(); }
-			virtual ITextInputMethodSystem* GetTextInputMethodSystem() { return WrappedApplication->GetTextInputMethodSystem(); }
-			virtual void SendAnalytics(IAnalyticsProvider* Provider) { WrappedApplication->SendAnalytics(Provider); }
-			virtual bool SupportsSystemHelp() const { return WrappedApplication->SupportsSystemHelp(); }
-			virtual void ShowSystemHelp() { WrappedApplication->ShowSystemHelp(); }
-			virtual bool ApplicationLicenseValid(FPlatformUserId PlatformUser = PLATFORMUSERID_NONE) { return WrappedApplication->ApplicationLicenseValid(PlatformUser); }
-
-			/**
-			 * Functions with overridden behavior.
-			 */
-			virtual bool IsCursorDirectlyOverSlateWindow() const { return true; }
-			virtual bool IsMouseAttached() const { return bMouseAlwaysAttached ? true : WrappedApplication->IsMouseAttached(); }
-
-			TSharedPtr<GenericApplication> WrappedApplication;
-			bool bMouseAlwaysAttached;
-		};
-	} // namespace PixelStreaming
-} // namespace UE
-
-const FVector2D UE::PixelStreaming::FInputDevice::UnfocusedPos(-1.0f, -1.0f);
-const size_t UE::PixelStreaming::FInputDevice::MessageHeaderOffset = 3;
-
-UE::PixelStreaming::FInputDevice::FInputDevice(const TSharedRef<FGenericApplicationMessageHandler>& InMessageHandler)
-	: PixelStreamerApplicationWrapper(MakeShareable(new FApplicationWrapper(FSlateApplication::Get().GetPlatformApplication())))
-	, MessageHandler(InMessageHandler)
-	, bFakingTouchEvents(FSlateApplication::Get().IsFakingTouchEvents())
-	, FocusedPos(UnfocusedPos)
-	, PixelStreamingModule(UE::PixelStreaming::FPixelStreamingModule::GetModule())
-{
-	if (GEngine->GameViewport && !GEngine->GameViewport->HasSoftwareCursor(EMouseCursor::Default))
+	// XY positions are the ratio (0.0..1.0) along a viewport axis, quantized
+	// into an uint16 (0..65536). This allows the browser viewport and player
+	// viewport to have a different size.
+	void QuantizeAndNormalize(const FVector2D& InPos, uint16& OutX, uint16& OutY)
+	{
+		FIntPoint SizeXY = GEngine->GameViewport->Viewport->GetSizeXY();
+		OutX = InPos.X / SizeXY.X * 65536.0f;
+		OutY = InPos.Y / SizeXY.Y * 65536.0f;
+	}
+
+	bool FilterKey(const FKey& Key)
 	{
-		// Pixel streaming always requires a default software cursor as it needs
-		// to be shown on the browser to allow the user to click UI elements.
-		const UPixelStreamingSettings* Settings = GetDefault<UPixelStreamingSettings>();
-		check(Settings);
-
-		// Check to see if we want to hide the cursor (by making it invisible).
-		// This is required if we want to make the cursor client-side (displayed
-		// only by the the browser).
-		bool bHideCursor = UE::PixelStreaming::Settings::IsPixelStreamingHideCursor();
-		if (bHideCursor)
+		for (auto&& FilteredKey : UE::PixelStreaming::Settings::FilteredKeys)
 		{
-			GEngine->GameViewport->AddSoftwareCursor(EMouseCursor::Default, Settings->HiddenCursorClassName);
-			GEngine->GameViewport->AddSoftwareCursor(EMouseCursor::TextEditBeam, Settings->HiddenCursorClassName);
+			if (FilteredKey == Key)
+				return false;
 		}
-		else
+		return true;
+	}
+
+	void FInputDevice::Tick(float DeltaTime)
+	{
+		FEvent Event;
+		while (Events.Dequeue(Event))
 		{
-			GEngine->GameViewport->AddSoftwareCursor(EMouseCursor::Default, Settings->DefaultCursorClassName);
-			GEngine->GameViewport->AddSoftwareCursor(EMouseCursor::TextEditBeam, Settings->TextEditBeamCursorClassName);
+			switch (Event.Event)
+			{
+				case EventType::UNDEFINED:
+				{
+					checkNoEntry();
+				}
+				break;
+				case EventType::KEY_DOWN:
+				{
+					uint8 JavaScriptKeyCode;
+					bool IsRepeat;
+					Event.GetKeyDown(JavaScriptKeyCode, IsRepeat);
+					const FKey* AgnosticKey = JavaScriptKeyCodeToFKey[JavaScriptKeyCode];
+					if (FilterKey(*AgnosticKey))
+					{
+						const uint32* KeyCodePtr;
+						const uint32* CharacterCodePtr;
+						FInputKeyManager::Get().GetCodesFromKey(*AgnosticKey, KeyCodePtr, CharacterCodePtr);
+						uint32 KeyCode = KeyCodePtr ? *KeyCodePtr : 0;
+						uint32 CharacterCode = CharacterCodePtr ? *CharacterCodePtr : 0;
+						MessageHandler->OnKeyDown(KeyCode, CharacterCode, IsRepeat);
+						UE_LOG(LogPixelStreamingInputDevice, Verbose, TEXT("KEY_DOWN: KeyCode = %d; CharacterCode = %d; IsRepeat = %s"), KeyCode, CharacterCode, IsRepeat ? TEXT("True") : TEXT("False"));
+					}
+				}
+				break;
+				case EventType::KEY_UP:
+				{
+					uint8 JavaScriptKeyCode;
+					Event.GetKeyUp(JavaScriptKeyCode);
+					const FKey* AgnosticKey = JavaScriptKeyCodeToFKey[JavaScriptKeyCode];
+					if (FilterKey(*AgnosticKey))
+					{
+						const uint32* KeyCodePtr;
+						const uint32* CharacterCodePtr;
+						FInputKeyManager::Get().GetCodesFromKey(*AgnosticKey, KeyCodePtr, CharacterCodePtr);
+						uint32 KeyCode = KeyCodePtr ? *KeyCodePtr : 0;
+						uint32 CharacterCode = CharacterCodePtr ? *CharacterCodePtr : 0;
+						MessageHandler->OnKeyUp(KeyCode, CharacterCode, false); // Key up events are never repeats.
+						UE_LOG(LogPixelStreamingInputDevice, Verbose, TEXT("KEY_UP: KeyCode = %d; CharacterCode = %d"), KeyCode, CharacterCode);
+					}
+				}
+				break;
+				case EventType::KEY_PRESS:
+				{
+					TCHAR UnicodeCharacter;
+					Event.GetCharacterCode(UnicodeCharacter);
+					MessageHandler->OnKeyChar(UnicodeCharacter, false); // Key press repeat not yet available but are not intrinsically used.
+					UE_LOG(LogPixelStreamingInputDevice, Verbose, TEXT("KEY_PRESSED: Character = '%c'"), UnicodeCharacter);
+				}
+				break;
+				case EventType::MOUSE_ENTER:
+				{
+					// Override application layer to special pixel streaming version.
+					FSlateApplication::Get().OverridePlatformApplication(PixelStreamerApplicationWrapper);
+					FSlateApplication::Get().OnCursorSet();
+
+					// Make sure the viewport is active.
+					FSlateApplication::Get().ProcessApplicationActivationEvent(true);
+
+					UE_LOG(LogPixelStreamingInputDevice, Verbose, TEXT("MOUSE_ENTER"));
+				}
+				break;
+				case EventType::MOUSE_LEAVE:
+				{
+					// Restore normal application layer.
+					FSlateApplication::Get().OverridePlatformApplication(PixelStreamerApplicationWrapper->WrappedApplication);
+
+					UE_LOG(LogPixelStreamingInputDevice, Verbose, TEXT("MOUSE_LEAVE"));
+				}
+				break;
+				case EventType::MOUSE_MOVE:
+				{
+					// Ensure the cursor will be updated so it can change appearance.
+					FSlateApplication::Get().OnCursorSet();
+
+					uint16 PosX;
+					uint16 PosY;
+					int16 DeltaX;
+					int16 DeltaY;
+					Event.GetMouseDelta(PosX, PosY, DeltaX, DeltaY);
+					FVector2D CursorPos = GEngine->GameViewport->GetWindow()->GetPositionInScreen() + FVector2D(PosX, PosY);
+					PixelStreamerApplicationWrapper->Cursor->SetPosition(CursorPos.X, CursorPos.Y);
+					MessageHandler->OnRawMouseMove(DeltaX, DeltaY);
+					UE_LOG(LogPixelStreamingInputDevice, VeryVerbose, TEXT("MOUSE_MOVE: Pos = (%d, %d); CursorPos = (%d, %d); Delta = (%d, %d)"), PosX, PosY, static_cast<int>(CursorPos.X), static_cast<int>(CursorPos.Y), DeltaX, DeltaY);
+				}
+				break;
+				case EventType::MOUSE_DOWN:
+				{
+					// If a user clicks on the application window and then clicks on the
+					// browser then this will move the focus away from the application
+					// window which will deactivate the application, so we need to check
+					// if we must reactivate the application.
+					if (!FSlateApplication::Get().IsActive())
+					{
+						FSlateApplication::Get().ProcessApplicationActivationEvent(true);
+					}
+
+					EMouseButtons::Type Button;
+					uint16 PosX;
+					uint16 PosY;
+					Event.GetMouseClick(Button, PosX, PosY);
+					FVector2D CursorPos = GEngine->GameViewport->GetWindow()->GetPositionInScreen() + FVector2D(PosX, PosY);
+					PixelStreamerApplicationWrapper->Cursor->SetPosition(CursorPos.X, CursorPos.Y);
+					MessageHandler->OnMouseDown(GEngine->GameViewport->GetWindow()->GetNativeWindow(), Button, CursorPos);
+					UE_LOG(LogPixelStreamingInputDevice, Verbose, TEXT("MOUSE_DOWN: Button = %d; Pos = (%d, %d); CursorPos = (%d, %d)"), Button, PosX, PosY, static_cast<int>(CursorPos.X), static_cast<int>(CursorPos.Y));
+
+					// The browser may be faking a mouse when touching so it will send
+					// over a mouse down event.
+					FindFocusedWidget();
+				}
+				break;
+				case EventType::MOUSE_UP:
+				{
+					EMouseButtons::Type Button;
+					uint16 PosX;
+					uint16 PosY;
+					Event.GetMouseClick(Button, PosX, PosY);
+					FVector2D CursorPos = GEngine->GameViewport->GetWindow()->GetPositionInScreen() + FVector2D(PosX, PosY);
+					PixelStreamerApplicationWrapper->Cursor->SetPosition(CursorPos.X, CursorPos.Y);
+					MessageHandler->OnMouseUp(Button);
+					UE_LOG(LogPixelStreamingInputDevice, Verbose, TEXT("MOUSE_UP: Button = %d; Pos = (%d, %d); CursorPos = (%d, %d)"), Button, PosX, PosY, static_cast<int>(CursorPos.X), static_cast<int>(CursorPos.Y));
+				}
+				break;
+				case EventType::MOUSE_WHEEL:
+				{
+					int16 Delta;
+					uint16 PosX;
+					uint16 PosY;
+					Event.GetMouseWheel(Delta, PosX, PosY);
+					const float SpinFactor = 1 / 120.0f;
+					FVector2D CursorPos = GEngine->GameViewport->GetWindow()->GetPositionInScreen() + FVector2D(PosX, PosY);
+					MessageHandler->OnMouseWheel(Delta * SpinFactor, CursorPos);
+					UE_LOG(LogPixelStreamingInputDevice, Verbose, TEXT("MOUSE_WHEEL: Delta = %d; Pos = (%d, %d); CursorPos = (%d, %d)"), Delta, PosX, PosY, static_cast<int>(CursorPos.X), static_cast<int>(CursorPos.Y));
+				}
+				break;
+				case EventType::TOUCH_START:
+				{
+					uint8 TouchIndex;
+					uint16 PosX;
+					uint16 PosY;
+					uint8 Force; // Force is between 0.0 and 1.0 so will need to unquantize from byte.
+					Event.GetTouch(TouchIndex, PosX, PosY, Force);
+					FVector2D CursorPos = GEngine->GameViewport->GetWindow()->GetPositionInScreen() + FVector2D(PosX, PosY);
+					MessageHandler->OnTouchStarted(GEngine->GameViewport->GetWindow()->GetNativeWindow(), CursorPos, Force / 255.0f, TouchIndex, 0); // TODO: ControllerId?
+					UE_LOG(LogPixelStreamingInputDevice, Verbose, TEXT("TOUCH_START: TouchIndex = %d; Pos = (%d, %d); CursorPos = (%d, %d); Force = %.3f"), TouchIndex, PosX, PosY, static_cast<int>(CursorPos.X), static_cast<int>(CursorPos.Y), Force / 255.0f);
+
+					// If the user starts a touch then they may be focusing an editable
+					// widget.
+					FindFocusedWidget();
+				}
+				break;
+				case EventType::TOUCH_END:
+				{
+					uint8 TouchIndex;
+					uint16 PosX;
+					uint16 PosY;
+					uint8 Force;
+					Event.GetTouch(TouchIndex, PosX, PosY, Force);
+					FVector2D CursorPos = GEngine->GameViewport->GetWindow()->GetPositionInScreen() + FVector2D(PosX, PosY);
+					MessageHandler->OnTouchEnded(CursorPos, TouchIndex, 0); // TODO: ControllerId?
+					UE_LOG(LogPixelStreamingInputDevice, Verbose, TEXT("TOUCH_END: TouchIndex = %d; Pos = (%d, %d); CursorPos = (%d, %d)"), TouchIndex, PosX, PosY, static_cast<int>(CursorPos.X), static_cast<int>(CursorPos.Y));
+				}
+				break;
+				case EventType::TOUCH_MOVE:
+				{
+					uint8 TouchIndex;
+					uint16 PosX;
+					uint16 PosY;
+					uint8 Force; // Force is between 0.0 and 1.0 so will need to unquantize from byte.
+					Event.GetTouch(TouchIndex, PosX, PosY, Force);
+					FVector2D CursorPos = GEngine->GameViewport->GetWindow()->GetPositionInScreen() + FVector2D(PosX, PosY);
+					MessageHandler->OnTouchMoved(CursorPos, Force / 255.0f, TouchIndex, 0); // TODO: ControllerId?
+					UE_LOG(LogPixelStreamingInputDevice, VeryVerbose, TEXT("TOUCH_MOVE: TouchIndex = %d; Pos = (%d, %d); CursorPos = (%d, %d); Force = %.3f"), TouchIndex, PosX, PosY, static_cast<int>(CursorPos.X), static_cast<int>(CursorPos.Y), Force / 255.0f);
+				}
+				break;
+				case EventType::GAMEPAD_PRESS:
+				{
+					uint8 ControllerId;
+					uint8 ButtonIndex;
+					bool bIsRepeat;
+					Event.GetGamepadButtonPressed(ControllerId, ButtonIndex, bIsRepeat);
+					FGamepadKeyNames::Type ControllerButton = ConvertButtonIndexToGamepadButton(ButtonIndex);
+					if (ControllerButton == FGamepadKeyNames::Invalid)
+					{
+						break;
+					}
+					MessageHandler->OnControllerButtonPressed(ControllerButton, (int32_t)ControllerId, bIsRepeat);
+					UE_LOG(LogPixelStreamingInputDevice, VeryVerbose, TEXT("GAMEPAD_PRESS: ControllerId = %d; ButtonIndex = %d; IsRepeat = %d"), ControllerId, ButtonIndex, bIsRepeat);
+				}
+				break;
+				case EventType::GAMEPAD_RELEASE:
+				{
+					uint8 ControllerId;
+					uint8 ButtonIndex;
+					Event.GetGamepadButtonReleased(ControllerId, ButtonIndex);
+					FGamepadKeyNames::Type ControllerButton = ConvertButtonIndexToGamepadButton(ButtonIndex);
+					if (ControllerButton == FGamepadKeyNames::Invalid)
+					{
+						break;
+					}
+					MessageHandler->OnControllerButtonReleased(ControllerButton, (int32_t)ControllerId, false);
+					UE_LOG(LogPixelStreamingInputDevice, VeryVerbose, TEXT("GAMEPAD_RELEASE: ControllerId = %d; ButtonIndex = %d;"), ControllerId, ButtonIndex);
+				}
+				break;
+				case EventType::GAMEPAD_ANALOG:
+				{
+					uint8 ControllerId;
+					uint8 AxisIndex;
+					float AnalogValue;
+					Event.GetGamepadAnalog(ControllerId, AxisIndex, AnalogValue);
+					FGamepadKeyNames::Type ControllerAxis = ConvertAxisIndexToGamepadAxis(AxisIndex);
+					if (ControllerAxis == FGamepadKeyNames::Invalid)
+					{
+						break;
+					}
+					MessageHandler->OnControllerAnalog(ControllerAxis, (int32_t)ControllerId, AnalogValue);
+					UE_LOG(LogPixelStreamingInputDevice, VeryVerbose, TEXT("GAMEPAD_ANALOG: ControllerId = %d; AxisIndex = %d; AnalogValue = %.4f;"), ControllerId, AxisIndex, AnalogValue);
+				}
+				break;
+				default:
+				{
+					UE_LOG(LogPixelStreamingInputDevice, Error, TEXT("Unknown Pixel Streaming event %d with word 0x%016llx"), static_cast<int>(Event.Event), Event.Data.Word);
+				}
+				break;
+			}
 		}
-	}
-}
 
-// XY positions are the ratio (0.0..1.0) along a viewport axis, quantized
-// into an uint16 (0..65536). This allows the browser viewport and player
-// viewport to have a different size.
-void QuantizeAndNormalize(const FVector2D& InPos, uint16& OutX, uint16& OutY)
-{
-	FIntPoint SizeXY = GEngine->GameViewport->Viewport->GetSizeXY();
-	OutX = InPos.X / SizeXY.X * 65536.0f;
-	OutY = InPos.Y / SizeXY.Y * 65536.0f;
-}
+		FString UIInteraction;
+		while (UIInteractions.Dequeue(UIInteraction))
+		{
+			for (UPixelStreamingInput* InputComponent : PixelStreamingModule->GetInputComponents())
+			{
+				InputComponent->OnInputEvent.Broadcast(UIInteraction);
+				UE_LOG(LogPixelStreamingInputDevice, Verbose, TEXT("UIInteraction = %s"), *UIInteraction);
+			}
+		}
 
-bool FilterKey(const FKey& Key)
-{
-	for (auto&& FilteredKey : UE::PixelStreaming::Settings::FilteredKeys)
-	{
-		if (FilteredKey == Key)
-			return false;
+		FString Command;
+		while (Commands.Dequeue(Command))
+		{
+			UPixelStreamingInput::OnCommand(Command);
+		}
 	}
-	return true;
-}
 
-void UE::PixelStreaming::FInputDevice::Tick(float DeltaTime)
-{
-	FEvent Event;
-	while (Events.Dequeue(Event))
+	FGamepadKeyNames::Type FInputDevice::ConvertAxisIndexToGamepadAxis(uint8 AnalogAxis)
 	{
-		switch (Event.Event)
+		switch (AnalogAxis)
 		{
-			case EventType::UNDEFINED:
+			case 1:
 			{
-				checkNoEntry();
+				return FGamepadKeyNames::LeftAnalogX;
 			}
 			break;
-			case EventType::KEY_DOWN:
+			case 2:
 			{
-				uint8 JavaScriptKeyCode;
-				bool IsRepeat;
-				Event.GetKeyDown(JavaScriptKeyCode, IsRepeat);
-				const FKey* AgnosticKey = UE::PixelStreaming::JavaScriptKeyCodeToFKey[JavaScriptKeyCode];
-				if (FilterKey(*AgnosticKey))
-				{
-					const uint32* KeyCodePtr;
-					const uint32* CharacterCodePtr;
-					FInputKeyManager::Get().GetCodesFromKey(*AgnosticKey, KeyCodePtr, CharacterCodePtr);
-					uint32 KeyCode = KeyCodePtr ? *KeyCodePtr : 0;
-					uint32 CharacterCode = CharacterCodePtr ? *CharacterCodePtr : 0;
-					MessageHandler->OnKeyDown(KeyCode, CharacterCode, IsRepeat);
-					UE_LOG(LogPixelStreamingInputDevice, Verbose, TEXT("KEY_DOWN: KeyCode = %d; CharacterCode = %d; IsRepeat = %s"), KeyCode, CharacterCode, IsRepeat ? TEXT("True") : TEXT("False"));
-				}
+				return FGamepadKeyNames::LeftAnalogY;
 			}
 			break;
-			case EventType::KEY_UP:
+			case 3:
 			{
-				uint8 JavaScriptKeyCode;
-				Event.GetKeyUp(JavaScriptKeyCode);
-				const FKey* AgnosticKey = UE::PixelStreaming::JavaScriptKeyCodeToFKey[JavaScriptKeyCode];
-				if (FilterKey(*AgnosticKey))
-				{
-					const uint32* KeyCodePtr;
-					const uint32* CharacterCodePtr;
-					FInputKeyManager::Get().GetCodesFromKey(*AgnosticKey, KeyCodePtr, CharacterCodePtr);
-					uint32 KeyCode = KeyCodePtr ? *KeyCodePtr : 0;
-					uint32 CharacterCode = CharacterCodePtr ? *CharacterCodePtr : 0;
-					MessageHandler->OnKeyUp(KeyCode, CharacterCode, false); // Key up events are never repeats.
-					UE_LOG(LogPixelStreamingInputDevice, Verbose, TEXT("KEY_UP: KeyCode = %d; CharacterCode = %d"), KeyCode, CharacterCode);
-				}
+				return FGamepadKeyNames::RightAnalogX;
 			}
 			break;
-			case EventType::KEY_PRESS:
+			case 4:
 			{
-				TCHAR UnicodeCharacter;
-				Event.GetCharacterCode(UnicodeCharacter);
-				MessageHandler->OnKeyChar(UnicodeCharacter, false); // Key press repeat not yet available but are not intrinsically used.
-				UE_LOG(LogPixelStreamingInputDevice, Verbose, TEXT("KEY_PRESSED: Character = '%c'"), UnicodeCharacter);
+				return FGamepadKeyNames::RightAnalogY;
 			}
 			break;
-			case EventType::MOUSE_ENTER:
+			case 5:
 			{
-				// Override application layer to special pixel streaming version.
-				FSlateApplication::Get().OverridePlatformApplication(PixelStreamerApplicationWrapper);
-				FSlateApplication::Get().OnCursorSet();
-
-				// Make sure the viewport is active.
-				FSlateApplication::Get().ProcessApplicationActivationEvent(true);
-
-				UE_LOG(LogPixelStreamingInputDevice, Verbose, TEXT("MOUSE_ENTER"));
+				return FGamepadKeyNames::LeftTriggerAnalog;
+			}
+			break;
+			case 6:
+			{
+				return FGamepadKeyNames::RightTriggerAnalog;
 			}
 			break;
-			case EventType::MOUSE_LEAVE:
+			default:
 			{
-				// Restore normal application layer.
-				FSlateApplication::Get().OverridePlatformApplication(PixelStreamerApplicationWrapper->WrappedApplication);
+				return FGamepadKeyNames::Invalid;
+			}
+			break;
+		}
+	}
 
-				UE_LOG(LogPixelStreamingInputDevice, Verbose, TEXT("MOUSE_LEAVE"));
+	FGamepadKeyNames::Type FInputDevice::ConvertButtonIndexToGamepadButton(uint8 ButtonIndex)
+	{
+		switch (ButtonIndex)
+		{
+			case 0:
+			{
+				return FGamepadKeyNames::FaceButtonBottom;
+			}
+			case 1:
+			{
+				return FGamepadKeyNames::FaceButtonRight;
 			}
 			break;
-			case EventType::MOUSE_MOVE:
+			case 2:
 			{
-				// Ensure the cursor will be updated so it can change appearance.
-				FSlateApplication::Get().OnCursorSet();
-
-				uint16 PosX;
-				uint16 PosY;
-				int16 DeltaX;
-				int16 DeltaY;
-				Event.GetMouseDelta(PosX, PosY, DeltaX, DeltaY);
-				FVector2D CursorPos = GEngine->GameViewport->GetWindow()->GetPositionInScreen() + FVector2D(PosX, PosY);
-				PixelStreamerApplicationWrapper->Cursor->SetPosition(CursorPos.X, CursorPos.Y);
-				MessageHandler->OnRawMouseMove(DeltaX, DeltaY);
-				UE_LOG(LogPixelStreamingInputDevice, VeryVerbose, TEXT("MOUSE_MOVE: Pos = (%d, %d); CursorPos = (%d, %d); Delta = (%d, %d)"), PosX, PosY, static_cast<int>(CursorPos.X), static_cast<int>(CursorPos.Y), DeltaX, DeltaY);
+				return FGamepadKeyNames::FaceButtonLeft;
 			}
 			break;
-			case EventType::MOUSE_DOWN:
+			case 3:
 			{
-				// If a user clicks on the application window and then clicks on the
-				// browser then this will move the focus away from the application
-				// window which will deactivate the application, so we need to check
-				// if we must reactivate the application.
-				if (!FSlateApplication::Get().IsActive())
-				{
-					FSlateApplication::Get().ProcessApplicationActivationEvent(true);
-				}
-
-				EMouseButtons::Type Button;
-				uint16 PosX;
-				uint16 PosY;
-				Event.GetMouseClick(Button, PosX, PosY);
-				FVector2D CursorPos = GEngine->GameViewport->GetWindow()->GetPositionInScreen() + FVector2D(PosX, PosY);
-				PixelStreamerApplicationWrapper->Cursor->SetPosition(CursorPos.X, CursorPos.Y);
-				MessageHandler->OnMouseDown(GEngine->GameViewport->GetWindow()->GetNativeWindow(), Button, CursorPos);
-				UE_LOG(LogPixelStreamingInputDevice, Verbose, TEXT("MOUSE_DOWN: Button = %d; Pos = (%d, %d); CursorPos = (%d, %d)"), Button, PosX, PosY, static_cast<int>(CursorPos.X), static_cast<int>(CursorPos.Y));
-
-				// The browser may be faking a mouse when touching so it will send
-				// over a mouse down event.
-				FindFocusedWidget();
+				return FGamepadKeyNames::FaceButtonTop;
 			}
 			break;
-			case EventType::MOUSE_UP:
+			case 4:
 			{
-				EMouseButtons::Type Button;
-				uint16 PosX;
-				uint16 PosY;
-				Event.GetMouseClick(Button, PosX, PosY);
-				FVector2D CursorPos = GEngine->GameViewport->GetWindow()->GetPositionInScreen() + FVector2D(PosX, PosY);
-				PixelStreamerApplicationWrapper->Cursor->SetPosition(CursorPos.X, CursorPos.Y);
-				MessageHandler->OnMouseUp(Button);
-				UE_LOG(LogPixelStreamingInputDevice, Verbose, TEXT("MOUSE_UP: Button = %d; Pos = (%d, %d); CursorPos = (%d, %d)"), Button, PosX, PosY, static_cast<int>(CursorPos.X), static_cast<int>(CursorPos.Y));
+				return FGamepadKeyNames::LeftShoulder;
 			}
 			break;
-			case EventType::MOUSE_WHEEL:
+			case 5:
 			{
-				int16 Delta;
-				uint16 PosX;
-				uint16 PosY;
-				Event.GetMouseWheel(Delta, PosX, PosY);
-				const float SpinFactor = 1 / 120.0f;
-				FVector2D CursorPos = GEngine->GameViewport->GetWindow()->GetPositionInScreen() + FVector2D(PosX, PosY);
-				MessageHandler->OnMouseWheel(Delta * SpinFactor, CursorPos);
-				UE_LOG(LogPixelStreamingInputDevice, Verbose, TEXT("MOUSE_WHEEL: Delta = %d; Pos = (%d, %d); CursorPos = (%d, %d)"), Delta, PosX, PosY, static_cast<int>(CursorPos.X), static_cast<int>(CursorPos.Y));
+				return FGamepadKeyNames::RightShoulder;
 			}
 			break;
-			case EventType::TOUCH_START:
+			// Buttons 6 and 7 are mapped as analog axis as they are the triggers
+			case 8:
 			{
-				uint8 TouchIndex;
-				uint16 PosX;
-				uint16 PosY;
-				uint8 Force; // Force is between 0.0 and 1.0 so will need to unquantize from byte.
-				Event.GetTouch(TouchIndex, PosX, PosY, Force);
-				FVector2D CursorPos = GEngine->GameViewport->GetWindow()->GetPositionInScreen() + FVector2D(PosX, PosY);
-				MessageHandler->OnTouchStarted(GEngine->GameViewport->GetWindow()->GetNativeWindow(), CursorPos, Force / 255.0f, TouchIndex, 0); // TODO: ControllerId?
-				UE_LOG(LogPixelStreamingInputDevice, Verbose, TEXT("TOUCH_START: TouchIndex = %d; Pos = (%d, %d); CursorPos = (%d, %d); Force = %.3f"), TouchIndex, PosX, PosY, static_cast<int>(CursorPos.X), static_cast<int>(CursorPos.Y), Force / 255.0f);
-
-				// If the user starts a touch then they may be focusing an editable
-				// widget.
-				FindFocusedWidget();
+				return FGamepadKeyNames::SpecialLeft;
 			}
 			break;
-			case EventType::TOUCH_END:
+			case 9:
 			{
-				uint8 TouchIndex;
-				uint16 PosX;
-				uint16 PosY;
-				uint8 Force;
-				Event.GetTouch(TouchIndex, PosX, PosY, Force);
-				FVector2D CursorPos = GEngine->GameViewport->GetWindow()->GetPositionInScreen() + FVector2D(PosX, PosY);
-				MessageHandler->OnTouchEnded(CursorPos, TouchIndex, 0); // TODO: ControllerId?
-				UE_LOG(LogPixelStreamingInputDevice, Verbose, TEXT("TOUCH_END: TouchIndex = %d; Pos = (%d, %d); CursorPos = (%d, %d)"), TouchIndex, PosX, PosY, static_cast<int>(CursorPos.X), static_cast<int>(CursorPos.Y));
+				return FGamepadKeyNames::SpecialRight;
+			}
+			case 10:
+			{
+				return FGamepadKeyNames::LeftThumb;
 			}
 			break;
-			case EventType::TOUCH_MOVE:
+			case 11:
 			{
-				uint8 TouchIndex;
-				uint16 PosX;
-				uint16 PosY;
-				uint8 Force; // Force is between 0.0 and 1.0 so will need to unquantize from byte.
-				Event.GetTouch(TouchIndex, PosX, PosY, Force);
-				FVector2D CursorPos = GEngine->GameViewport->GetWindow()->GetPositionInScreen() + FVector2D(PosX, PosY);
-				MessageHandler->OnTouchMoved(CursorPos, Force / 255.0f, TouchIndex, 0); // TODO: ControllerId?
-				UE_LOG(LogPixelStreamingInputDevice, VeryVerbose, TEXT("TOUCH_MOVE: TouchIndex = %d; Pos = (%d, %d); CursorPos = (%d, %d); Force = %.3f"), TouchIndex, PosX, PosY, static_cast<int>(CursorPos.X), static_cast<int>(CursorPos.Y), Force / 255.0f);
+				return FGamepadKeyNames::RightThumb;
 			}
 			break;
-			case EventType::GAMEPAD_PRESS:
+			case 12:
 			{
-				uint8 ControllerId;
-				uint8 ButtonIndex;
-				bool bIsRepeat;
-				Event.GetGamepadButtonPressed(ControllerId, ButtonIndex, bIsRepeat);
-				FGamepadKeyNames::Type ControllerButton = ConvertButtonIndexToGamepadButton(ButtonIndex);
-				if (ControllerButton == FGamepadKeyNames::Invalid)
-				{
-					break;
-				}
-				MessageHandler->OnControllerButtonPressed(ControllerButton, (int32_t)ControllerId, bIsRepeat);
-				UE_LOG(LogPixelStreamingInputDevice, VeryVerbose, TEXT("GAMEPAD_PRESS: ControllerId = %d; ButtonIndex = %d; IsRepeat = %d"), ControllerId, ButtonIndex, bIsRepeat);
+				return FGamepadKeyNames::DPadUp;
 			}
 			break;
-			case EventType::GAMEPAD_RELEASE:
+			case 13:
 			{
-				uint8 ControllerId;
-				uint8 ButtonIndex;
-				Event.GetGamepadButtonReleased(ControllerId, ButtonIndex);
-				FGamepadKeyNames::Type ControllerButton = ConvertButtonIndexToGamepadButton(ButtonIndex);
-				if (ControllerButton == FGamepadKeyNames::Invalid)
-				{
-					break;
-				}
-				MessageHandler->OnControllerButtonReleased(ControllerButton, (int32_t)ControllerId, false);
-				UE_LOG(LogPixelStreamingInputDevice, VeryVerbose, TEXT("GAMEPAD_RELEASE: ControllerId = %d; ButtonIndex = %d;"), ControllerId, ButtonIndex);
+				return FGamepadKeyNames::DPadDown;
+			}
+			case 14:
+			{
+				return FGamepadKeyNames::DPadLeft;
 			}
 			break;
-			case EventType::GAMEPAD_ANALOG:
+			case 15:
 			{
-				uint8 ControllerId;
-				uint8 AxisIndex;
-				float AnalogValue;
-				Event.GetGamepadAnalog(ControllerId, AxisIndex, AnalogValue);
-				FGamepadKeyNames::Type ControllerAxis = ConvertAxisIndexToGamepadAxis(AxisIndex);
-				if (ControllerAxis == FGamepadKeyNames::Invalid)
-				{
-					break;
-				}
-				MessageHandler->OnControllerAnalog(ControllerAxis, (int32_t)ControllerId, AnalogValue);
-				UE_LOG(LogPixelStreamingInputDevice, VeryVerbose, TEXT("GAMEPAD_ANALOG: ControllerId = %d; AxisIndex = %d; AnalogValue = %.4f;"), ControllerId, AxisIndex, AnalogValue);
+				return FGamepadKeyNames::DPadRight;
 			}
 			break;
 			default:
 			{
-				UE_LOG(LogPixelStreamingInputDevice, Error, TEXT("Unknown Pixel Streaming event %d with word 0x%016llx"), static_cast<int>(Event.Event), Event.Data.Word);
+				return FGamepadKeyNames::Invalid;
 			}
 			break;
 		}
 	}
 
-	FString UIInteraction;
-	while (UIInteractions.Dequeue(UIInteraction))
+	void FInputDevice::SendControllerEvents()
 	{
-		for (UPixelStreamingInput* InputComponent : PixelStreamingModule->GetInputComponents())
-		{
-			InputComponent->OnInputEvent.Broadcast(UIInteraction);
-			UE_LOG(LogPixelStreamingInputDevice, Verbose, TEXT("UIInteraction = %s"), *UIInteraction);
-		}
 	}
 
-	FString Command;
-	while (Commands.Dequeue(Command))
+	void FInputDevice::SetMessageHandler(const TSharedRef<FGenericApplicationMessageHandler>& InMessageHandler)
 	{
-		UPixelStreamingInput::OnCommand(Command);
+		MessageHandler = InMessageHandler;
 	}
-}
 
-FGamepadKeyNames::Type UE::PixelStreaming::FInputDevice::ConvertAxisIndexToGamepadAxis(uint8 AnalogAxis)
-{
-	switch (AnalogAxis)
+	bool FInputDevice::Exec(UWorld* InWorld, const TCHAR* Cmd, FOutputDevice& Ar)
 	{
-		case 1:
-		{
-			return FGamepadKeyNames::LeftAnalogX;
-		}
-		break;
-		case 2:
-		{
-			return FGamepadKeyNames::LeftAnalogY;
-		}
-		break;
-		case 3:
-		{
-			return FGamepadKeyNames::RightAnalogX;
-		}
-		break;
-		case 4:
-		{
-			return FGamepadKeyNames::RightAnalogY;
-		}
-		break;
-		case 5:
-		{
-			return FGamepadKeyNames::LeftTriggerAnalog;
-		}
-		break;
-		case 6:
-		{
-			return FGamepadKeyNames::RightTriggerAnalog;
-		}
-		break;
-		default:
-		{
-			return FGamepadKeyNames::Invalid;
-		}
-		break;
+		return true;
 	}
-}
 
-FGamepadKeyNames::Type UE::PixelStreaming::FInputDevice::ConvertButtonIndexToGamepadButton(uint8 ButtonIndex)
-{
-	switch (ButtonIndex)
+	void FInputDevice::SetChannelValue(int32 ControllerId, FForceFeedbackChannelType ChannelType, float Value)
 	{
-		case 0:
-		{
-			return FGamepadKeyNames::FaceButtonBottom;
-		}
-		case 1:
-		{
-			return FGamepadKeyNames::FaceButtonRight;
-		}
-		break;
-		case 2:
-		{
-			return FGamepadKeyNames::FaceButtonLeft;
-		}
-		break;
-		case 3:
-		{
-			return FGamepadKeyNames::FaceButtonTop;
-		}
-		break;
-		case 4:
-		{
-			return FGamepadKeyNames::LeftShoulder;
-		}
-		break;
-		case 5:
-		{
-			return FGamepadKeyNames::RightShoulder;
-		}
-		break;
-		// Buttons 6 and 7 are mapped as analog axis as they are the triggers
-		case 8:
-		{
-			return FGamepadKeyNames::SpecialLeft;
-		}
-		break;
-		case 9:
-		{
-			return FGamepadKeyNames::SpecialRight;
-		}
-		case 10:
-		{
-			return FGamepadKeyNames::LeftThumb;
-		}
-		break;
-		case 11:
-		{
-			return FGamepadKeyNames::RightThumb;
-		}
-		break;
-		case 12:
-		{
-			return FGamepadKeyNames::DPadUp;
-		}
-		break;
-		case 13:
-		{
-			return FGamepadKeyNames::DPadDown;
-		}
-		case 14:
-		{
-			return FGamepadKeyNames::DPadLeft;
-		}
-		break;
-		case 15:
-		{
-			return FGamepadKeyNames::DPadRight;
-		}
-		break;
-		default:
-		{
-			return FGamepadKeyNames::Invalid;
-		}
-		break;
 	}
-}
-
-void UE::PixelStreaming::FInputDevice::SendControllerEvents()
-{
-}
-
-void UE::PixelStreaming::FInputDevice::SetMessageHandler(const TSharedRef<FGenericApplicationMessageHandler>& InMessageHandler)
-{
-	MessageHandler = InMessageHandler;
-}
-
-bool UE::PixelStreaming::FInputDevice::Exec(UWorld* InWorld, const TCHAR* Cmd, FOutputDevice& Ar)
-{
-	return true;
-}
-
-void UE::PixelStreaming::FInputDevice::SetChannelValue(int32 ControllerId, FForceFeedbackChannelType ChannelType, float Value)
-{
-}
-
-void UE::PixelStreaming::FInputDevice::SetChannelValues(int32 ControllerId, const FForceFeedbackValues& values)
-{
-}
 
-void UE::PixelStreaming::FInputDevice::ProcessEvent(const FEvent& InEvent)
-{
-	bool bSuccess = Events.Enqueue(InEvent);
-	checkf(bSuccess, TEXT("Unable to enqueue new event of type %d"), static_cast<int>(InEvent.Event));
-}
-
-void UE::PixelStreaming::FInputDevice::ProcessUIInteraction(const FString& InDescriptor)
-{
-	bool bSuccess = UIInteractions.Enqueue(InDescriptor);
-	checkf(bSuccess, TEXT("Unable to enqueue new UI Interaction %s"), *InDescriptor);
-}
-
-void UE::PixelStreaming::FInputDevice::ProcessCommand(const FString& InDescriptor)
-{
-	if (UE::PixelStreaming::Settings::CVarPixelStreamingAllowConsoleCommands.GetValueOnAnyThread())
+	void FInputDevice::SetChannelValues(int32 ControllerId, const FForceFeedbackValues& values)
 	{
-		bool bSuccess = Commands.Enqueue(InDescriptor);
-		checkf(bSuccess, TEXT("Unable to enqueue new Command %s"), *InDescriptor);
 	}
-}
-
-void UE::PixelStreaming::FInputDevice::FindFocusedWidget()
-{
-	FSlateApplication::Get().ForEachUser([this](FSlateUser& User) {
-		TSharedPtr<SWidget> FocusedWidget = User.GetFocusedWidget();
-
-		static FName SEditableTextType(TEXT("SEditableText"));
-		static FName SMultiLineEditableTextType(TEXT("SMultiLineEditableText"));
-		bool bEditable = FocusedWidget && (FocusedWidget->GetType() == SEditableTextType || FocusedWidget->GetType() == SMultiLineEditableTextType);
-
-		// Check to see if the focus has changed.
-		FVector2D Pos = bEditable ? FocusedWidget->GetCachedGeometry().GetAbsolutePosition() : UnfocusedPos;
-		if (Pos != FocusedPos)
-		{
-			FocusedPos = Pos;
-
-			// Tell the browser that the focus has changed.
-			TSharedPtr<FJsonObject> JsonObject = MakeShareable(new FJsonObject);
-			JsonObject->SetStringField(TEXT("command"), TEXT("onScreenKeyboard"));
-			JsonObject->SetBoolField(TEXT("showOnScreenKeyboard"), bEditable);
 
-			if (bEditable)
-			{
-				Pos = Pos - GEngine->GameViewport->GetWindow()->GetPositionInScreen();
-				uint16 PosX;
-				uint16 PosY;
-				QuantizeAndNormalize(Pos, PosX, PosY);
-				JsonObject->SetNumberField(TEXT("x"), PosX);
-				JsonObject->SetNumberField(TEXT("y"), PosY);
-			}
-
-			FString Descriptor;
-			TSharedRef<TJsonWriter<TCHAR, TCondensedJsonPrintPolicy<TCHAR>>> JsonWriter = TJsonWriterFactory<TCHAR, TCondensedJsonPrintPolicy<TCHAR>>::Create(&Descriptor);
-			FJsonSerializer::Serialize(JsonObject.ToSharedRef(), JsonWriter);
-
-			PixelStreamingModule->SendCommand(Descriptor);
-		}
-	});
-}
-
-// player message handling
-
-namespace
-{
-
-#define GET(Type, Var) Type Var = UE::PixelStreaming::ParseBuffer<Type>(Data, Size)
-
-	// XY positions are the ratio (0.0..1.0) along a viewport axis, quantized
-	// into an uint16 (0..65536). This allows the browser viewport and player
-	// viewport to have a different size.
-	void UnquantizeAndDenormalize(uint16& InOutX, uint16& InOutY)
+	void FInputDevice::ProcessEvent(const FEvent& InEvent)
 	{
-		checkf(GEngine->GameViewport != nullptr, TEXT("PixelStreaming is not supported in editor, and should have been disabled earlier"));
-		FIntPoint SizeXY = GEngine->GameViewport->Viewport->GetSizeXY();
-		InOutX = InOutX / 65536.0f * SizeXY.X;
-		InOutY = InOutY / 65536.0f * SizeXY.Y;
+		bool bSuccess = Events.Enqueue(InEvent);
+		checkf(bSuccess, TEXT("Unable to enqueue new event of type %d"), static_cast<int>(InEvent.Event));
 	}
 
-	// XY deltas are the ratio (-1.0..1.0) along a viewport axis, quantized
-	// into an int16 (-32767..32767). This allows the browser viewport and
-	// player viewport to have a different size.
-	void UnquantizeAndDenormalize(int16& InOutX, int16& InOutY)
+	void FInputDevice::ProcessUIInteraction(const FString& InDescriptor)
 	{
-		FIntPoint SizeXY = GEngine->GameViewport->Viewport->GetSizeXY();
-		InOutX = InOutX / 32767.0f * SizeXY.X;
-		InOutY = InOutY / 32767.0f * SizeXY.Y;
+		bool bSuccess = UIInteractions.Enqueue(InDescriptor);
+		checkf(bSuccess, TEXT("Unable to enqueue new UI Interaction %s"), *InDescriptor);
 	}
 
-	/**
-	* A touch is a specific finger placed on the canvas as a specific position.
-	*/
-	struct FTouch
-	{
-		uint16 PosX;	  // X position of finger.
-		uint16 PosY;	  // Y position of finger.
-		uint8 TouchIndex; // Index of finger for tracking multi-touch events.
-		uint8 Force;	  // Amount of pressure being applied by the finger.
-		uint8 Valid;	  // 1 if the touch was within bounds.
-	};
-
-	using FKeyCodeType = uint8;
-	using FCharacterType = TCHAR;
-	using FRepeatType = uint8;
-	using FButtonType = uint8;
-	using FPosType = uint16;
-	using FDeltaType = int16;
-	using FTouchesType = TArray<FTouch>;
-	using FControllerIndex = uint8;
-	using FControllerButtonIndex = uint8;
-	using FControllerAnalog = double;
-	using FControllerAxis = uint8;
-
-	/**
-	* Get the array of touch positions and touch indices for a touch event,
-	* consumed from the receive buffer.
-	* @param Consumed - The number of bytes consumed from the receive buffer.
-	* @param OutTouches - The array of touches.
-	* @return False if there insufficient room in the receive buffer to read the entire event.
-	*/
-	FTouchesType GetTouches(const uint8*& Data, uint32& Size)
+	void FInputDevice::ProcessCommand(const FString& InDescriptor)
 	{
-		// Get the number of touches in the array.
-		GET(uint8, NumTouches);
-
-		FTouchesType Touches;
-		// Get the value of each touch position and then the touch index.
-		for (int Touch = 0; Touch < NumTouches; Touch++)
+		if (UE::PixelStreaming::Settings::CVarPixelStreamingAllowConsoleCommands.GetValueOnAnyThread())
 		{
-			GET(FPosType, PosX);
-			GET(FPosType, PosY);
-			UnquantizeAndDenormalize(PosX, PosY);
-			GET(uint8, TouchIndex);
-			GET(uint8, Force);
-			GET(uint8, Valid);
-			Touches.Add({ PosX, PosY, TouchIndex, Force, Valid });
+			bool bSuccess = Commands.Enqueue(InDescriptor);
+			checkf(bSuccess, TEXT("Unable to enqueue new Command %s"), *InDescriptor);
 		}
-
-		return Touches;
 	}
 
-	/**
-	* Convert the given array of touches to a friendly string for logging.
-	* @param InTouches - The array of touches.
-	* @return The string representation of the array.
-	*/
-	FString TouchesToString(const FTouchesType& InTouches)
+	void FInputDevice::FindFocusedWidget()
 	{
-		FString String;
-		for (const FTouch& Touch : InTouches)
-		{
-			String += FString::Printf(TEXT("F[%d]=(%d, %d)(%.3f)"), Touch.TouchIndex, Touch.PosX, Touch.PosY, Touch.Force / 255.0f);
-		}
-		return String;
-	}
+		FSlateApplication::Get().ForEachUser([this](FSlateUser& User) {
+			TSharedPtr<SWidget> FocusedWidget = User.GetFocusedWidget();
 
-	enum class KeyState
-	{
-		Alt = 1 << 0,
-		Ctrl = 1 << 1,
-		Shift = 1 << 2
-	};
-	enum class MouseButtonState
-	{
-		Left = 1 << 0,
-		Right = 1 << 1,
-		Middle = 1 << 2,
-		Button4 = 1 << 3,
-		Button5 = 1 << 4,
-		Button6 = 1 << 5,
-		Button7 = 1 << 6,
-		Button8 = 1 << 7
-	};
-} // namespace
+			static FName SEditableTextType(TEXT("SEditableText"));
+			static FName SMultiLineEditableTextType(TEXT("SMultiLineEditableText"));
+			bool bEditable = FocusedWidget && (FocusedWidget->GetType() == SEditableTextType || FocusedWidget->GetType() == SMultiLineEditableTextType);
 
-void UE::PixelStreaming::FInputDevice::OnMessage(const webrtc::DataBuffer& Buffer)
-{
-	using namespace UE::PixelStreaming::Protocol;
+			// Check to see if the focus has changed.
+			FVector2D Pos = bEditable ? FocusedWidget->GetCachedGeometry().GetAbsolutePosition() : UnfocusedPos;
+			if (Pos != FocusedPos)
+			{
+				FocusedPos = Pos;
+
+				// Tell the browser that the focus has changed.
+				TSharedPtr<FJsonObject> JsonObject = MakeShareable(new FJsonObject);
+				JsonObject->SetStringField(TEXT("command"), TEXT("onScreenKeyboard"));
+				JsonObject->SetBoolField(TEXT("showOnScreenKeyboard"), bEditable);
+
+				if (bEditable)
+				{
+					Pos = Pos - GEngine->GameViewport->GetWindow()->GetPositionInScreen();
+					uint16 PosX;
+					uint16 PosY;
+					QuantizeAndNormalize(Pos, PosX, PosY);
+					JsonObject->SetNumberField(TEXT("x"), PosX);
+					JsonObject->SetNumberField(TEXT("y"), PosY);
+				}
 
-	const uint8* Data = Buffer.data.data();
-	uint32 Size = (uint32)Buffer.data.size();
+				FString Descriptor;
+				TSharedRef<TJsonWriter<TCHAR, TCondensedJsonPrintPolicy<TCHAR>>> JsonWriter = TJsonWriterFactory<TCHAR, TCondensedJsonPrintPolicy<TCHAR>>::Create(&Descriptor);
+				FJsonSerializer::Serialize(JsonObject.ToSharedRef(), JsonWriter);
 
-	GET(EToStreamerMsg, MsgType);
+				PixelStreamingModule->SendCommand(Descriptor);
+			}
+		});
+	}
+
+	// player message handling
 
-	switch (MsgType)
+	namespace
 	{
-		case EToStreamerMsg::UIInteraction:
-		{
-			FString Descriptor = UE::PixelStreaming::ParseString(Buffer, UE::PixelStreaming::FInputDevice::GetMessageHeaderOffset());
-			UE_LOG(LogPixelStreamingInput, Verbose, TEXT("UIInteraction: %s"), *Descriptor);
-			ProcessUIInteraction(Descriptor);
-			break;
-		}
-		case EToStreamerMsg::Command:
-		{
-			FString Descriptor = UE::PixelStreaming::ParseString(Buffer, UE::PixelStreaming::FInputDevice::GetMessageHeaderOffset());
-			UE_LOG(LogPixelStreamingInput, Verbose, TEXT("Command: %s"), *Descriptor);
-			ProcessCommand(Descriptor);
-			break;
-		}
-		case EToStreamerMsg::KeyDown:
-		{
-			GET(FKeyCodeType, KeyCode);
-			GET(FRepeatType, Repeat);
-			checkf(Size == 0, TEXT("%d"), Size);
-			UE_LOG(LogPixelStreamingInput, Verbose, TEXT("key down: %d, repeat: %d"), KeyCode, Repeat);
-
-			FEvent KeyDownEvent(EventType::KEY_DOWN);
-			KeyDownEvent.SetKeyDown(KeyCode, Repeat != 0);
-			ProcessEvent(KeyDownEvent);
-			break;
-		}
-		case EToStreamerMsg::KeyUp:
-		{
-			GET(FKeyCodeType, KeyCode);
-			checkf(Size == 0, TEXT("%d"), Size);
-			UE_LOG(LogPixelStreamingInput, Verbose, TEXT("key up: %d"), KeyCode);
 
-			FEvent KeyUpEvent(EventType::KEY_UP);
-			KeyUpEvent.SetKeyUp(KeyCode);
-			ProcessEvent(KeyUpEvent);
-			break;
-		}
-		case EToStreamerMsg::KeyPress:
-		{
-			GET(FCharacterType, Character);
-			checkf(Size == 0, TEXT("%d"), Size);
-			UE_LOG(LogPixelStreamingInput, Verbose, TEXT("key press: '%c'"), Character);
+	#define GET(Type, Var) Type Var = UE::PixelStreaming::ParseBuffer<Type>(Data, Size)
 
-			FEvent KeyPressEvent(EventType::KEY_PRESS);
-			KeyPressEvent.SetCharCode(Character);
-			ProcessEvent(KeyPressEvent);
-			break;
-		}
-		case EToStreamerMsg::MouseEnter:
+		// XY positions are the ratio (0.0..1.0) along a viewport axis, quantized
+		// into an uint16 (0..65536). This allows the browser viewport and player
+		// viewport to have a different size.
+		void UnquantizeAndDenormalize(uint16& InOutX, uint16& InOutY)
 		{
-			checkf(Size == 0, TEXT("%d"), Size);
-			UE_LOG(LogPixelStreamingInput, Verbose, TEXT("mouseEnter"));
-			ProcessEvent(FEvent(UE::PixelStreaming::FInputDevice::EventType::MOUSE_ENTER));
-			break;
+			checkf(GEngine->GameViewport != nullptr, TEXT("PixelStreaming is not supported in editor, and should have been disabled earlier"));
+			FIntPoint SizeXY = GEngine->GameViewport->Viewport->GetSizeXY();
+			InOutX = InOutX / 65536.0f * SizeXY.X;
+			InOutY = InOutY / 65536.0f * SizeXY.Y;
 		}
-		case EToStreamerMsg::MouseLeave:
+
+		// XY deltas are the ratio (-1.0..1.0) along a viewport axis, quantized
+		// into an int16 (-32767..32767). This allows the browser viewport and
+		// player viewport to have a different size.
+		void UnquantizeAndDenormalize(int16& InOutX, int16& InOutY)
 		{
-			checkf(Size == 0, TEXT("%d"), Size);
-			ProcessEvent(FEvent(UE::PixelStreaming::FInputDevice::EventType::MOUSE_LEAVE));
-			UE_LOG(LogPixelStreamingInput, Verbose, TEXT("mouseLeave"));
-			break;
+			FIntPoint SizeXY = GEngine->GameViewport->Viewport->GetSizeXY();
+			InOutX = InOutX / 32767.0f * SizeXY.X;
+			InOutY = InOutY / 32767.0f * SizeXY.Y;
 		}
-		case EToStreamerMsg::MouseDown:
+
+		/**
+		* A touch is a specific finger placed on the canvas as a specific position.
+		*/
+		struct FTouch
 		{
-			GET(FButtonType, Button);
-			GET(FPosType, PosX);
-			GET(FPosType, PosY);
-			checkf(Size == 0, TEXT("%d"), Size);
-			UE_LOG(LogPixelStreamingInput, Verbose, TEXT("mouseDown at (%d, %d), button %d"), PosX, PosY, Button);
+			uint16 PosX;	  // X position of finger.
+			uint16 PosY;	  // Y position of finger.
+			uint8 TouchIndex; // Index of finger for tracking multi-touch events.
+			uint8 Force;	  // Amount of pressure being applied by the finger.
+			uint8 Valid;	  // 1 if the touch was within bounds.
+		};
 
-			UnquantizeAndDenormalize(PosX, PosY);
+		using FKeyCodeType = uint8;
+		using FCharacterType = TCHAR;
+		using FRepeatType = uint8;
+		using FButtonType = uint8;
+		using FPosType = uint16;
+		using FDeltaType = int16;
+		using FTouchesType = TArray<FTouch>;
+		using FControllerIndex = uint8;
+		using FControllerButtonIndex = uint8;
+		using FControllerAnalog = double;
+		using FControllerAxis = uint8;
 
-			FEvent MouseDownEvent(EventType::MOUSE_DOWN);
-			MouseDownEvent.SetMouseClick(Button, PosX, PosY);
-			ProcessEvent(MouseDownEvent);
-			break;
-		}
-		case EToStreamerMsg::MouseUp:
+		/**
+		* Get the array of touch positions and touch indices for a touch event,
+		* consumed from the receive buffer.
+		* @param Consumed - The number of bytes consumed from the receive buffer.
+		* @param OutTouches - The array of touches.
+		* @return False if there insufficient room in the receive buffer to read the entire event.
+		*/
+		FTouchesType GetTouches(const uint8*& Data, uint32& Size)
 		{
-			GET(FButtonType, Button);
-			GET(FPosType, PosX);
-			GET(FPosType, PosY);
-			checkf(Size == 0, TEXT("%d"), Size);
-			UE_LOG(LogPixelStreamingInput, Verbose, TEXT("mouseUp at (%d, %d), button %d"), PosX, PosY, Button);
+			// Get the number of touches in the array.
+			GET(uint8, NumTouches);
 
-			UnquantizeAndDenormalize(PosX, PosY);
+			FTouchesType Touches;
+			// Get the value of each touch position and then the touch index.
+			for (int Touch = 0; Touch < NumTouches; Touch++)
+			{
+				GET(FPosType, PosX);
+				GET(FPosType, PosY);
+				UnquantizeAndDenormalize(PosX, PosY);
+				GET(uint8, TouchIndex);
+				GET(uint8, Force);
+				GET(uint8, Valid);
+				Touches.Add({ PosX, PosY, TouchIndex, Force, Valid });
+			}
 
-			FEvent MouseDownEvent(EventType::MOUSE_UP);
-			MouseDownEvent.SetMouseClick(Button, PosX, PosY);
-			ProcessEvent(MouseDownEvent);
-			break;
+			return Touches;
 		}
-		case EToStreamerMsg::MouseMove:
+
+		/**
+		* Convert the given array of touches to a friendly string for logging.
+		* @param InTouches - The array of touches.
+		* @return The string representation of the array.
+		*/
+		FString TouchesToString(const FTouchesType& InTouches)
 		{
-			GET(FPosType, PosX);
-			GET(FPosType, PosY);
-			GET(FDeltaType, DeltaX);
-			GET(FDeltaType, DeltaY);
-			checkf(Size == 0, TEXT("%d"), Size);
-			UE_LOG(LogPixelStreamingInput, Verbose, TEXT("mouseMove to (%d, %d), delta (%d, %d)"), PosX, PosY, DeltaX, DeltaY);
-
-			UnquantizeAndDenormalize(PosX, PosY);
-			UnquantizeAndDenormalize(DeltaX, DeltaY);
-
-			FEvent MouseMoveEvent(EventType::MOUSE_MOVE);
-			MouseMoveEvent.SetMouseDelta(PosX, PosY, DeltaX, DeltaY);
-			ProcessEvent(MouseMoveEvent);
-			break;
+			FString String;
+			for (const FTouch& Touch : InTouches)
+			{
+				String += FString::Printf(TEXT("F[%d]=(%d, %d)(%.3f)"), Touch.TouchIndex, Touch.PosX, Touch.PosY, Touch.Force / 255.0f);
+			}
+			return String;
 		}
-		case EToStreamerMsg::MouseWheel:
+
+		enum class KeyState
 		{
-			GET(FDeltaType, Delta);
-			GET(FPosType, PosX);
-			GET(FPosType, PosY);
-			checkf(Size == 0, TEXT("%d"), Size);
-			UE_LOG(LogPixelStreamingInput, Verbose, TEXT("mouseWheel, delta %d"), Delta);
+			Alt = 1 << 0,
+			Ctrl = 1 << 1,
+			Shift = 1 << 2
+		};
+		enum class MouseButtonState
+		{
+			Left = 1 << 0,
+			Right = 1 << 1,
+			Middle = 1 << 2,
+			Button4 = 1 << 3,
+			Button5 = 1 << 4,
+			Button6 = 1 << 5,
+			Button7 = 1 << 6,
+			Button8 = 1 << 7
+		};
+	} // namespace
 
-			UnquantizeAndDenormalize(PosX, PosY);
+	void FInputDevice::OnMessage(const webrtc::DataBuffer& Buffer)
+	{
+		using namespace Protocol;
 
-			FEvent MouseWheelEvent(EventType::MOUSE_WHEEL);
-			MouseWheelEvent.SetMouseWheel(Delta, PosX, PosY);
-			ProcessEvent(MouseWheelEvent);
-			break;
-		}
-		case EToStreamerMsg::TouchStart:
-		{
-			FTouchesType Touches = GetTouches(Data, Size);
-			checkf(Size == 0, TEXT("%d"), Size);
-			UE_LOG(LogPixelStreamingInput, Verbose, TEXT("TouchStart: %s"), *TouchesToString(Touches));
+		const uint8* Data = Buffer.data.data();
+		uint32 Size = (uint32)Buffer.data.size();
+
+		GET(EToStreamerMsg, MsgType);
 
-			for (const FTouch& Touch : Touches)
+		switch (MsgType)
+		{
+			case EToStreamerMsg::UIInteraction:
+			{
+				FString Descriptor = UE::PixelStreaming::ParseString(Buffer, FInputDevice::GetMessageHeaderOffset());
+				UE_LOG(LogPixelStreamingInput, Verbose, TEXT("UIInteraction: %s"), *Descriptor);
+				ProcessUIInteraction(Descriptor);
+				break;
+			}
+			case EToStreamerMsg::Command:
+			{
+				FString Descriptor = UE::PixelStreaming::ParseString(Buffer, FInputDevice::GetMessageHeaderOffset());
+				UE_LOG(LogPixelStreamingInput, Verbose, TEXT("Command: %s"), *Descriptor);
+				ProcessCommand(Descriptor);
+				break;
+			}
+			case EToStreamerMsg::KeyDown:
+			{
+				GET(FKeyCodeType, KeyCode);
+				GET(FRepeatType, Repeat);
+				checkf(Size == 0, TEXT("%d"), Size);
+				UE_LOG(LogPixelStreamingInput, Verbose, TEXT("key down: %d, repeat: %d"), KeyCode, Repeat);
+
+				FEvent KeyDownEvent(EventType::KEY_DOWN);
+				KeyDownEvent.SetKeyDown(KeyCode, Repeat != 0);
+				ProcessEvent(KeyDownEvent);
+				break;
+			}
+			case EToStreamerMsg::KeyUp:
+			{
+				GET(FKeyCodeType, KeyCode);
+				checkf(Size == 0, TEXT("%d"), Size);
+				UE_LOG(LogPixelStreamingInput, Verbose, TEXT("key up: %d"), KeyCode);
+
+				FEvent KeyUpEvent(EventType::KEY_UP);
+				KeyUpEvent.SetKeyUp(KeyCode);
+				ProcessEvent(KeyUpEvent);
+				break;
+			}
+			case EToStreamerMsg::KeyPress:
+			{
+				GET(FCharacterType, Character);
+				checkf(Size == 0, TEXT("%d"), Size);
+				UE_LOG(LogPixelStreamingInput, Verbose, TEXT("key press: '%c'"), Character);
+
+				FEvent KeyPressEvent(EventType::KEY_PRESS);
+				KeyPressEvent.SetCharCode(Character);
+				ProcessEvent(KeyPressEvent);
+				break;
+			}
+			case EToStreamerMsg::MouseEnter:
+			{
+				checkf(Size == 0, TEXT("%d"), Size);
+				UE_LOG(LogPixelStreamingInput, Verbose, TEXT("mouseEnter"));
+				ProcessEvent(FEvent(FInputDevice::EventType::MOUSE_ENTER));
+				break;
+			}
+			case EToStreamerMsg::MouseLeave:
+			{
+				checkf(Size == 0, TEXT("%d"), Size);
+				ProcessEvent(FEvent(FInputDevice::EventType::MOUSE_LEAVE));
+				UE_LOG(LogPixelStreamingInput, Verbose, TEXT("mouseLeave"));
+				break;
+			}
+			case EToStreamerMsg::MouseDown:
 			{
-				if (Touch.Valid)
+				GET(FButtonType, Button);
+				GET(FPosType, PosX);
+				GET(FPosType, PosY);
+				checkf(Size == 0, TEXT("%d"), Size);
+				UE_LOG(LogPixelStreamingInput, Verbose, TEXT("mouseDown at (%d, %d), button %d"), PosX, PosY, Button);
+
+				UnquantizeAndDenormalize(PosX, PosY);
+
+				FEvent MouseDownEvent(EventType::MOUSE_DOWN);
+				MouseDownEvent.SetMouseClick(Button, PosX, PosY);
+				ProcessEvent(MouseDownEvent);
+				break;
+			}
+			case EToStreamerMsg::MouseUp:
+			{
+				GET(FButtonType, Button);
+				GET(FPosType, PosX);
+				GET(FPosType, PosY);
+				checkf(Size == 0, TEXT("%d"), Size);
+				UE_LOG(LogPixelStreamingInput, Verbose, TEXT("mouseUp at (%d, %d), button %d"), PosX, PosY, Button);
+
+				UnquantizeAndDenormalize(PosX, PosY);
+
+				FEvent MouseDownEvent(EventType::MOUSE_UP);
+				MouseDownEvent.SetMouseClick(Button, PosX, PosY);
+				ProcessEvent(MouseDownEvent);
+				break;
+			}
+			case EToStreamerMsg::MouseMove:
+			{
+				GET(FPosType, PosX);
+				GET(FPosType, PosY);
+				GET(FDeltaType, DeltaX);
+				GET(FDeltaType, DeltaY);
+				checkf(Size == 0, TEXT("%d"), Size);
+				UE_LOG(LogPixelStreamingInput, Verbose, TEXT("mouseMove to (%d, %d), delta (%d, %d)"), PosX, PosY, DeltaX, DeltaY);
+
+				UnquantizeAndDenormalize(PosX, PosY);
+				UnquantizeAndDenormalize(DeltaX, DeltaY);
+
+				FEvent MouseMoveEvent(EventType::MOUSE_MOVE);
+				MouseMoveEvent.SetMouseDelta(PosX, PosY, DeltaX, DeltaY);
+				ProcessEvent(MouseMoveEvent);
+				break;
+			}
+			case EToStreamerMsg::MouseWheel:
+			{
+				GET(FDeltaType, Delta);
+				GET(FPosType, PosX);
+				GET(FPosType, PosY);
+				checkf(Size == 0, TEXT("%d"), Size);
+				UE_LOG(LogPixelStreamingInput, Verbose, TEXT("mouseWheel, delta %d"), Delta);
+
+				UnquantizeAndDenormalize(PosX, PosY);
+
+				FEvent MouseWheelEvent(EventType::MOUSE_WHEEL);
+				MouseWheelEvent.SetMouseWheel(Delta, PosX, PosY);
+				ProcessEvent(MouseWheelEvent);
+				break;
+			}
+			case EToStreamerMsg::TouchStart:
+			{
+				FTouchesType Touches = GetTouches(Data, Size);
+				checkf(Size == 0, TEXT("%d"), Size);
+				UE_LOG(LogPixelStreamingInput, Verbose, TEXT("TouchStart: %s"), *TouchesToString(Touches));
+
+				for (const FTouch& Touch : Touches)
 				{
-					FEvent TouchStartEvent(EventType::TOUCH_START);
-					TouchStartEvent.SetTouch(Touch.TouchIndex, Touch.PosX, Touch.PosY, Touch.Force);
-					ProcessEvent(TouchStartEvent);
+					if (Touch.Valid)
+					{
+						FEvent TouchStartEvent(EventType::TOUCH_START);
+						TouchStartEvent.SetTouch(Touch.TouchIndex, Touch.PosX, Touch.PosY, Touch.Force);
+						ProcessEvent(TouchStartEvent);
+					}
 				}
+				break;
 			}
-			break;
-		}
-		case EToStreamerMsg::TouchEnd:
-		{
-			FTouchesType Touches = GetTouches(Data, Size);
-			checkf(Size == 0, TEXT("%d"), Size);
-			UE_LOG(LogPixelStreamingInput, Verbose, TEXT("TouchEnd: %s"), *TouchesToString(Touches));
-
-			for (const FTouch& Touch : Touches)
+			case EToStreamerMsg::TouchEnd:
 			{
-				// Always allowing the "up" events regardless of in or outside the valid region so
-				// states aren't stuck "down". Might want to uncomment this if it causes other issues.
-				//if (Touch.Valid)
+				FTouchesType Touches = GetTouches(Data, Size);
+				checkf(Size == 0, TEXT("%d"), Size);
+				UE_LOG(LogPixelStreamingInput, Verbose, TEXT("TouchEnd: %s"), *TouchesToString(Touches));
+
+				for (const FTouch& Touch : Touches)
 				{
-					FEvent TouchEndEvent(EventType::TOUCH_END);
-					TouchEndEvent.SetTouch(Touch.TouchIndex, Touch.PosX, Touch.PosY, Touch.Force);
-					ProcessEvent(TouchEndEvent);
+					// Always allowing the "up" events regardless of in or outside the valid region so
+					// states aren't stuck "down". Might want to uncomment this if it causes other issues.
+					//if (Touch.Valid)
+					{
+						FEvent TouchEndEvent(EventType::TOUCH_END);
+						TouchEndEvent.SetTouch(Touch.TouchIndex, Touch.PosX, Touch.PosY, Touch.Force);
+						ProcessEvent(TouchEndEvent);
+					}
 				}
+				break;
 			}
-			break;
-		}
-		case EToStreamerMsg::TouchMove:
-		{
-			FTouchesType Touches = GetTouches(Data, Size);
-			checkf(Size == 0, TEXT("%d"), Size);
-			UE_LOG(LogPixelStreamingInput, Verbose, TEXT("TouchMove: %s"), *TouchesToString(Touches));
-
-			for (const FTouch& Touch : Touches)
+			case EToStreamerMsg::TouchMove:
 			{
-				if (Touch.Valid)
+				FTouchesType Touches = GetTouches(Data, Size);
+				checkf(Size == 0, TEXT("%d"), Size);
+				UE_LOG(LogPixelStreamingInput, Verbose, TEXT("TouchMove: %s"), *TouchesToString(Touches));
+
+				for (const FTouch& Touch : Touches)
 				{
-					FEvent TouchMoveEvent(EventType::TOUCH_MOVE);
-					TouchMoveEvent.SetTouch(Touch.TouchIndex, Touch.PosX, Touch.PosY, Touch.Force);
-					ProcessEvent(TouchMoveEvent);
+					if (Touch.Valid)
+					{
+						FEvent TouchMoveEvent(EventType::TOUCH_MOVE);
+						TouchMoveEvent.SetTouch(Touch.TouchIndex, Touch.PosX, Touch.PosY, Touch.Force);
+						ProcessEvent(TouchMoveEvent);
+					}
 				}
+				break;
+			}
+			case EToStreamerMsg::GamepadButtonPressed:
+			{
+				GET(FControllerIndex, ControllerIndex);
+				GET(FControllerButtonIndex, ButtonIndex);
+				GET(FRepeatType, Repeat);
+				checkf(Size == 0, TEXT("%d"), Size);
+
+				FEvent GamepadAnalogEvent(EventType::GAMEPAD_PRESS);
+				GamepadAnalogEvent.SetGamepadButtonPressed(ControllerIndex, ButtonIndex, Repeat != 0);
+				ProcessEvent(GamepadAnalogEvent);
+				break;
+			}
+			case EToStreamerMsg::GamepadButtonReleased:
+			{
+				GET(FControllerIndex, ControllerIndex);
+				GET(FControllerButtonIndex, ButtonIndex);
+				checkf(Size == 0, TEXT("%d"), Size);
+
+				FEvent GamepadAnalogEvent(EventType::GAMEPAD_RELEASE);
+				GamepadAnalogEvent.SetGamepadButtonReleased(ControllerIndex, ButtonIndex);
+				ProcessEvent(GamepadAnalogEvent);
+				break;
+			}
+			case EToStreamerMsg::GamepadAnalog:
+			{
+				GET(FControllerIndex, ControllerIndex);
+				GET(FControllerAxis, AxisIndex);
+				GET(FControllerAnalog, AnalogValue);
+				checkf(Size == 0, TEXT("%d"), Size);
+
+				FEvent GamepadAnalogEvent(EventType::GAMEPAD_ANALOG);
+				GamepadAnalogEvent.SetGamepadAnalog(ControllerIndex, AxisIndex, (float)AnalogValue);
+				ProcessEvent(GamepadAnalogEvent);
+				break;
 			}
-			break;
-		}
-		case EToStreamerMsg::GamepadButtonPressed:
-		{
-			GET(FControllerIndex, ControllerIndex);
-			GET(FControllerButtonIndex, ButtonIndex);
-			GET(FRepeatType, Repeat);
-			checkf(Size == 0, TEXT("%d"), Size);
-
-			FEvent GamepadAnalogEvent(EventType::GAMEPAD_PRESS);
-			GamepadAnalogEvent.SetGamepadButtonPressed(ControllerIndex, ButtonIndex, Repeat != 0);
-			ProcessEvent(GamepadAnalogEvent);
-			break;
-		}
-		case EToStreamerMsg::GamepadButtonReleased:
-		{
-			GET(FControllerIndex, ControllerIndex);
-			GET(FControllerButtonIndex, ButtonIndex);
-			checkf(Size == 0, TEXT("%d"), Size);
-
-			FEvent GamepadAnalogEvent(EventType::GAMEPAD_RELEASE);
-			GamepadAnalogEvent.SetGamepadButtonReleased(ControllerIndex, ButtonIndex);
-			ProcessEvent(GamepadAnalogEvent);
-			break;
-		}
-		case EToStreamerMsg::GamepadAnalog:
-		{
-			GET(FControllerIndex, ControllerIndex);
-			GET(FControllerAxis, AxisIndex);
-			GET(FControllerAnalog, AnalogValue);
-			checkf(Size == 0, TEXT("%d"), Size);
-
-			FEvent GamepadAnalogEvent(EventType::GAMEPAD_ANALOG);
-			GamepadAnalogEvent.SetGamepadAnalog(ControllerIndex, AxisIndex, (float)AnalogValue);
-			ProcessEvent(GamepadAnalogEvent);
-			break;
 		}
 	}
-}
 
-#undef GET
+	#undef GET
 
-void UE::PixelStreaming::FInputDevice::FEvent::GetMouseClick(EMouseButtons::Type& OutButton, uint16& OutPosX, uint16& OutPosY)
-{
-	check(Event == EventType::MOUSE_DOWN || Event == EventType::MOUSE_UP);
-	// https://developer.mozilla.org/en-US/docs/Web/Events/mousedown
-	uint8 Button = Data.MouseButton.Button;
-	switch (Button)
+	void FInputDevice::FEvent::GetMouseClick(EMouseButtons::Type& OutButton, uint16& OutPosX, uint16& OutPosY)
 	{
-		case 0:
-		{
-			OutButton = EMouseButtons::Left;
-		}
-		break;
-		case 1:
-		{
-			OutButton = EMouseButtons::Middle;
-		}
-		break;
-		case 2:
-		{
-			OutButton = EMouseButtons::Right;
-		}
-		break;
-		case 3:
-		{
-			OutButton = EMouseButtons::Thumb01;
-		}
-		break;
-		case 4:
-		{
-			OutButton = EMouseButtons::Thumb02;
-		}
-		break;
-		default:
+		check(Event == EventType::MOUSE_DOWN || Event == EventType::MOUSE_UP);
+		// https://developer.mozilla.org/en-US/docs/Web/Events/mousedown
+		uint8 Button = Data.MouseButton.Button;
+		switch (Button)
 		{
-			UE_LOG(LogPixelStreamingInputDevice, Error, TEXT("Unknown Pixel Streaming mouse click with button %d and word 0x%016llx"), Button, Data.Word);
+			case 0:
+			{
+				OutButton = EMouseButtons::Left;
+			}
+			break;
+			case 1:
+			{
+				OutButton = EMouseButtons::Middle;
+			}
+			break;
+			case 2:
+			{
+				OutButton = EMouseButtons::Right;
+			}
+			break;
+			case 3:
+			{
+				OutButton = EMouseButtons::Thumb01;
+			}
+			break;
+			case 4:
+			{
+				OutButton = EMouseButtons::Thumb02;
+			}
+			break;
+			default:
+			{
+				UE_LOG(LogPixelStreamingInputDevice, Error, TEXT("Unknown Pixel Streaming mouse click with button %d and word 0x%016llx"), Button, Data.Word);
+			}
+			break;
 		}
-		break;
+		OutPosX = Data.MouseButton.PosX;
+		OutPosY = Data.MouseButton.PosY;
 	}
-	OutPosX = Data.MouseButton.PosX;
-	OutPosY = Data.MouseButton.PosY;
-}
+} // namespace UE::PixelStreaming
\ No newline at end of file
