This commit is a large collection (from various authors) of Pixel Streaming 5.1 work

diff --git a/Engine/Plugins/Media/PixelStreaming/Source/PixelStreaming/Private/InputDevice.h b/Engine/Plugins/Media/PixelStreaming/Source/PixelStreaming/Private/InputDevice.h
--- a/Engine/Plugins/Media/PixelStreaming/Source/PixelStreaming/Private/InputDevice.h
+++ b/Engine/Plugins/Media/PixelStreaming/Source/PixelStreaming/Private/InputDevice.h
@@ -9,522 +9,519 @@
 
 class IPixelStreamingModule;
 
-namespace UE {
-	namespace PixelStreaming {
-
-		class FInputDevice : public IInputDevice
+namespace UE::PixelStreaming 
+{
+	class FInputDevice : public IInputDevice
+	{
+	public:
+		/**
+		 * The pixel streaming input device takes input events and passes them to
+		 * a message handler.
+		 * @param InMessageHandler - The message handler which events should be passed to.
+		 * @param InInputComponents - The preexisting list of pixel streaming input components.
+		 */
+		FInputDevice(const TSharedRef<FGenericApplicationMessageHandler>& InMessageHandler);
+
+		/** Tick the interface (e.g. check for new controllers) */
+		virtual void Tick(float DeltaTime) override;
+
+		/** Poll for controller state and send events if needed */
+		virtual void SendControllerEvents() override;
+
+		/** Set which MessageHandler will get the events from SendControllerEvents. */
+		virtual void SetMessageHandler(const TSharedRef<FGenericApplicationMessageHandler>& InMessageHandler) override;
+
+		/** Exec handler to allow console commands to be passed through for debugging */
+		virtual bool Exec(UWorld* InWorld, const TCHAR* Cmd, FOutputDevice& Ar) override;
+
+		/**
+		* IForceFeedbackSystem pass through functions
+		*/
+		virtual void SetChannelValue(int32 ControllerId, FForceFeedbackChannelType ChannelType, float Value) override;
+
+		virtual void SetChannelValues(int32 ControllerId, const FForceFeedbackValues& values) override;
+
+		/** The types of event which can be processed by the device. */
+		enum class EventType
 		{
-		public:
-			/**
-			 * The pixel streaming input device takes input events and passes them to
-			 * a message handler.
-			 * @param InMessageHandler - The message handler which events should be passed to.
-			 * @param InInputComponents - The preexisting list of pixel streaming input components.
-			 */
-			FInputDevice(const TSharedRef<FGenericApplicationMessageHandler>& InMessageHandler);
-
-			/** Tick the interface (e.g. check for new controllers) */
-			virtual void Tick(float DeltaTime) override;
-
-			/** Poll for controller state and send events if needed */
-			virtual void SendControllerEvents() override;
-
-			/** Set which MessageHandler will get the events from SendControllerEvents. */
-			virtual void SetMessageHandler(const TSharedRef<FGenericApplicationMessageHandler>& InMessageHandler) override;
-
-			/** Exec handler to allow console commands to be passed through for debugging */
-			virtual bool Exec(UWorld* InWorld, const TCHAR* Cmd, FOutputDevice& Ar) override;
-
-			/**
-			* IForceFeedbackSystem pass through functions
-			*/
-			virtual void SetChannelValue(int32 ControllerId, FForceFeedbackChannelType ChannelType, float Value) override;
+			UNDEFINED,		 /** No value. */
+			KEY_DOWN,		 /** A key has been pushed down. */
+			KEY_UP,			 /** A key has been released. */
+			KEY_PRESS,		 /** A key has been pressed and a character has been input. */
+			MOUSE_ENTER,	 /** The mouse has entered canvas. */
+			MOUSE_LEAVE,	 /** The mouse has left the canvas. */
+			MOUSE_MOVE,		 /** The mouse has been moved. */
+			MOUSE_DOWN,		 /** A mouse button has been clicked. */
+			MOUSE_UP,		 /** A mouse button has been released. */
+			MOUSE_WHEEL,	 /** The mouse wheel was scrolled. */
+			TOUCH_START,	 /** A finger is put down onto the canvas. */
+			TOUCH_END,		 /** A finger is lifted from the canvas. */
+			TOUCH_MOVE,		 /** A finger is being dragged along the surface of the canvas. */
+			GAMEPAD_PRESS,	 /** A gamepad button has been pressed. */
+			GAMEPAD_RELEASE, /** A gamepad button has been released. */
+			GAMEPAD_ANALOG	 /** A gamepad analog stick has been moved. */
+		};
 
-			virtual void SetChannelValues(int32 ControllerId, const FForceFeedbackValues& values) override;
+		/** A general input event. */
+		struct FEvent
+		{
+			/** The type of the general event. */
+			EventType Event;
 
-			/** The types of event which can be processed by the device. */
-			enum class EventType
-			{
-				UNDEFINED,		 /** No value. */
-				KEY_DOWN,		 /** A key has been pushed down. */
-				KEY_UP,			 /** A key has been released. */
-				KEY_PRESS,		 /** A key has been pressed and a character has been input. */
-				MOUSE_ENTER,	 /** The mouse has entered canvas. */
-				MOUSE_LEAVE,	 /** The mouse has left the canvas. */
-				MOUSE_MOVE,		 /** The mouse has been moved. */
-				MOUSE_DOWN,		 /** A mouse button has been clicked. */
-				MOUSE_UP,		 /** A mouse button has been released. */
-				MOUSE_WHEEL,	 /** The mouse wheel was scrolled. */
-				TOUCH_START,	 /** A finger is put down onto the canvas. */
-				TOUCH_END,		 /** A finger is lifted from the canvas. */
-				TOUCH_MOVE,		 /** A finger is being dragged along the surface of the canvas. */
-				GAMEPAD_PRESS,	 /** A gamepad button has been pressed. */
-				GAMEPAD_RELEASE, /** A gamepad button has been released. */
-				GAMEPAD_ANALOG	 /** A gamepad analog stick has been moved. */
-			};
-
-			/** A general input event. */
-			struct FEvent
+			/** A generic piece of data which is used to hold information about the
+			* event, specialized by making a union with an appropriate struct. */
+			union
 			{
-				/** The type of the general event. */
-				EventType Event;
+				uint64 Word;
 
-				/** A generic piece of data which is used to hold information about the
-				* event, specialized by making a union with an appropriate struct. */
-				union
-				{
-					uint64 Word;
-
-					struct /** KEY_DOWN */
-					{
-						uint8 KeyCode;
-						bool bIsRepeat;
-					} KeyDown;
-
-					struct /* KEY_UP */
-					{
-						uint8 KeyCode;
-					} KeyUp;
-
-					struct /** KEY_PRESSED */
-					{
-						TCHAR Character;
-					} Character;
-
-					struct /** MOUSE_MOVE */
-					{
-						int16 DeltaX;
-						int16 DeltaY;
-						uint16 PosX;
-						uint16 PosY;
-					} MouseMove;
-
-					struct /** MOUSE_DOWN, MOUSE_UP */
-					{
-						uint8 Button;
-						uint16 PosX;
-						uint16 PosY;
-					} MouseButton;
-
-					struct /** MOUSE_WHEEL */
-					{
-						int16 Delta;
-						uint16 PosX;
-						uint16 PosY;
-					} MouseWheel;
-
-					struct /** TOUCH_START, TOUCH_END, TOUCH_MOVE */
-					{
-						uint8 TouchIndex;
-						uint16 PosX;
-						uint16 PosY;
-						uint8 Force;
-					} Touch;
-
-					struct /** GAMEPAD_PRESSED, GAMEPAD_RELEASED */
-					{
-						uint8 ControllerIndex;
-						uint8 ButtonIndex;
-						bool bIsRepeat;
-
-					} GamepadButton;
-
-					struct /** GAMEPAD_PRESSED, GAMEPAD_RELEASED */
-					{
-						uint8 ControllerIndex;
-						float AnalogValue;
-						uint8 AxisIndex;
-					} GamepadAnalog;
-				} Data;
-
-				/**
-				* Create a completely empty event.
-				*/
-				FEvent()
-					: Event(EventType::UNDEFINED)
-				{
-				}
-
-				/**
-				* Create an event of the given type.
-				* @param InEvent - The type of the event.
-				*/
-				FEvent(EventType InEvent)
-					: Event(InEvent)
-				{
-				}
-
-				/**
-				* An event related to gamepad analog stick being moved.
-				* @param InControllerIndex - Number representing the index of the controller.
-				* @param InAxisIndex - The axis that is non-centered.
-				* @param InAnalogValue - The value of the analog stick.
-				*/
-				void SetGamepadAnalog(uint8 InControllerIndex, uint8 InAxisIndex, float InAnalogValue)
-				{
-					check(Event == EventType::GAMEPAD_ANALOG);
-					Data.GamepadAnalog.ControllerIndex = InControllerIndex;
-					Data.GamepadAnalog.AxisIndex = InAxisIndex;
-					Data.GamepadAnalog.AnalogValue = InAnalogValue;
-				}
-
-				/**
-				* An event related to a gamepad button being pushed down.
-				* @param InControllerIndex - Number representing the index of the controller.
-				* @param InButtonIndex - Numerical code identifying the pushed down button.
-				* @param InIsRepeat - Whether the button is being kept down and is repeating.
-				*/
-				void SetGamepadButtonPressed(uint8 InControllerIndex, uint8 InButtonIndex, bool InIsRepeat)
-				{
-					check(Event == EventType::GAMEPAD_PRESS);
-					Data.GamepadButton.ControllerIndex = InControllerIndex;
-					Data.GamepadButton.ButtonIndex = InButtonIndex;
-					Data.GamepadButton.bIsRepeat = InIsRepeat;
-				}
-
-				/**
-				* An event related to a gamepad button being released.
-				* @param InControllerIndex - Number representing the index of the controller.
-				* @param InButtonIndex - Numerical code identifying the released button.
-				*/
-				void SetGamepadButtonReleased(uint8 InControllerIndex, uint8 InButtonIndex)
-				{
-					check(Event == EventType::GAMEPAD_RELEASE);
-					Data.GamepadButton.ControllerIndex = InControllerIndex;
-					Data.GamepadButton.ButtonIndex = InButtonIndex;
-				}
-
-				/**
-				* An event related to a key being pushed down.
-				* @param InKeyCode - Numerical code identifying the pushed down key.
-				* @param InIsRepeat - Whether the key is being kept down and is repeating.
-				*/
-				void SetKeyDown(uint8 InKeyCode, bool InIsRepeat)
-				{
-					check(Event == EventType::KEY_DOWN);
-					Data.KeyDown.KeyCode = InKeyCode;
-					Data.KeyDown.bIsRepeat = InIsRepeat;
-				}
-
-				/**
-				* An event related to a key being released.
-				* @param InKeyCode - Numerical code identifying the released key.
-				*/
-				void SetKeyUp(uint8 InKeyCode)
-				{
-					check(Event == EventType::KEY_UP);
-					Data.KeyUp.KeyCode = InKeyCode;
-				}
-
-				/**
-				* An event related to character input.
-				* @param InCharacter - The character being input.
-				*/
-				void SetCharCode(TCHAR InCharacter)
-				{
-					check(Event == EventType::KEY_PRESS);
-					Data.Character.Character = InCharacter;
-				}
-
-				/**
-				* An event related to mouse movement.
-				* @param InPoxX - The X position of the mouse pointer.
-				* @param InPosY - The Y position of the mouse pointer.
-				* @param InDeltaX - The change in the X position of the mouse pointer.
-				* @param InDeltaY - The change in the Y position of the mouse pointer.
-				*/
-				void SetMouseDelta(uint16 InPosX, uint16 InPosY, int16 InDeltaX, int16 InDeltaY)
-				{
-					check(Event == EventType::MOUSE_MOVE);
-					Data.MouseMove.DeltaX = InDeltaX;
-					Data.MouseMove.DeltaY = InDeltaY;
-					Data.MouseMove.PosX = InPosX;
-					Data.MouseMove.PosY = InPosY;
-				}
-
-				/**
-				* An event related to mouse buttons.
-				* @param InButton - The button number corresponding to left, middle, right, etc.
-				* @param InPoxX - The X position of the mouse pointer.
-				* @param InPosY - The Y position of the mouse pointer.
-				*/
-				void SetMouseClick(uint8 InButton, uint16 InPosX, uint16 InPosY)
-				{
-					check(Event == EventType::MOUSE_DOWN || Event == EventType::MOUSE_UP);
-					Data.MouseButton.Button = InButton;
-					Data.MouseButton.PosX = InPosX;
-					Data.MouseButton.PosY = InPosY;
-				}
-
-				/**
-				* An event related to the mouse scroll wheel.
-				* @param InButton - The amount by which the mouse wheel was scrolled.
-				* @param InPoxX - The X position of the mouse pointer when the wheel was scrolled.
-				* @param InPosY - The Y position of the mouse pointer when the wheel was scrolled.
-				*/
-				void SetMouseWheel(int16 InDelta, uint16 InPosX, uint16 InPosY)
+				struct /** KEY_DOWN */
 				{
-					check(Event == EventType::MOUSE_WHEEL);
-					Data.MouseWheel.Delta = InDelta;
-					Data.MouseWheel.PosX = InPosX;
-					Data.MouseWheel.PosY = InPosY;
-				}
-
-				/**
-				* An event related to a finger touching the canvas.
-				* @param InTouchIndex - The finger used in multi-touch.
-				* @param InPoxX - The X position of the finger.
-				* @param InPosY - The Y position of the finger.
-				* @param InForce - The amount of pressure being applied by the finger.
-				*/
-				void SetTouch(uint8 InTouchIndex, uint16 InPosX, uint16 InPosY, uint8 InForce)
-				{
-					check(Event == EventType::TOUCH_START || Event == EventType::TOUCH_END || Event == EventType::TOUCH_MOVE);
-					Data.Touch.TouchIndex = InTouchIndex;
-					Data.Touch.PosX = InPosX;
-					Data.Touch.PosY = InPosY;
-					Data.Touch.Force = InForce;
-				}
-
-				/**
-				* Get information about an event related to gamepad analog stick being moved.
-				* @param OutControllerIndex - Number representing the index of the controller.
-				* @param OutAxisIndex - Number representing the axis of the controller.
-				* @param OutAnalogValue - The value of the analog stick.
-				*/
-				void GetGamepadAnalog(uint8& OutControllerIndex, uint8& OutAxisIndex, float& OutAnalogValue)
-				{
-					check(Event == EventType::GAMEPAD_ANALOG);
-					OutControllerIndex = Data.GamepadAnalog.ControllerIndex;
-					OutAxisIndex = Data.GamepadAnalog.AxisIndex;
-					OutAnalogValue = Data.GamepadAnalog.AnalogValue;
-				}
-
-				/**
-				* Get information about an event related to a gamepad button being pushed down.
-				* @param OutControllerIndex - Number representing the index of the controller.
-				* @param OutButtonIndex - Numerical code identifying the pushed down button.
-				* @param OutIsRepeat - Whether the button is being kept down and is repeating.
-				*/
-				void GetGamepadButtonPressed(uint8& OutControllerIndex, uint8& OutButtonIndex, bool& OutIsRepeat)
+					uint8 KeyCode;
+					bool bIsRepeat;
+				} KeyDown;
+
+				struct /* KEY_UP */
 				{
-					check(Event == EventType::GAMEPAD_PRESS);
-					OutControllerIndex = Data.GamepadButton.ControllerIndex;
-					OutButtonIndex = Data.GamepadButton.ButtonIndex;
-					OutIsRepeat = Data.GamepadButton.bIsRepeat;
-				}
-
-				/**
-				* Get information about an event related to a gamepad button being released.
-				* @param OutControllerIndex - Number representing the index of the controller.
-				* @param OutButtonIndex - Numerical code identifying the pushed down button.
-				*/
-				void GetGamepadButtonReleased(uint8& OutControllerIndex, uint8& OutButtonIndex)
+					uint8 KeyCode;
+				} KeyUp;
+
+				struct /** KEY_PRESSED */
 				{
-					check(Event == EventType::GAMEPAD_RELEASE);
-					OutControllerIndex = Data.GamepadButton.ControllerIndex;
-					OutButtonIndex = Data.GamepadButton.ButtonIndex;
-				}
-
-				/**
-				* Get information about an event related to a key being pushed down.
-				* @param OutKeyCode - Numerical code identifying the pushed down key.
-				* @param OutIsRepeat - Whether the key is being kept down and is repeating.
-				*/
-				void GetKeyDown(uint8& OutKeyCode, bool& OutIsRepeat)
+					TCHAR Character;
+				} Character;
+
+				struct /** MOUSE_MOVE */
 				{
-					check(Event == EventType::KEY_DOWN);
-					OutKeyCode = Data.KeyDown.KeyCode;
-					OutIsRepeat = Data.KeyDown.bIsRepeat;
-				}
-
-				/**
-				* Get information about an event related to a key being released.
-				* @param OutKeyCode - Numerical code identifying the released key.
-				*/
-				void GetKeyUp(uint8& OutKeyCode)
+					int16 DeltaX;
+					int16 DeltaY;
+					uint16 PosX;
+					uint16 PosY;
+				} MouseMove;
+
+				struct /** MOUSE_DOWN, MOUSE_UP */
 				{
-					check(Event == EventType::KEY_UP);
-					OutKeyCode = Data.KeyUp.KeyCode;
-				}
-
-				/**
-				* Get information about an event related to character input.
-				* @param OutCharacter - The character being input.
-				*/
-				void GetCharacterCode(TCHAR& OutCharacter)
+					uint8 Button;
+					uint16 PosX;
+					uint16 PosY;
+				} MouseButton;
+
+				struct /** MOUSE_WHEEL */
 				{
-					check(Event == EventType::KEY_PRESS);
-					OutCharacter = Data.Character.Character;
-				}
-
-				/**
-				* Get information about an event related to mouse movement.
-				* @param OutPoxX - The X position of the mouse pointer.
-				* @param OutPosY - The Y position of the mouse pointer.
-				* @param OutDeltaX - The change in the X position of the mouse pointer.
-				* @param OutDeltaY - The change in the Y position of the mouse pointer.
-				*/
-				void GetMouseDelta(uint16& OutPosX, uint16& OutPosY, int16& OutDeltaX, int16& OutDeltaY)
+					int16 Delta;
+					uint16 PosX;
+					uint16 PosY;
+				} MouseWheel;
+
+				struct /** TOUCH_START, TOUCH_END, TOUCH_MOVE */
 				{
-					check(Event == EventType::MOUSE_MOVE);
-					OutPosX = Data.MouseMove.PosX;
-					OutPosY = Data.MouseMove.PosY;
-					OutDeltaX = Data.MouseMove.DeltaX;
-					OutDeltaY = Data.MouseMove.DeltaY;
-				}
-
-				/**
-				* Get information about an event related to mouse buttons.
-				* @param OutButton - The button number corresponding to left, middle, right, etc.
-				* @param OutPosX - The X position of the mouse pointer.
-				* @param OutPosY - The Y position of the mouse pointer.
-				*/
-				void GetMouseClick(EMouseButtons::Type& OutButton, uint16& OutPosX, uint16& OutPosY);
-
-				/**
-				* Get information about an event related to the mouse wheel.
-				* @param OutDelta - The amount by which the mouse wheel was scrolled.
-				* @param PosX - The X position of the mouse pointer when the wheel was scrolled.
-				* @param PosY - The Y position of the mouse pointer when the wheel was scrolled.
-				*/
-				void GetMouseWheel(int16& OutDelta, uint16& OutPosX, uint16& OutPosY)
+					uint8 TouchIndex;
+					uint16 PosX;
+					uint16 PosY;
+					uint8 Force;
+				} Touch;
+
+				struct /** GAMEPAD_PRESSED, GAMEPAD_RELEASED */
 				{
-					check(Event == EventType::MOUSE_WHEEL);
-					OutDelta = Data.MouseWheel.Delta;
-					OutPosX = Data.MouseWheel.PosX;
-					OutPosY = Data.MouseWheel.PosY;
-				}
-
-				/**
-				* Get information about an event related to a finger touching the canvas.
-				* @param OutTouchIndex - The finger used in multi-touch.
-				* @param OutPoxX - The X position of the finger.
-				* @param OutPosY - The Y position of the finger.
-				* @param OutForce - Amount of pressure being applied by the finger.
-				*/
-				void GetTouch(uint8& OutTouchIndex, uint16& OutPosX, uint16& OutPosY, uint8& OutForce)
+					uint8 ControllerIndex;
+					uint8 ButtonIndex;
+					bool bIsRepeat;
+				} GamepadButton;
+
+				struct /** GAMEPAD_PRESSED, GAMEPAD_RELEASED */
 				{
-					check(Event == EventType::TOUCH_START || Event == EventType::TOUCH_END || Event == EventType::TOUCH_MOVE);
-					OutTouchIndex = Data.Touch.TouchIndex;
-					OutPosX = Data.Touch.PosX;
-					OutPosY = Data.Touch.PosY;
-					OutForce = Data.Touch.Force;
-				}
-			};
+					uint8 ControllerIndex;
+					float AnalogValue;
+					uint8 AxisIndex;
+				} GamepadAnalog;
+			} Data;
+
+			/**
+			* Create a completely empty event.
+			*/
+			FEvent()
+				: Event(EventType::UNDEFINED)
+			{
+			}
+
+			/**
+			* Create an event of the given type.
+			* @param InEvent - The type of the event.
+			*/
+			FEvent(EventType InEvent)
+				: Event(InEvent)
+			{
+			}
+
+			/**
+			* An event related to gamepad analog stick being moved.
+			* @param InControllerIndex - Number representing the index of the controller.
+			* @param InAxisIndex - The axis that is non-centered.
+			* @param InAnalogValue - The value of the analog stick.
+			*/
+			void SetGamepadAnalog(uint8 InControllerIndex, uint8 InAxisIndex, float InAnalogValue)
+			{
+				check(Event == EventType::GAMEPAD_ANALOG);
+				Data.GamepadAnalog.ControllerIndex = InControllerIndex;
+				Data.GamepadAnalog.AxisIndex = InAxisIndex;
+				Data.GamepadAnalog.AnalogValue = InAnalogValue;
+			}
 
 			/**
-			* Add a new event to the input device for later processing.
-			* @param InEvent - The new event.
+			* An event related to a gamepad button being pushed down.
+			* @param InControllerIndex - Number representing the index of the controller.
+			* @param InButtonIndex - Numerical code identifying the pushed down button.
+			* @param InIsRepeat - Whether the button is being kept down and is repeating.
 			*/
-			void ProcessEvent(const FEvent& InEvent);
+			void SetGamepadButtonPressed(uint8 InControllerIndex, uint8 InButtonIndex, bool InIsRepeat)
+			{
+				check(Event == EventType::GAMEPAD_PRESS);
+				Data.GamepadButton.ControllerIndex = InControllerIndex;
+				Data.GamepadButton.ButtonIndex = InButtonIndex;
+				Data.GamepadButton.bIsRepeat = InIsRepeat;
+			}
 
 			/**
-			* Add a new UI interaction descriptor to the input device for later processing.
-			* @param InDescriptor - The new UI interaction descriptor.
+			* An event related to a gamepad button being released.
+			* @param InControllerIndex - Number representing the index of the controller.
+			* @param InButtonIndex - Numerical code identifying the released button.
 			*/
-			void ProcessUIInteraction(const FString& InDescriptor);
+			void SetGamepadButtonReleased(uint8 InControllerIndex, uint8 InButtonIndex)
+			{
+				check(Event == EventType::GAMEPAD_RELEASE);
+				Data.GamepadButton.ControllerIndex = InControllerIndex;
+				Data.GamepadButton.ButtonIndex = InButtonIndex;
+			}
 
 			/**
-			* Add a new command descriptor to the input device for later execution.
-			* @param InDescriptor - The new command descriptor.
+			* An event related to a key being pushed down.
+			* @param InKeyCode - Numerical code identifying the pushed down key.
+			* @param InIsRepeat - Whether the key is being kept down and is repeating.
 			*/
-			void ProcessCommand(const FString& InDescriptor);
+			void SetKeyDown(uint8 InKeyCode, bool InIsRepeat)
+			{
+				check(Event == EventType::KEY_DOWN);
+				Data.KeyDown.KeyCode = InKeyCode;
+				Data.KeyDown.bIsRepeat = InIsRepeat;
+			}
 
 			/**
-			* Is the application faking touch events?
-			* @return True if the application is faking touch events.
+			* An event related to a key being released.
+			* @param InKeyCode - Numerical code identifying the released key.
 			*/
-			bool IsFakingTouchEvents() const
+			void SetKeyUp(uint8 InKeyCode)
 			{
-				return bFakingTouchEvents;
+				check(Event == EventType::KEY_UP);
+				Data.KeyUp.KeyCode = InKeyCode;
 			}
 
 			/**
-			* Handle incoming messages
-			* @param Buffer - Incoming data buffer.
+			* An event related to character input.
+			* @param InCharacter - The character being input.
 			*/
-			void OnMessage(const webrtc::DataBuffer& Buffer);
+			void SetCharCode(TCHAR InCharacter)
+			{
+				check(Event == EventType::KEY_PRESS);
+				Data.Character.Character = InCharacter;
+			}
 
 			/**
-			* How long is the message type specifier and the length specifier in bytes?
-			* @return The number of header bytes in an incoming message what is ignored when it is parsed as string.
+			* An event related to mouse movement.
+			* @param InPoxX - The X position of the mouse pointer.
+			* @param InPosY - The Y position of the mouse pointer.
+			* @param InDeltaX - The change in the X position of the mouse pointer.
+			* @param InDeltaY - The change in the Y position of the mouse pointer.
 			*/
-			static const size_t GetMessageHeaderOffset()
+			void SetMouseDelta(uint16 InPosX, uint16 InPosY, int16 InDeltaX, int16 InDeltaY)
 			{
-				return MessageHeaderOffset;
+				check(Event == EventType::MOUSE_MOVE);
+				Data.MouseMove.DeltaX = InDeltaX;
+				Data.MouseMove.DeltaY = InDeltaY;
+				Data.MouseMove.PosX = InPosX;
+				Data.MouseMove.PosY = InPosY;
 			}
 
-		private:
 			/**
-			* Touch only. If the user touches the UI and focuses an editable widget
-			* then we need to tell the browser to allow it to display an on-screen
-			* keyboard.
+			* An event related to mouse buttons.
+			* @param InButton - The button number corresponding to left, middle, right, etc.
+			* @param InPoxX - The X position of the mouse pointer.
+			* @param InPosY - The Y position of the mouse pointer.
 			*/
-			void FindFocusedWidget();
+			void SetMouseClick(uint8 InButton, uint16 InPosX, uint16 InPosY)
+			{
+				check(Event == EventType::MOUSE_DOWN || Event == EventType::MOUSE_UP);
+				Data.MouseButton.Button = InButton;
+				Data.MouseButton.PosX = InPosX;
+				Data.MouseButton.PosY = InPosY;
+			}
 
 			/**
-			* Convert the controller analog axis index to its corresponding FType
+			* An event related to the mouse scroll wheel.
+			* @param InButton - The amount by which the mouse wheel was scrolled.
+			* @param InPoxX - The X position of the mouse pointer when the wheel was scrolled.
+			* @param InPosY - The Y position of the mouse pointer when the wheel was scrolled.
 			*/
-			FGamepadKeyNames::Type ConvertAxisIndexToGamepadAxis(uint8 AnalogAxis);
+			void SetMouseWheel(int16 InDelta, uint16 InPosX, uint16 InPosY)
+			{
+				check(Event == EventType::MOUSE_WHEEL);
+				Data.MouseWheel.Delta = InDelta;
+				Data.MouseWheel.PosX = InPosX;
+				Data.MouseWheel.PosY = InPosY;
+			}
 
 			/**
-			* Convert the controller button index to its corresponding FType
+			* An event related to a finger touching the canvas.
+			* @param InTouchIndex - The finger used in multi-touch.
+			* @param InPoxX - The X position of the finger.
+			* @param InPosY - The Y position of the finger.
+			* @param InForce - The amount of pressure being applied by the finger.
 			*/
-			FGamepadKeyNames::Type ConvertButtonIndexToGamepadButton(uint8 ButtonIndex);
+			void SetTouch(uint8 InTouchIndex, uint16 InPosX, uint16 InPosY, uint8 InForce)
+			{
+				check(Event == EventType::TOUCH_START || Event == EventType::TOUCH_END || Event == EventType::TOUCH_MOVE);
+				Data.Touch.TouchIndex = InTouchIndex;
+				Data.Touch.PosX = InPosX;
+				Data.Touch.PosY = InPosY;
+				Data.Touch.Force = InForce;
+			}
 
 			/**
-			* A special wrapper over the GenericApplication layer which allows us to
-			* override certain behavior.
+			* Get information about an event related to gamepad analog stick being moved.
+			* @param OutControllerIndex - Number representing the index of the controller.
+			* @param OutAxisIndex - Number representing the axis of the controller.
+			* @param OutAnalogValue - The value of the analog stick.
 			*/
-			TSharedPtr<class FApplicationWrapper> PixelStreamerApplicationWrapper;
+			void GetGamepadAnalog(uint8& OutControllerIndex, uint8& OutAxisIndex, float& OutAnalogValue)
+			{
+				check(Event == EventType::GAMEPAD_ANALOG);
+				OutControllerIndex = Data.GamepadAnalog.ControllerIndex;
+				OutAxisIndex = Data.GamepadAnalog.AxisIndex;
+				OutAnalogValue = Data.GamepadAnalog.AnalogValue;
+			}
 
-			/** Reference to the message handler which events should be passed to. */
-			TSharedRef<FGenericApplicationMessageHandler> MessageHandler;
+			/**
+			* Get information about an event related to a gamepad button being pushed down.
+			* @param OutControllerIndex - Number representing the index of the controller.
+			* @param OutButtonIndex - Numerical code identifying the pushed down button.
+			* @param OutIsRepeat - Whether the button is being kept down and is repeating.
+			*/
+			void GetGamepadButtonPressed(uint8& OutControllerIndex, uint8& OutButtonIndex, bool& OutIsRepeat)
+			{
+				check(Event == EventType::GAMEPAD_PRESS);
+				OutControllerIndex = Data.GamepadButton.ControllerIndex;
+				OutButtonIndex = Data.GamepadButton.ButtonIndex;
+				OutIsRepeat = Data.GamepadButton.bIsRepeat;
+			}
 
-			/** The queue of events which are awaiting processing. */
-			TQueue<FEvent> Events;
+			/**
+			* Get information about an event related to a gamepad button being released.
+			* @param OutControllerIndex - Number representing the index of the controller.
+			* @param OutButtonIndex - Numerical code identifying the pushed down button.
+			*/
+			void GetGamepadButtonReleased(uint8& OutControllerIndex, uint8& OutButtonIndex)
+			{
+				check(Event == EventType::GAMEPAD_RELEASE);
+				OutControllerIndex = Data.GamepadButton.ControllerIndex;
+				OutButtonIndex = Data.GamepadButton.ButtonIndex;
+			}
 
 			/**
-			* A queue of UI interaction descriptor strings which contain arbitrary
-			* information related to the interaction.
+			* Get information about an event related to a key being pushed down.
+			* @param OutKeyCode - Numerical code identifying the pushed down key.
+			* @param OutIsRepeat - Whether the key is being kept down and is repeating.
 			*/
-			TQueue<FString> UIInteractions;
+			void GetKeyDown(uint8& OutKeyCode, bool& OutIsRepeat)
+			{
+				check(Event == EventType::KEY_DOWN);
+				OutKeyCode = Data.KeyDown.KeyCode;
+				OutIsRepeat = Data.KeyDown.bIsRepeat;
+			}
 
 			/**
-			* A queue of command descriptor strings which contain the command to
-			* execute and its arguments.
+			* Get information about an event related to a key being released.
+			* @param OutKeyCode - Numerical code identifying the released key.
 			*/
-			TQueue<FString> Commands;
+			void GetKeyUp(uint8& OutKeyCode)
+			{
+				check(Event == EventType::KEY_UP);
+				OutKeyCode = Data.KeyUp.KeyCode;
+			}
 
 			/**
-			* Is the application faking touch events by dragging the mouse along
-			* the canvas? If so then we must put the browser canvas in a special
-			* state to replicate the behavior of the application.
+			* Get information about an event related to character input.
+			* @param OutCharacter - The character being input.
 			*/
-			bool bFakingTouchEvents;
+			void GetCharacterCode(TCHAR& OutCharacter)
+			{
+				check(Event == EventType::KEY_PRESS);
+				OutCharacter = Data.Character.Character;
+			}
 
 			/**
-			* Touch only. Location of the focused UI widget. If no UI widget is focused
-			* then this has the UnfocusedPos value.
+			* Get information about an event related to mouse movement.
+			* @param OutPoxX - The X position of the mouse pointer.
+			* @param OutPosY - The Y position of the mouse pointer.
+			* @param OutDeltaX - The change in the X position of the mouse pointer.
+			* @param OutDeltaY - The change in the Y position of the mouse pointer.
 			*/
-			FVector2D FocusedPos;
+			void GetMouseDelta(uint16& OutPosX, uint16& OutPosY, int16& OutDeltaX, int16& OutDeltaY)
+			{
+				check(Event == EventType::MOUSE_MOVE);
+				OutPosX = Data.MouseMove.PosX;
+				OutPosY = Data.MouseMove.PosY;
+				OutDeltaX = Data.MouseMove.DeltaX;
+				OutDeltaY = Data.MouseMove.DeltaY;
+			}
 
 			/**
-			* Touch only. A special position which indicates that no UI widget is
-			* focused.
+			* Get information about an event related to mouse buttons.
+			* @param OutButton - The button number corresponding to left, middle, right, etc.
+			* @param OutPosX - The X position of the mouse pointer.
+			* @param OutPosY - The Y position of the mouse pointer.
 			*/
-			static const FVector2D UnfocusedPos;
+			void GetMouseClick(EMouseButtons::Type& OutButton, uint16& OutPosX, uint16& OutPosY);
 
-			/** For convenience we keep a reference to the Pixel Streaming plugin. */
-			IPixelStreamingModule* PixelStreamingModule;
+			/**
+			* Get information about an event related to the mouse wheel.
+			* @param OutDelta - The amount by which the mouse wheel was scrolled.
+			* @param PosX - The X position of the mouse pointer when the wheel was scrolled.
+			* @param PosY - The Y position of the mouse pointer when the wheel was scrolled.
+			*/
+			void GetMouseWheel(int16& OutDelta, uint16& OutPosX, uint16& OutPosY)
+			{
+				check(Event == EventType::MOUSE_WHEEL);
+				OutDelta = Data.MouseWheel.Delta;
+				OutPosX = Data.MouseWheel.PosX;
+				OutPosY = Data.MouseWheel.PosY;
+			}
 
-			/*
-			* Padding for string parsing when handling messages.
-			* 1 character for the actual message and then
-			* 2 characters for the length which are skipped
+			/**
+			* Get information about an event related to a finger touching the canvas.
+			* @param OutTouchIndex - The finger used in multi-touch.
+			* @param OutPoxX - The X position of the finger.
+			* @param OutPosY - The Y position of the finger.
+			* @param OutForce - Amount of pressure being applied by the finger.
 			*/
-			static const size_t MessageHeaderOffset;
+			void GetTouch(uint8& OutTouchIndex, uint16& OutPosX, uint16& OutPosY, uint8& OutForce)
+			{
+				check(Event == EventType::TOUCH_START || Event == EventType::TOUCH_END || Event == EventType::TOUCH_MOVE);
+				OutTouchIndex = Data.Touch.TouchIndex;
+				OutPosX = Data.Touch.PosX;
+				OutPosY = Data.Touch.PosY;
+				OutForce = Data.Touch.Force;
+			}
 		};
-	}
-}
+
+		/**
+		* Add a new event to the input device for later processing.
+		* @param InEvent - The new event.
+		*/
+		void ProcessEvent(const FEvent& InEvent);
+
+		/**
+		* Add a new UI interaction descriptor to the input device for later processing.
+		* @param InDescriptor - The new UI interaction descriptor.
+		*/
+		void ProcessUIInteraction(const FString& InDescriptor);
+
+		/**
+		* Add a new command descriptor to the input device for later execution.
+		* @param InDescriptor - The new command descriptor.
+		*/
+		void ProcessCommand(const FString& InDescriptor);
+
+		/**
+		* Is the application faking touch events?
+		* @return True if the application is faking touch events.
+		*/
+		bool IsFakingTouchEvents() const
+		{
+			return bFakingTouchEvents;
+		}
+
+		/**
+		* Handle incoming messages
+		* @param Buffer - Incoming data buffer.
+		*/
+		void OnMessage(const webrtc::DataBuffer& Buffer);
+
+		/**
+		* How long is the message type specifier and the length specifier in bytes?
+		* @return The number of header bytes in an incoming message what is ignored when it is parsed as string.
+		*/
+		static const size_t GetMessageHeaderOffset()
+		{
+			return MessageHeaderOffset;
+		}
+
+	private:
+		/**
+		* Touch only. If the user touches the UI and focuses an editable widget
+		* then we need to tell the browser to allow it to display an on-screen
+		* keyboard.
+		*/
+		void FindFocusedWidget();
+
+		/**
+		* Convert the controller analog axis index to its corresponding FType
+		*/
+		FGamepadKeyNames::Type ConvertAxisIndexToGamepadAxis(uint8 AnalogAxis);
+
+		/**
+		* Convert the controller button index to its corresponding FType
+		*/
+		FGamepadKeyNames::Type ConvertButtonIndexToGamepadButton(uint8 ButtonIndex);
+
+		/**
+		* A special wrapper over the GenericApplication layer which allows us to
+		* override certain behavior.
+		*/
+		TSharedPtr<class FApplicationWrapper> PixelStreamerApplicationWrapper;
+
+		/** Reference to the message handler which events should be passed to. */
+		TSharedRef<FGenericApplicationMessageHandler> MessageHandler;
+
+		/** The queue of events which are awaiting processing. */
+		TQueue<FEvent> Events;
+
+		/**
+		* A queue of UI interaction descriptor strings which contain arbitrary
+		* information related to the interaction.
+		*/
+		TQueue<FString> UIInteractions;
+
+		/**
+		* A queue of command descriptor strings which contain the command to
+		* execute and its arguments.
+		*/
+		TQueue<FString> Commands;
+
+		/**
+		* Is the application faking touch events by dragging the mouse along
+		* the canvas? If so then we must put the browser canvas in a special
+		* state to replicate the behavior of the application.
+		*/
+		bool bFakingTouchEvents;
+
+		/**
+		* Touch only. Location of the focused UI widget. If no UI widget is focused
+		* then this has the UnfocusedPos value.
+		*/
+		FVector2D FocusedPos;
+
+		/**
+		* Touch only. A special position which indicates that no UI widget is
+		* focused.
+		*/
+		static const FVector2D UnfocusedPos;
+
+		/** For convenience we keep a reference to the Pixel Streaming plugin. */
+		IPixelStreamingModule* PixelStreamingModule;
+
+		/*
+		* Padding for string parsing when handling messages.
+		* 1 character for the actual message and then
+		* 2 characters for the length which are skipped
+		*/
+		static const size_t MessageHeaderOffset;
+	};
+} // namespace UE::PixelStreaming
\ No newline at end of file
diff --git a/Engine/Plugins/Media/PixelStreaming/Source/PixelStreaming/Private/JavaScriptKeyCodes.inl b/Engine/Plugins/Media/PixelStreaming/Source/PixelStreaming/Private/JavaScriptKeyCodes.inl
--- a/Engine/Plugins/Media/PixelStreaming/Source/PixelStreaming/Private/JavaScriptKeyCodes.inl
+++ b/Engine/Plugins/Media/PixelStreaming/Source/PixelStreaming/Private/JavaScriptKeyCodes.inl
@@ -4,8 +4,9 @@
 
 #include "InputCoreTypes.h"
 
-namespace UE {
-	namespace PixelStreaming {
+namespace UE::PixelStreaming 
+{
+
 /**
  * An array which will convert JavaScript key codes into FKeys. The JavaScript
  * key code should be used as an index into the array to perform the lookup.
@@ -15,267 +16,265 @@ namespace UE {
  * have been choosen specially for 'right' versions of the same key.
  */
 
-		static const FKey* JavaScriptKeyCodeToFKey[256] =
-		{
-			/*   0 */ &EKeys::Invalid,
-			/*   1 */ &EKeys::Invalid,
-			/*   2 */ &EKeys::Invalid,
-			/*   3 */ &EKeys::Invalid,
-			/*   4 */ &EKeys::Invalid,
-			/*   5 */ &EKeys::Invalid,
-			/*   6 */ &EKeys::Invalid,
-			/*   7 */ &EKeys::Invalid,
-			/*   8 */ &EKeys::BackSpace,
-			/*   9 */ &EKeys::Tab,
-			/*  10 */ &EKeys::Invalid,
-			/*  11 */ &EKeys::Invalid,
-			/*  12 */ /*&EKeys::Clear*/ &EKeys::Invalid,
-			/*  13 */ &EKeys::Enter,
-			/*  14 */ &EKeys::Invalid,
-			/*  15 */ &EKeys::Invalid,
-			/*  16 */ &EKeys::LeftShift,
-			/*  17 */ &EKeys::LeftControl,
-			/*  18 */ &EKeys::LeftAlt,
-			/*  19 */ &EKeys::Pause,
-			/*  20 */ &EKeys::CapsLock,
-			/*  21 */ &EKeys::Invalid,
-			/*  22 */ &EKeys::Invalid,
-			/*  23 */ &EKeys::Invalid,
-			/*  24 */ &EKeys::Invalid,
-			/*  25 */ &EKeys::Invalid,
-			/*  26 */ &EKeys::Invalid,
-			/*  27 */ &EKeys::Escape,
-			/*  28 */ &EKeys::Invalid,
-			/*  29 */ &EKeys::Invalid,
-			/*  30 */ &EKeys::Invalid,
-			/*  31 */ &EKeys::Invalid,
-			/*  32 */ &EKeys::SpaceBar,
-			/*  33 */ &EKeys::PageUp,
-			/*  34 */ &EKeys::PageDown,
-			/*  35 */ &EKeys::End,
-			/*  36 */ &EKeys::Home,
-			/*  37 */ &EKeys::Left,
-			/*  38 */ &EKeys::Up,
-			/*  39 */ &EKeys::Right,
-			/*  40 */ &EKeys::Down,
-			/*  41 */ &EKeys::Invalid,
-			/*  42 */ &EKeys::Invalid,
-			/*  43 */ &EKeys::Invalid,
-			/*  44 */ /*&EKeys::PrintScreen*/ &EKeys::Invalid,
-			/*  45 */ &EKeys::Insert,
-			/*  46 */ &EKeys::Delete,
-			/*  47 */ &EKeys::Invalid,
-			/*  48 */ &EKeys::Zero,
-			/*  49 */ &EKeys::One,
-			/*  50 */ &EKeys::Two,
-			/*  51 */ &EKeys::Three,
-			/*  52 */ &EKeys::Four,
-			/*  53 */ &EKeys::Five,
-			/*  54 */ &EKeys::Six,
-			/*  55 */ &EKeys::Seven,
-			/*  56 */ &EKeys::Eight,
-			/*  57 */ &EKeys::Nine,
-			/*  58 */ &EKeys::Invalid,
-			/*  59 */ &EKeys::Invalid,
-			/*  60 */ &EKeys::Invalid,
-			/*  61 */ &EKeys::Invalid,
-			/*  62 */ &EKeys::Invalid,
-			/*  63 */ &EKeys::Invalid,
-			/*  64 */ &EKeys::Invalid,
-			/*  65 */ &EKeys::A,
-			/*  66 */ &EKeys::B,
-			/*  67 */ &EKeys::C,
-			/*  68 */ &EKeys::D,
-			/*  69 */ &EKeys::E,
-			/*  70 */ &EKeys::F,
-			/*  71 */ &EKeys::G,
-			/*  72 */ &EKeys::H,
-			/*  73 */ &EKeys::I,
-			/*  74 */ &EKeys::J,
-			/*  75 */ &EKeys::K,
-			/*  76 */ &EKeys::L,
-			/*  77 */ &EKeys::M,
-			/*  78 */ &EKeys::N,
-			/*  79 */ &EKeys::O,
-			/*  80 */ &EKeys::P,
-			/*  81 */ &EKeys::Q,
-			/*  82 */ &EKeys::R,
-			/*  83 */ &EKeys::S,
-			/*  84 */ &EKeys::T,
-			/*  85 */ &EKeys::U,
-			/*  86 */ &EKeys::V,
-			/*  87 */ &EKeys::W,
-			/*  88 */ &EKeys::X,
-			/*  89 */ &EKeys::Y,
-			/*  90 */ &EKeys::Z,
-			/*  91 */ /*&EKeys::LeftWindowKey*/ &EKeys::Invalid,
-			/*  92 */ /*&EKeys::RightWindowKey*/ &EKeys::Invalid,
-			/*  93 */ /*&EKeys::SelectKey*/ &EKeys::Invalid,
-			/*  94 */ &EKeys::Invalid,
-			/*  95 */ &EKeys::Invalid,
-			/*  96 */ &EKeys::NumPadZero,
-			/*  97 */ &EKeys::NumPadOne,
-			/*  98 */ &EKeys::NumPadTwo,
-			/*  99 */ &EKeys::NumPadThree,
-			/* 100 */ &EKeys::NumPadFour,
-			/* 101 */ &EKeys::NumPadFive,
-			/* 102 */ &EKeys::NumPadSix,
-			/* 103 */ &EKeys::NumPadSeven,
-			/* 104 */ &EKeys::NumPadEight,
-			/* 105 */ &EKeys::NumPadNine,
-			/* 106 */ &EKeys::Multiply,
-			/* 107 */ &EKeys::Add,
-			/* 108 */ &EKeys::Invalid,
-			/* 109 */ &EKeys::Subtract,
-			/* 110 */ &EKeys::Decimal,
-			/* 111 */ &EKeys::Divide,
-			/* 112 */ &EKeys::F1,
-			/* 113 */ &EKeys::F2,
-			/* 114 */ &EKeys::F3,
-			/* 115 */ &EKeys::F4,
-			/* 116 */ &EKeys::F5,
-			/* 117 */ &EKeys::F6,
-			/* 118 */ &EKeys::F7,
-			/* 119 */ &EKeys::F8,
-			/* 120 */ &EKeys::F9,
-			/* 121 */ &EKeys::F10,
-			/* 122 */ &EKeys::F11,
-			/* 123 */ &EKeys::F12,
-			/* 124 */ &EKeys::Invalid,
-			/* 125 */ &EKeys::Invalid,
-			/* 126 */ &EKeys::Invalid,
-			/* 127 */ &EKeys::Invalid,
-			/* 128 */ &EKeys::Invalid,
-			/* 129 */ &EKeys::Invalid,
-			/* 130 */ &EKeys::Invalid,
-			/* 131 */ &EKeys::Invalid,
-			/* 132 */ &EKeys::Invalid,
-			/* 133 */ &EKeys::Invalid,
-			/* 134 */ &EKeys::Invalid,
-			/* 135 */ &EKeys::Invalid,
-			/* 136 */ &EKeys::Invalid,
-			/* 137 */ &EKeys::Invalid,
-			/* 138 */ &EKeys::Invalid,
-			/* 139 */ &EKeys::Invalid,
-			/* 140 */ &EKeys::Invalid,
-			/* 141 */ &EKeys::Invalid,
-			/* 142 */ &EKeys::Invalid,
-			/* 143 */ &EKeys::Invalid,
-			/* 144 */ &EKeys::NumLock,
-			/* 145 */ &EKeys::ScrollLock,
-			/* 146 */ &EKeys::Invalid,
-			/* 147 */ &EKeys::Invalid,
-			/* 148 */ &EKeys::Invalid,
-			/* 149 */ &EKeys::Invalid,
-			/* 150 */ &EKeys::Invalid,
-			/* 151 */ &EKeys::Invalid,
-			/* 152 */ &EKeys::Invalid,
-			/* 153 */ &EKeys::Invalid,
-			/* 154 */ &EKeys::Invalid,
-			/* 155 */ &EKeys::Invalid,
-			/* 156 */ &EKeys::Invalid,
-			/* 157 */ &EKeys::Invalid,
-			/* 158 */ &EKeys::Invalid,
-			/* 159 */ &EKeys::Invalid,
-			/* 160 */ &EKeys::Invalid,
-			/* 161 */ &EKeys::Invalid,
-			/* 162 */ &EKeys::Invalid,
-			/* 163 */ &EKeys::Invalid,
-			/* 164 */ &EKeys::Invalid,
-			/* 165 */ &EKeys::Invalid,
-			/* 166 */ &EKeys::Invalid,
-			/* 167 */ &EKeys::Invalid,
-			/* 168 */ &EKeys::Invalid,
-			/* 169 */ &EKeys::Invalid,
-			/* 170 */ &EKeys::Invalid,
-			/* 171 */ &EKeys::Invalid,
-			/* 172 */ &EKeys::Invalid,
-			/* 173 */ &EKeys::Invalid,
-			/* 174 */ &EKeys::Invalid,
-			/* 175 */ &EKeys::Invalid,
-			/* 176 */ &EKeys::Invalid,
-			/* 177 */ &EKeys::Invalid,
-			/* 178 */ &EKeys::Invalid,
-			/* 179 */ &EKeys::Invalid,
-			/* 180 */ &EKeys::Invalid,
-			/* 181 */ &EKeys::Invalid,
-			/* 182 */ &EKeys::Invalid,
-			/* 183 */ &EKeys::Invalid,
-			/* 184 */ &EKeys::Invalid,
-			/* 185 */ &EKeys::Invalid,
-			/* 186 */ &EKeys::Semicolon,
-			/* 187 */ &EKeys::Equals,
-			/* 188 */ &EKeys::Comma,
-			/* 189 */ &EKeys::Hyphen,
-			/* 190 */ &EKeys::Period,
-			/* 191 */ &EKeys::Slash,
-			/* 192 */ &EKeys::Tilde,
-			/* 193 */ &EKeys::Invalid,
-			/* 194 */ &EKeys::Invalid,
-			/* 195 */ &EKeys::Invalid,
-			/* 196 */ &EKeys::Invalid,
-			/* 197 */ &EKeys::Invalid,
-			/* 198 */ &EKeys::Invalid,
-			/* 199 */ &EKeys::Invalid,
-			/* 200 */ &EKeys::Invalid,
-			/* 201 */ &EKeys::Invalid,
-			/* 202 */ &EKeys::Invalid,
-			/* 203 */ &EKeys::Invalid,
-			/* 204 */ &EKeys::Invalid,
-			/* 205 */ &EKeys::Invalid,
-			/* 206 */ &EKeys::Invalid,
-			/* 207 */ &EKeys::Invalid,
-			/* 208 */ &EKeys::Invalid,
-			/* 209 */ &EKeys::Invalid,
-			/* 210 */ &EKeys::Invalid,
-			/* 211 */ &EKeys::Invalid,
-			/* 212 */ &EKeys::Invalid,
-			/* 213 */ &EKeys::Invalid,
-			/* 214 */ &EKeys::Invalid,
-			/* 215 */ &EKeys::Invalid,
-			/* 216 */ &EKeys::Invalid,
-			/* 217 */ &EKeys::Invalid,
-			/* 218 */ &EKeys::Invalid,
-			/* 219 */ &EKeys::LeftBracket,
-			/* 220 */ &EKeys::Backslash,
-			/* 221 */ &EKeys::RightBracket,
-			/* 222 */ &EKeys::Apostrophe,
-			/* 223 */ &EKeys::Tilde,
-			/* 224 */ &EKeys::Invalid,
-			/* 225 */ &EKeys::Invalid,
-			/* 226 */ &EKeys::Invalid,
-			/* 227 */ &EKeys::Invalid,
-			/* 228 */ &EKeys::Invalid,
-			/* 229 */ &EKeys::Invalid,
-			/* 230 */ &EKeys::Invalid,
-			/* 231 */ &EKeys::Invalid,
-			/* 232 */ &EKeys::Invalid,
-			/* 233 */ &EKeys::Invalid,
-			/* 234 */ &EKeys::Invalid,
-			/* 235 */ &EKeys::Invalid,
-			/* 236 */ &EKeys::Invalid,
-			/* 237 */ &EKeys::Invalid,
-			/* 238 */ &EKeys::Invalid,
-			/* 239 */ &EKeys::Invalid,
-			/* 240 */ &EKeys::Invalid,
-			/* 241 */ &EKeys::Invalid,
-			/* 242 */ &EKeys::Invalid,
-			/* 243 */ &EKeys::Invalid,
-			/* 244 */ &EKeys::Invalid,
-			/* 245 */ &EKeys::Invalid,
-			/* 246 */ &EKeys::Invalid,
-			/* 247 */ &EKeys::Invalid,
-			/* 248 */ &EKeys::Invalid,
-			/* 249 */ &EKeys::Invalid,
-			/* 250 */ &EKeys::Invalid,
-			/* 251 */ &EKeys::Invalid,
-			/* 252 */ &EKeys::Invalid,
-
-			// Special key codes.
+	static const FKey* JavaScriptKeyCodeToFKey[256] =
+	{
+		/*   0 */ &EKeys::Invalid,
+		/*   1 */ &EKeys::Invalid,
+		/*   2 */ &EKeys::Invalid,
+		/*   3 */ &EKeys::Invalid,
+		/*   4 */ &EKeys::Invalid,
+		/*   5 */ &EKeys::Invalid,
+		/*   6 */ &EKeys::Invalid,
+		/*   7 */ &EKeys::Invalid,
+		/*   8 */ &EKeys::BackSpace,
+		/*   9 */ &EKeys::Tab,
+		/*  10 */ &EKeys::Invalid,
+		/*  11 */ &EKeys::Invalid,
+		/*  12 */ /*&EKeys::Clear*/ &EKeys::Invalid,
+		/*  13 */ &EKeys::Enter,
+		/*  14 */ &EKeys::Invalid,
+		/*  15 */ &EKeys::Invalid,
+		/*  16 */ &EKeys::LeftShift,
+		/*  17 */ &EKeys::LeftControl,
+		/*  18 */ &EKeys::LeftAlt,
+		/*  19 */ &EKeys::Pause,
+		/*  20 */ &EKeys::CapsLock,
+		/*  21 */ &EKeys::Invalid,
+		/*  22 */ &EKeys::Invalid,
+		/*  23 */ &EKeys::Invalid,
+		/*  24 */ &EKeys::Invalid,
+		/*  25 */ &EKeys::Invalid,
+		/*  26 */ &EKeys::Invalid,
+		/*  27 */ &EKeys::Escape,
+		/*  28 */ &EKeys::Invalid,
+		/*  29 */ &EKeys::Invalid,
+		/*  30 */ &EKeys::Invalid,
+		/*  31 */ &EKeys::Invalid,
+		/*  32 */ &EKeys::SpaceBar,
+		/*  33 */ &EKeys::PageUp,
+		/*  34 */ &EKeys::PageDown,
+		/*  35 */ &EKeys::End,
+		/*  36 */ &EKeys::Home,
+		/*  37 */ &EKeys::Left,
+		/*  38 */ &EKeys::Up,
+		/*  39 */ &EKeys::Right,
+		/*  40 */ &EKeys::Down,
+		/*  41 */ &EKeys::Invalid,
+		/*  42 */ &EKeys::Invalid,
+		/*  43 */ &EKeys::Invalid,
+		/*  44 */ /*&EKeys::PrintScreen*/ &EKeys::Invalid,
+		/*  45 */ &EKeys::Insert,
+		/*  46 */ &EKeys::Delete,
+		/*  47 */ &EKeys::Invalid,
+		/*  48 */ &EKeys::Zero,
+		/*  49 */ &EKeys::One,
+		/*  50 */ &EKeys::Two,
+		/*  51 */ &EKeys::Three,
+		/*  52 */ &EKeys::Four,
+		/*  53 */ &EKeys::Five,
+		/*  54 */ &EKeys::Six,
+		/*  55 */ &EKeys::Seven,
+		/*  56 */ &EKeys::Eight,
+		/*  57 */ &EKeys::Nine,
+		/*  58 */ &EKeys::Invalid,
+		/*  59 */ &EKeys::Invalid,
+		/*  60 */ &EKeys::Invalid,
+		/*  61 */ &EKeys::Invalid,
+		/*  62 */ &EKeys::Invalid,
+		/*  63 */ &EKeys::Invalid,
+		/*  64 */ &EKeys::Invalid,
+		/*  65 */ &EKeys::A,
+		/*  66 */ &EKeys::B,
+		/*  67 */ &EKeys::C,
+		/*  68 */ &EKeys::D,
+		/*  69 */ &EKeys::E,
+		/*  70 */ &EKeys::F,
+		/*  71 */ &EKeys::G,
+		/*  72 */ &EKeys::H,
+		/*  73 */ &EKeys::I,
+		/*  74 */ &EKeys::J,
+		/*  75 */ &EKeys::K,
+		/*  76 */ &EKeys::L,
+		/*  77 */ &EKeys::M,
+		/*  78 */ &EKeys::N,
+		/*  79 */ &EKeys::O,
+		/*  80 */ &EKeys::P,
+		/*  81 */ &EKeys::Q,
+		/*  82 */ &EKeys::R,
+		/*  83 */ &EKeys::S,
+		/*  84 */ &EKeys::T,
+		/*  85 */ &EKeys::U,
+		/*  86 */ &EKeys::V,
+		/*  87 */ &EKeys::W,
+		/*  88 */ &EKeys::X,
+		/*  89 */ &EKeys::Y,
+		/*  90 */ &EKeys::Z,
+		/*  91 */ /*&EKeys::LeftWindowKey*/ &EKeys::Invalid,
+		/*  92 */ /*&EKeys::RightWindowKey*/ &EKeys::Invalid,
+		/*  93 */ /*&EKeys::SelectKey*/ &EKeys::Invalid,
+		/*  94 */ &EKeys::Invalid,
+		/*  95 */ &EKeys::Invalid,
+		/*  96 */ &EKeys::NumPadZero,
+		/*  97 */ &EKeys::NumPadOne,
+		/*  98 */ &EKeys::NumPadTwo,
+		/*  99 */ &EKeys::NumPadThree,
+		/* 100 */ &EKeys::NumPadFour,
+		/* 101 */ &EKeys::NumPadFive,
+		/* 102 */ &EKeys::NumPadSix,
+		/* 103 */ &EKeys::NumPadSeven,
+		/* 104 */ &EKeys::NumPadEight,
+		/* 105 */ &EKeys::NumPadNine,
+		/* 106 */ &EKeys::Multiply,
+		/* 107 */ &EKeys::Add,
+		/* 108 */ &EKeys::Invalid,
+		/* 109 */ &EKeys::Subtract,
+		/* 110 */ &EKeys::Decimal,
+		/* 111 */ &EKeys::Divide,
+		/* 112 */ &EKeys::F1,
+		/* 113 */ &EKeys::F2,
+		/* 114 */ &EKeys::F3,
+		/* 115 */ &EKeys::F4,
+		/* 116 */ &EKeys::F5,
+		/* 117 */ &EKeys::F6,
+		/* 118 */ &EKeys::F7,
+		/* 119 */ &EKeys::F8,
+		/* 120 */ &EKeys::F9,
+		/* 121 */ &EKeys::F10,
+		/* 122 */ &EKeys::F11,
+		/* 123 */ &EKeys::F12,
+		/* 124 */ &EKeys::Invalid,
+		/* 125 */ &EKeys::Invalid,
+		/* 126 */ &EKeys::Invalid,
+		/* 127 */ &EKeys::Invalid,
+		/* 128 */ &EKeys::Invalid,
+		/* 129 */ &EKeys::Invalid,
+		/* 130 */ &EKeys::Invalid,
+		/* 131 */ &EKeys::Invalid,
+		/* 132 */ &EKeys::Invalid,
+		/* 133 */ &EKeys::Invalid,
+		/* 134 */ &EKeys::Invalid,
+		/* 135 */ &EKeys::Invalid,
+		/* 136 */ &EKeys::Invalid,
+		/* 137 */ &EKeys::Invalid,
+		/* 138 */ &EKeys::Invalid,
+		/* 139 */ &EKeys::Invalid,
+		/* 140 */ &EKeys::Invalid,
+		/* 141 */ &EKeys::Invalid,
+		/* 142 */ &EKeys::Invalid,
+		/* 143 */ &EKeys::Invalid,
+		/* 144 */ &EKeys::NumLock,
+		/* 145 */ &EKeys::ScrollLock,
+		/* 146 */ &EKeys::Invalid,
+		/* 147 */ &EKeys::Invalid,
+		/* 148 */ &EKeys::Invalid,
+		/* 149 */ &EKeys::Invalid,
+		/* 150 */ &EKeys::Invalid,
+		/* 151 */ &EKeys::Invalid,
+		/* 152 */ &EKeys::Invalid,
+		/* 153 */ &EKeys::Invalid,
+		/* 154 */ &EKeys::Invalid,
+		/* 155 */ &EKeys::Invalid,
+		/* 156 */ &EKeys::Invalid,
+		/* 157 */ &EKeys::Invalid,
+		/* 158 */ &EKeys::Invalid,
+		/* 159 */ &EKeys::Invalid,
+		/* 160 */ &EKeys::Invalid,
+		/* 161 */ &EKeys::Invalid,
+		/* 162 */ &EKeys::Invalid,
+		/* 163 */ &EKeys::Invalid,
+		/* 164 */ &EKeys::Invalid,
+		/* 165 */ &EKeys::Invalid,
+		/* 166 */ &EKeys::Invalid,
+		/* 167 */ &EKeys::Invalid,
+		/* 168 */ &EKeys::Invalid,
+		/* 169 */ &EKeys::Invalid,
+		/* 170 */ &EKeys::Invalid,
+		/* 171 */ &EKeys::Invalid,
+		/* 172 */ &EKeys::Invalid,
+		/* 173 */ &EKeys::Invalid,
+		/* 174 */ &EKeys::Invalid,
+		/* 175 */ &EKeys::Invalid,
+		/* 176 */ &EKeys::Invalid,
+		/* 177 */ &EKeys::Invalid,
+		/* 178 */ &EKeys::Invalid,
+		/* 179 */ &EKeys::Invalid,
+		/* 180 */ &EKeys::Invalid,
+		/* 181 */ &EKeys::Invalid,
+		/* 182 */ &EKeys::Invalid,
+		/* 183 */ &EKeys::Invalid,
+		/* 184 */ &EKeys::Invalid,
+		/* 185 */ &EKeys::Invalid,
+		/* 186 */ &EKeys::Semicolon,
+		/* 187 */ &EKeys::Equals,
+		/* 188 */ &EKeys::Comma,
+		/* 189 */ &EKeys::Hyphen,
+		/* 190 */ &EKeys::Period,
+		/* 191 */ &EKeys::Slash,
+		/* 192 */ &EKeys::Tilde,
+		/* 193 */ &EKeys::Invalid,
+		/* 194 */ &EKeys::Invalid,
+		/* 195 */ &EKeys::Invalid,
+		/* 196 */ &EKeys::Invalid,
+		/* 197 */ &EKeys::Invalid,
+		/* 198 */ &EKeys::Invalid,
+		/* 199 */ &EKeys::Invalid,
+		/* 200 */ &EKeys::Invalid,
+		/* 201 */ &EKeys::Invalid,
+		/* 202 */ &EKeys::Invalid,
+		/* 203 */ &EKeys::Invalid,
+		/* 204 */ &EKeys::Invalid,
+		/* 205 */ &EKeys::Invalid,
+		/* 206 */ &EKeys::Invalid,
+		/* 207 */ &EKeys::Invalid,
+		/* 208 */ &EKeys::Invalid,
+		/* 209 */ &EKeys::Invalid,
+		/* 210 */ &EKeys::Invalid,
+		/* 211 */ &EKeys::Invalid,
+		/* 212 */ &EKeys::Invalid,
+		/* 213 */ &EKeys::Invalid,
+		/* 214 */ &EKeys::Invalid,
+		/* 215 */ &EKeys::Invalid,
+		/* 216 */ &EKeys::Invalid,
+		/* 217 */ &EKeys::Invalid,
+		/* 218 */ &EKeys::Invalid,
+		/* 219 */ &EKeys::LeftBracket,
+		/* 220 */ &EKeys::Backslash,
+		/* 221 */ &EKeys::RightBracket,
+		/* 222 */ &EKeys::Apostrophe,
+		/* 223 */ &EKeys::Tilde,
+		/* 224 */ &EKeys::Invalid,
+		/* 225 */ &EKeys::Invalid,
+		/* 226 */ &EKeys::Invalid,
+		/* 227 */ &EKeys::Invalid,
+		/* 228 */ &EKeys::Invalid,
+		/* 229 */ &EKeys::Invalid,
+		/* 230 */ &EKeys::Invalid,
+		/* 231 */ &EKeys::Invalid,
+		/* 232 */ &EKeys::Invalid,
+		/* 233 */ &EKeys::Invalid,
+		/* 234 */ &EKeys::Invalid,
+		/* 235 */ &EKeys::Invalid,
+		/* 236 */ &EKeys::Invalid,
+		/* 237 */ &EKeys::Invalid,
+		/* 238 */ &EKeys::Invalid,
+		/* 239 */ &EKeys::Invalid,
+		/* 240 */ &EKeys::Invalid,
+		/* 241 */ &EKeys::Invalid,
+		/* 242 */ &EKeys::Invalid,
+		/* 243 */ &EKeys::Invalid,
+		/* 244 */ &EKeys::Invalid,
+		/* 245 */ &EKeys::Invalid,
+		/* 246 */ &EKeys::Invalid,
+		/* 247 */ &EKeys::Invalid,
+		/* 248 */ &EKeys::Invalid,
+		/* 249 */ &EKeys::Invalid,
+		/* 250 */ &EKeys::Invalid,
+		/* 251 */ &EKeys::Invalid,
+		/* 252 */ &EKeys::Invalid,
 
-			/* 253 */ &EKeys::RightShift,
-			/* 254 */ &EKeys::RightControl,
-			/* 255 */ &EKeys::RightAlt
-		};
-	}
-}
+		// Special key codes.
+		/* 253 */ &EKeys::RightShift,
+		/* 254 */ &EKeys::RightControl,
+		/* 255 */ &EKeys::RightAlt
+	};
+} // namespace UE::PixelStreaming
\ No newline at end of file
diff --git a/Engine/Plugins/Media/PixelStreaming/Source/PixelStreaming/Private/PixelStreamingModule.cpp b/Engine/Plugins/Media/PixelStreaming/Source/PixelStreaming/Private/PixelStreamingModule.cpp
--- a/Engine/Plugins/Media/PixelStreaming/Source/PixelStreaming/Private/PixelStreamingModule.cpp
+++ b/Engine/Plugins/Media/PixelStreaming/Source/PixelStreaming/Private/PixelStreamingModule.cpp
@@ -47,6 +47,7 @@ THIRD_PARTY_INCLUDES_END
 #include "Async/Async.h"
 #include "Engine/Engine.h"
 #include "EncoderFactory.h"
+#include "VideoSource.h"
 
 #if !UE_BUILD_SHIPPING
 	#include "DrawDebugHelpers.h"
@@ -106,6 +107,10 @@ namespace UE::PixelStreaming
 			return;
 		}
 
+		Poller = MakeUnique<FPoller>();
+		FramePump = MakeUnique<FFixedFPSPump>();
+		bDecoupleFrameRate = UE::PixelStreaming::Settings::CVarPixelStreamingDecoupleFrameRate.GetValueOnAnyThread();
+
 		// subscribe to engine delegates here for init / framebuffer creation / whatever
 		// TODO check if there is a better callback to attach so that we can use with editor
 		if (FSlateApplication::IsInitialized())
@@ -123,7 +128,7 @@ namespace UE::PixelStreaming
 
 		verify(FModuleManager::Get().LoadModule(FName("ImageWrapper")));
 
-		Streamer = MakeShared<FStreamer>(StreamerId);
+		Streamer = MakeUnique<FStreamer>(StreamerId);
 
 		// Streamer has been created, so module is now "ready" for external use.
 		ReadyEvent.Broadcast(*this);
@@ -185,6 +190,8 @@ namespace UE::PixelStreaming
 		{
 			UE_LOG(LogPixelStreaming, Warning, TEXT("Only D3D11/D3D12/Vulkan Dynamic RHI is supported. Detected %s"), GDynamicRHI != nullptr ? GDynamicRHI->GetName() : TEXT("[null]"));
 		}
+
+		TextureSourceFactory = MakeUnique<FTextureSourceFactory>();
 	}
 
 	void FPixelStreamingModule::ShutdownModule()
@@ -196,6 +203,9 @@ namespace UE::PixelStreaming
 		}
 
 		IModularFeatures::Get().UnregisterModularFeature(GetModularFeatureName(), this);
+
+		// We explicitly call release on streamer so WebRTC gets shutdown before our module is deleted
+		Streamer.Reset(nullptr);
 	}
 
 	IPixelStreamingModule* FPixelStreamingModule::GetModule()
@@ -274,6 +284,12 @@ namespace UE::PixelStreaming
 
 		check(IsInRenderingThread());
 
+		// Framerate isn't decoupled, manually pump
+		if (!bDecoupleFrameRate)
+		{
+			FramePump->PumpAll();
+		}
+
 		// Check to see if we have been instructed to capture the back buffer as a
 		// freeze frame.
 		if (bCaptureNextBackBufferAndStream && Streamer->IsStreaming())
@@ -282,7 +298,7 @@ namespace UE::PixelStreaming
 
 			// Read the data out of the back buffer and send as a JPEG.
 			FRHICommandListImmediate& RHICmdList = FRHICommandListExecutor::GetImmediateCommandList();
-			FIntRect Rect(0, 0, BackBuffer->GetSizeX(), BackBuffer->GetSizeY());
+			FIntRect Rect(0, 0, BackBuffer->GetDesc().Extent.X, BackBuffer->GetDesc().Extent.Y);
 			TArray<FColor> Data;
 
 			RHICmdList.ReadSurfaceData(BackBuffer, Rect, Data, FReadSurfaceDataFlags());
@@ -292,7 +308,8 @@ namespace UE::PixelStreaming
 
 	TSharedPtr<class IInputDevice> FPixelStreamingModule::CreateInputDevice(const TSharedRef<FGenericApplicationMessageHandler>& InMessageHandler)
 	{
-		InputDevice = MakeShareable(new FInputDevice(InMessageHandler));
+		FInputDevice* InputDevicePtr = new FInputDevice(InMessageHandler);
+		InputDevice = TSharedPtr<FInputDevice>(InputDevicePtr);
 		return InputDevice;
 	}
 
@@ -348,21 +365,21 @@ namespace UE::PixelStreaming
 			ENQUEUE_RENDER_COMMAND(ReadSurfaceCommand)
 			([this, Texture](FRHICommandListImmediate& RHICmdList) {
 				// A frame is supplied so immediately read its data and send as a JPEG.
-				FTexture2DRHIRef Texture2DRHI = (Texture->GetResource() && Texture->GetResource()->TextureRHI) ? Texture->GetResource()->TextureRHI->GetTexture2D() : nullptr;
-				if (!Texture2DRHI)
+				FTextureRHIRef TextureRHI = Texture->GetResource() ? Texture->GetResource()->TextureRHI : nullptr;
+				if (!TextureRHI)
 				{
-					UE_LOG(LogPixelStreaming, Error, TEXT("Attempting freeze frame with texture %s with no texture 2D RHI"), *Texture->GetName());
+					UE_LOG(LogPixelStreaming, Error, TEXT("Attempting freeze frame with texture %s with no texture RHI"), *Texture->GetName());
 					return;
 				}
-				uint32 Width = Texture2DRHI->GetSizeX();
-				uint32 Height = Texture2DRHI->GetSizeY();
+				uint32 Width = TextureRHI->GetDesc().Extent.X;
+				uint32 Height = TextureRHI->GetDesc().Extent.Y;
 
-				FTexture2DRHIRef DestTexture = CreateTexture(Width, Height);
+				FTextureRHIRef DestTexture = CreateRHITexture(Width, Height);
 
 				FGPUFenceRHIRef CopyFence = GDynamicRHI->RHICreateGPUFence(*FString::Printf(TEXT("FreezeFrameFence")));
 
 				// Copy freeze frame texture to empty texture
-				CopyTexture(Texture2DRHI, DestTexture, CopyFence);
+				CopyTextureToRHI(TextureRHI, DestTexture, CopyFence);
 
 				TArray<FColor> Data;
 				FIntRect Rect(0, 0, Width, Height);
@@ -516,26 +533,48 @@ namespace UE::PixelStreaming
 
 		return Result;
 	}
-} // namespace UE::PixelStreaming
 
-void UE::PixelStreaming::FPixelStreamingModule::AddGPUFencePollerTask(FGPUFenceRHIRef Fence, TSharedRef<bool, ESPMode::ThreadSafe> bIsEnabled, TFunction<void()> Task)
-{
-	FencePollerThread.AddJob(Fence, bIsEnabled, Task);
-}
+	void FPixelStreamingModule::AddPollerTask(TFunction<void()> Task, TFunction<bool()> IsTaskFinished, TSharedRef<bool, ESPMode::ThreadSafe> bKeepRunning)
+	{
+		Poller->AddJob(IsTaskFinished, bKeepRunning, Task);
+	}
 
-void UE::PixelStreaming::FPixelStreamingModule::RegisterVideoSource(FPixelStreamingPlayerId PlayerId, IPumpedVideoSource* VideoSource)
-{
-	PumpThread.RegisterVideoSource(PlayerId, VideoSource);
-}
+	void FPixelStreamingModule::RegisterPumpable(rtc::scoped_refptr<FPixelStreamingPumpable> Pumpable)
+	{
+		FramePump->RegisterPumpable(Pumpable);
+	}
 
-void UE::PixelStreaming::FPixelStreamingModule::UnregisterVideoSource(FPixelStreamingPlayerId PlayerId)
-{
-	PumpThread.UnregisterVideoSource(PlayerId);
-}
+	void FPixelStreamingModule::UnregisterPumpable(rtc::scoped_refptr<FPixelStreamingPumpable> Pumpable)
+	{
+		FramePump->UnregisterPumpable(Pumpable);
+	}
 
-webrtc::VideoEncoderFactory* UE::PixelStreaming::FPixelStreamingModule::CreateVideoEncoderFactory()
-{
-	return new UE::PixelStreaming::FVideoEncoderFactory();
-}
+	webrtc::VideoEncoderFactory* FPixelStreamingModule::CreateVideoEncoderFactory()
+	{
+		return new FVideoEncoderFactory();
+	}
+
+	rtc::scoped_refptr<webrtc::VideoTrackSourceInterface> FPixelStreamingModule::CreateExternalVideoSource(FName SourceType)
+	{
+		FVideoSourceP2P* VideoSource = new FVideoSourceP2P(SourceType, []() { return true; });
+		VideoSource->Initialize();
+		return rtc::scoped_refptr<webrtc::VideoTrackSourceInterface>(VideoSource);
+	}
+
+	IPixelStreamingTextureSourceFactory& FPixelStreamingModule::GetTextureSourceFactory()
+	{
+		return *TextureSourceFactory;
+	}
+
+	void FPixelStreamingModule::SetActiveTextureSourceTypes(const TArray<FName>& SourceTypes)
+	{
+		UE::PixelStreaming::Settings::SetActiveTextureSourceTypes(SourceTypes);
+	}
+
+	const TArray<FName>& FPixelStreamingModule::GetActiveTextureSourceTypes() const
+	{
+		return UE::PixelStreaming::Settings::GetActiveTextureSourceTypes();
+	}
+} // namespace UE::PixelStreaming
 
 IMPLEMENT_MODULE(UE::PixelStreaming::FPixelStreamingModule, PixelStreaming)
diff --git a/Engine/Plugins/Media/PixelStreaming/Source/PixelStreaming/Private/PixelStreamingModule.h b/Engine/Plugins/Media/PixelStreaming/Source/PixelStreaming/Private/PixelStreamingModule.h
--- a/Engine/Plugins/Media/PixelStreaming/Source/PixelStreaming/Private/PixelStreamingModule.h
+++ b/Engine/Plugins/Media/PixelStreaming/Source/PixelStreaming/Private/PixelStreamingModule.h
@@ -6,8 +6,10 @@
 #include "RHI.h"
 #include "Tickable.h"
 #include "InputDevice.h"
-#include "GPUFencePoller.h"
+#include "Poller.h"
 #include "FixedFPSPump.h"
+#include "PixelStreamingPumpable.h"
+#include "TextureSourceFactory.h"
 
 class AController;
 class AGameModeBase;
@@ -31,11 +33,6 @@ namespace UE::PixelStreaming
 		virtual bool StartStreaming(const FString& SignallingServerUrl) override;
 		virtual void StopStreaming() override;
 
-		virtual webrtc::VideoEncoderFactory* CreateVideoEncoderFactory() override;
-		virtual void RegisterVideoSource(FPixelStreamingPlayerId PlayerId, IPumpedVideoSource* VideoSource) override;
-		virtual void UnregisterVideoSource(FPixelStreamingPlayerId PlayerId) override;
-		virtual void AddGPUFencePollerTask(FGPUFenceRHIRef Fence, TSharedRef<bool, ESPMode::ThreadSafe> bIsEnabled, TFunction<void()> Task) override;
-
 	private:
 		/** IModuleInterface implementation */
 		void StartupModule() override;
@@ -68,6 +65,14 @@ namespace UE::PixelStreaming
 		void KickPlayer(FPixelStreamingPlayerId PlayerId);
 		IPixelStreamingAudioSink* GetPeerAudioSink(FPixelStreamingPlayerId PlayerId) override;
 		IPixelStreamingAudioSink* GetUnlistenedAudioSink() override;
+		virtual rtc::scoped_refptr<webrtc::VideoTrackSourceInterface> CreateExternalVideoSource(FName SourceType) override;
+		virtual webrtc::VideoEncoderFactory* CreateVideoEncoderFactory() override;
+		virtual void RegisterPumpable(rtc::scoped_refptr<FPixelStreamingPumpable> Pumpable) override;
+		virtual void UnregisterPumpable(rtc::scoped_refptr<FPixelStreamingPumpable> Pumpable) override;
+		virtual void AddPollerTask(TFunction<void()> Task, TFunction<bool()> IsTaskFinished, TSharedRef<bool, ESPMode::ThreadSafe> bKeepRunning) override;
+		virtual IPixelStreamingTextureSourceFactory& GetTextureSourceFactory() override;
+		virtual void SetActiveTextureSourceTypes(const TArray<FName>& SourceTypes) override;
+		virtual const TArray<FName>& GetActiveTextureSourceTypes() const override;
 		/** End IPixelStreamingModule implementation */
 
 		// FTickableGameObject
@@ -88,7 +93,7 @@ namespace UE::PixelStreaming
 		FReadyEvent ReadyEvent;
 		FStreamingStartedEvent StreamingStartedEvent;
 		FStreamingStoppedEvent StreamingStoppedEvent;
-		TSharedPtr<FStreamer> Streamer;
+		TUniquePtr<FStreamer> Streamer;
 		TSharedPtr<FInputDevice> InputDevice;
 		TArray<UPixelStreamingInput*> InputComponents;
 		bool bFrozen = false;
@@ -96,7 +101,10 @@ namespace UE::PixelStreaming
 		double LastVideoEncoderQPReportTime = 0;
 		static IPixelStreamingModule* PixelStreamingModule;
 
-		FFixedFPSPump PumpThread;
-		FGPUFencePoller FencePollerThread;
+		TUniquePtr<FFixedFPSPump> FramePump;
+		TUniquePtr<FPoller> Poller;
+		TUniquePtr<IPixelStreamingTextureSourceFactory> TextureSourceFactory;
+
+		bool bDecoupleFrameRate = false;
 	};
 } // namespace UE::PixelStreaming
diff --git a/Engine/Plugins/Media/PixelStreaming/Source/PixelStreaming/Private/PlayerSession.cpp b/Engine/Plugins/Media/PixelStreaming/Source/PixelStreaming/Private/PlayerSession.cpp
--- a/Engine/Plugins/Media/PixelStreaming/Source/PixelStreaming/Private/PlayerSession.cpp
+++ b/Engine/Plugins/Media/PixelStreaming/Source/PixelStreaming/Private/PlayerSession.cpp
@@ -202,14 +202,14 @@ namespace UE::PixelStreaming
 		return &AudioSink;
 	}
 
-	bool FPlayerSession::SendMessage(Protocol::EToPlayerMsg Type, const FString& Descriptor) const
+	bool FPlayerSession::SendMessage(Protocol::EToPlayerMsg InMessageType, const FString& Descriptor) const
 	{
 		if (!DataChannel)
 		{
 			return false;
 		}
 
-		const uint8 MessageType = static_cast<uint8>(Type);
+		const uint8 MessageType = static_cast<uint8>(InMessageType);
 		const size_t DescriptorSize = Descriptor.Len() * sizeof(TCHAR);
 
 		rtc::CopyOnWriteBuffer Buffer(sizeof(MessageType) + DescriptorSize);
diff --git a/Engine/Plugins/Media/PixelStreaming/Source/PixelStreaming/Private/PlayerSession.h b/Engine/Plugins/Media/PixelStreaming/Source/PixelStreaming/Private/PlayerSession.h
--- a/Engine/Plugins/Media/PixelStreaming/Source/PixelStreaming/Private/PlayerSession.h
+++ b/Engine/Plugins/Media/PixelStreaming/Source/PixelStreaming/Private/PlayerSession.h
@@ -30,6 +30,7 @@ namespace UE::PixelStreaming
 		virtual void OnRemoteIceCandidate(const FString& SdpMid, int SdpMLineIndex, const FString& Sdp) override;
 		virtual void DisconnectPlayer(const FString& Reason) override;
 
+		virtual FName GetSessionType() const override { return Type; }
 		virtual FPixelStreamingPlayerId GetPlayerId() const override;
 		virtual IPixelStreamingAudioSink* GetAudioSink() override;
 		virtual FDataChannelObserver* GetDataChannelObserver() override;
@@ -63,6 +64,9 @@ namespace UE::PixelStreaming
 		virtual void OnTrack(rtc::scoped_refptr<webrtc::RtpTransceiverInterface> transceiver) override;
 		virtual void OnRemoveTrack(rtc::scoped_refptr<webrtc::RtpReceiverInterface> receiver) override;
 
+	public:
+		inline static const FName Type = FName(TEXT("P2P"));
+
 	protected:
 		FPlayerSessions* PlayerSessions;
 		FSignallingServerConnection* SignallingServerConnection;
diff --git a/Engine/Plugins/Media/PixelStreaming/Source/PixelStreaming/Private/PlayerSessionDataOnly.cpp b/Engine/Plugins/Media/PixelStreaming/Source/PixelStreaming/Private/PlayerSessionDataOnly.cpp
--- a/Engine/Plugins/Media/PixelStreaming/Source/PixelStreaming/Private/PlayerSessionDataOnly.cpp
+++ b/Engine/Plugins/Media/PixelStreaming/Source/PixelStreaming/Private/PlayerSessionDataOnly.cpp
@@ -51,14 +51,14 @@ namespace UE::PixelStreaming
 		RecvDataChannel = nullptr;
 	}
 
-	bool FPlayerSessionDataOnly::SendMessage(Protocol::EToPlayerMsg Type, const FString& Descriptor) const
+	bool FPlayerSessionDataOnly::SendMessage(Protocol::EToPlayerMsg InMessageType, const FString& Descriptor) const
 	{
 		if (!SendDataChannel)
 		{
 			return false;
 		}
 
-		const uint8 MessageType = static_cast<uint8>(Type);
+		const uint8 MessageType = static_cast<uint8>(InMessageType);
 		const size_t DescriptorSize = Descriptor.Len() * sizeof(TCHAR);
 
 		rtc::CopyOnWriteBuffer Buffer(sizeof(MessageType) + DescriptorSize);
diff --git a/Engine/Plugins/Media/PixelStreaming/Source/PixelStreaming/Private/PlayerSessionDataOnly.h b/Engine/Plugins/Media/PixelStreaming/Source/PixelStreaming/Private/PlayerSessionDataOnly.h
--- a/Engine/Plugins/Media/PixelStreaming/Source/PixelStreaming/Private/PlayerSessionDataOnly.h
+++ b/Engine/Plugins/Media/PixelStreaming/Source/PixelStreaming/Private/PlayerSessionDataOnly.h
@@ -31,6 +31,7 @@ namespace UE::PixelStreaming
 		virtual void OnRemoteIceCandidate(const FString& SdpMid, int SdpMLineIndex, const FString& Sdp) override {}
 		virtual void DisconnectPlayer(const FString& Reason) override {}
 
+		virtual FName GetSessionType() const override { return Type; };
 		virtual FPixelStreamingPlayerId GetPlayerId() const override { return PlayerId; }
 		virtual IPixelStreamingAudioSink* GetAudioSink() override { return nullptr; }
 		virtual FDataChannelObserver* GetDataChannelObserver() override { return &DataChannelObserver; }
@@ -62,6 +63,9 @@ namespace UE::PixelStreaming
 		virtual void OnTrack(rtc::scoped_refptr<webrtc::RtpTransceiverInterface> transceiver) override {}
 		virtual void OnRemoveTrack(rtc::scoped_refptr<webrtc::RtpReceiverInterface> receiver) override {}
 
+	public:
+		inline static const FName Type = FName(TEXT("DataOnly"));
+
 	private:
 		FSignallingServerConnection* SignallingServerConnection;
 		FPixelStreamingPlayerId PlayerId;
diff --git a/Engine/Plugins/Media/PixelStreaming/Source/PixelStreaming/Private/PlayerSessionSFU.h b/Engine/Plugins/Media/PixelStreaming/Source/PixelStreaming/Private/PlayerSessionSFU.h
--- a/Engine/Plugins/Media/PixelStreaming/Source/PixelStreaming/Private/PlayerSessionSFU.h
+++ b/Engine/Plugins/Media/PixelStreaming/Source/PixelStreaming/Private/PlayerSessionSFU.h
@@ -14,6 +14,10 @@ namespace UE::PixelStreaming
 		virtual ~FPlayerSessionSFU();
 
 		void AddChildSession(TSharedPtr<FPlayerSessionDataOnly> ChildSession);
+		virtual FName GetSessionType() const override { return Type; }
+
+	public:
+		inline static const FName Type = FName(TEXT("SFU"));
 
 	private:
 		TArray<TWeakPtr<FPlayerSessionDataOnly>> ChildSessions;
diff --git a/Engine/Plugins/Media/PixelStreaming/Source/PixelStreaming/Private/PlayerSessions.cpp b/Engine/Plugins/Media/PixelStreaming/Source/PixelStreaming/Private/PlayerSessions.cpp
--- a/Engine/Plugins/Media/PixelStreaming/Source/PixelStreaming/Private/PlayerSessions.cpp
+++ b/Engine/Plugins/Media/PixelStreaming/Source/PixelStreaming/Private/PlayerSessions.cpp
@@ -12,10 +12,9 @@ namespace UE::PixelStreaming
 {
 	TSharedPtr<IPlayerSession> FPlayerSessions::CreatePlayerSession(
 		FPixelStreamingPlayerId PlayerId,
-		rtc::scoped_refptr<webrtc::PeerConnectionFactoryInterface>
-			PeerConnectionFactory,
+		rtc::scoped_refptr<webrtc::PeerConnectionFactoryInterface> PeerConnectionFactory,
 		webrtc::PeerConnectionInterface::RTCConfiguration PeerConnectionConfig,
-		FSignallingServerConnection *SignallingServerConnection, int Flags)
+		FSignallingServerConnection* SignallingServerConnection, int Flags)
 	{
 		check(PeerConnectionFactory);
 
@@ -27,84 +26,82 @@ namespace UE::PixelStreaming
 			return nullptr;
 		}
 
-		UE_LOG(LogPixelStreaming, Log,
-			   TEXT("Creating player session for PlayerId=%s"), *PlayerId);
+		UE_LOG(LogPixelStreaming, Log, TEXT("Creating player session for PlayerId=%s"), *PlayerId);
 
 		// this is called from WebRTC signalling thread, the only thread where
 		// `Players` map is modified, so no need to lock it
 		bool bIsSFU = (Flags & Protocol::EPlayerFlags::PSPFlag_IsSFU) != 0;
-		bool bMakeQualityController =
-			!bIsSFU && Players.Num() == 0; // first player controls quality by default
+
 		TSharedPtr<IPlayerSession> Session;
 		if (bIsSFU)
 		{
-			Session = MakeShared<FPlayerSessionSFU>(this, SignallingServerConnection,
-													PlayerId);
+			Session = MakeShared<FPlayerSessionSFU>(this, SignallingServerConnection, PlayerId);
 		}
 		else
 		{
-			Session =
-				MakeShared<FPlayerSession>(this, SignallingServerConnection, PlayerId);
-		}
-
-		rtc::scoped_refptr<webrtc::PeerConnectionInterface> PeerConnection =
-			PeerConnectionFactory->CreatePeerConnection(
-				PeerConnectionConfig,
-				webrtc::PeerConnectionDependencies{Session.Get()});
-		if (!PeerConnection)
-		{
-			UE_LOG(LogPixelStreaming, Error,
-				   TEXT("Failed to created PeerConnection. This may indicate you "
-						"passed malformed peerConnectionOptions."));
-			return nullptr;
+			Session = MakeShared<FPlayerSession>(this, SignallingServerConnection, PlayerId);
 		}
 
-		// Setup suggested bitrate settings on the Peer Connection based on our CVars
-		webrtc::BitrateSettings BitrateSettings;
-		BitrateSettings.min_bitrate_bps =
-			Settings::CVarPixelStreamingWebRTCMinBitrate.GetValueOnAnyThread();
-		BitrateSettings.max_bitrate_bps =
-			Settings::CVarPixelStreamingWebRTCMaxBitrate.GetValueOnAnyThread();
-		BitrateSettings.start_bitrate_bps =
-			Settings::CVarPixelStreamingWebRTCStartBitrate.GetValueOnAnyThread();
-		PeerConnection->SetBitrate(BitrateSettings);
-
-		Session->SetPeerConnection(PeerConnection);
+		// Note: this flag is very specifically calculated BEFORE the new session is added.
+		// Peer only becomes quality controller if there is no other P2P peers already connected.
+		bool bMakeQualityController = !bIsSFU && NumP2PPeers() == 0 && Session->GetSessionType() == FPlayerSession::Type;
 
-		// The actual modification of the players map
-		{
-			FScopeLock Lock(&PlayersCS);
-			Players.Add(PlayerId, Session);
-		}
+		Session->SetPeerConnection(CreatePeerConnection(PeerConnectionFactory, PeerConnectionConfig, Session));
+		AddSession(PlayerId, Session);
 
 		if (bMakeQualityController)
 		{
 			SetQualityController(PlayerId);
 		}
 
-		if (UPixelStreamingDelegates *Delegates =
-				UPixelStreamingDelegates::GetPixelStreamingDelegates())
+		if (UPixelStreamingDelegates* Delegates = UPixelStreamingDelegates::GetPixelStreamingDelegates())
 		{
 			Delegates->OnNewConnection.Broadcast(PlayerId, bMakeQualityController);
-			Delegates->OnNewConnectionNative.Broadcast(PlayerId,
-													   bMakeQualityController);
+			Delegates->OnNewConnectionNative.Broadcast(PlayerId, bMakeQualityController);
 		}
 
 		return Session;
 	}
 
+	void FPlayerSessions::AddSession(FPixelStreamingPlayerId PlayerId, TSharedPtr<IPlayerSession> Session)
+	{
+		FScopeLock Lock(&PlayersCS);
+		Players.Add(PlayerId, Session);
+	}
+
+	rtc::scoped_refptr<webrtc::PeerConnectionInterface> FPlayerSessions::CreatePeerConnection(
+		rtc::scoped_refptr<webrtc::PeerConnectionFactoryInterface> PeerConnectionFactory,
+		webrtc::PeerConnectionInterface::RTCConfiguration PeerConnectionConfig,
+		TSharedPtr<IPlayerSession> Session) const
+	{
+		rtc::scoped_refptr<webrtc::PeerConnectionInterface> PeerConnection =
+			PeerConnectionFactory->CreatePeerConnection(PeerConnectionConfig, webrtc::PeerConnectionDependencies{ Session.Get() });
+
+		if (!PeerConnection)
+		{
+			UE_LOG(LogPixelStreaming, Error, TEXT("Failed to created PeerConnection. This may indicate you passed malformed peerConnectionOptions."));
+			return nullptr;
+		}
+
+		// Setup suggested bitrate settings on the Peer Connection based on our CVars
+		webrtc::BitrateSettings BitrateSettings;
+		BitrateSettings.min_bitrate_bps = Settings::CVarPixelStreamingWebRTCMinBitrate.GetValueOnAnyThread();
+		BitrateSettings.max_bitrate_bps = Settings::CVarPixelStreamingWebRTCMaxBitrate.GetValueOnAnyThread();
+		BitrateSettings.start_bitrate_bps = Settings::CVarPixelStreamingWebRTCStartBitrate.GetValueOnAnyThread();
+		PeerConnection->SetBitrate(BitrateSettings);
+
+		return PeerConnection;
+	}
+
 	void FPlayerSessions::CreateNewDataChannel(
 		FPixelStreamingPlayerId SFUId, FPixelStreamingPlayerId PlayerId,
 		int32 SendStreamId, int32 RecvStreamId,
-		FSignallingServerConnection *SignallingServerConnection)
+		FSignallingServerConnection* SignallingServerConnection)
 	{
-		if (auto ParentSession =
-				StaticCastSharedPtr<FPlayerSessionSFU>(GetPlayerSession(SFUId)))
+		if (auto ParentSession = StaticCastSharedPtr<FPlayerSessionSFU>(GetPlayerSession(SFUId)))
 		{
 			auto PeerConnection = &ParentSession->GetPeerConnection();
-			auto NewSession = MakeShared<FPlayerSessionDataOnly>(
-				this, SignallingServerConnection, PlayerId, PeerConnection,
-				SendStreamId, RecvStreamId);
+			auto NewSession = MakeShared<FPlayerSessionDataOnly>(this, SignallingServerConnection, PlayerId, PeerConnection, SendStreamId, RecvStreamId);
 			ParentSession->AddChildSession(NewSession);
 
 			{
@@ -119,9 +116,7 @@ namespace UE::PixelStreaming
 		TSharedPtr<IPlayerSession> PlayerSession = GetPlayerSession(PlayerId);
 		if (!PlayerSession)
 		{
-			UE_LOG(LogPixelStreaming, VeryVerbose,
-				   TEXT("Failed to delete player %s - that player was not found."),
-				   *PlayerId);
+			UE_LOG(LogPixelStreaming, VeryVerbose, TEXT("Failed to delete player %s - that player was not found."), *PlayerId);
 			return GetNumPlayers();
 		}
 
@@ -135,13 +130,11 @@ namespace UE::PixelStreaming
 			RemainingCount = Players.Num();
 		}
 
-		UPixelStreamingDelegates *Delegates =
-			UPixelStreamingDelegates::GetPixelStreamingDelegates();
+		UPixelStreamingDelegates* Delegates = UPixelStreamingDelegates::GetPixelStreamingDelegates();
 		if (Delegates && FModuleManager::Get().IsModuleLoaded("PixelStreaming"))
 		{
 			Delegates->OnClosedConnection.Broadcast(PlayerId, bWasQualityController);
-			Delegates->OnClosedConnectionNative.Broadcast(PlayerId,
-														  bWasQualityController);
+			Delegates->OnClosedConnectionNative.Broadcast(PlayerId, bWasQualityController);
 
 			if (RemainingCount == 0)
 			{
@@ -155,13 +148,51 @@ namespace UE::PixelStreaming
 		if (bWasQualityController && RemainingCount > 0)
 		{
 			// Quality Controller session has been just removed, set quality control to
-			// any of remaining sessions
-			SetQualityController(Players.begin()->Key);
+			// any of remaining player sessions (not SFU!)
+			FPixelStreamingPlayerId OutPlayerId = INVALID_PLAYER_ID;
+			bool bHasPlayer = GetFirstP2PPeer(OutPlayerId);
+			if (bHasPlayer)
+			{
+				SetQualityController(OutPlayerId);
+			}
 		}
 
 		return RemainingCount;
 	}
 
+	int FPlayerSessions::NumP2PPeers() const
+	{
+		FScopeLock Lock(&PlayersCS);
+		int Count = 0;
+		for (auto& Entry : Players)
+		{
+			FPixelStreamingPlayerId Id = Entry.Key;
+			TSharedPtr<IPlayerSession> Session = Entry.Value;
+			if (Session.IsValid() && Session->GetSessionType() == FPlayerSession::Type)
+			{
+				Count++;
+			}
+		}
+		return Count;
+	}
+
+	bool FPlayerSessions::GetFirstP2PPeer(FPixelStreamingPlayerId& OutPlayerId) const
+	{
+		FScopeLock Lock(&PlayersCS);
+
+		for (auto& Entry : Players)
+		{
+			FPixelStreamingPlayerId Id = Entry.Key;
+			TSharedPtr<IPlayerSession> Session = Entry.Value;
+			if (Session.IsValid() && Session->GetSessionType() == FPlayerSession::Type)
+			{
+				OutPlayerId = Id;
+				return true;
+			}
+		}
+		return false;
+	}
+
 	void FPlayerSessions::DeleteAllPlayerSessions()
 	{
 		const FPixelStreamingPlayerId OldQualityController = QualityControllingPlayer;
@@ -180,17 +211,15 @@ namespace UE::PixelStreaming
 		}
 
 		// Notify all delegates of all the closed players
-		UPixelStreamingDelegates *Delegates =
-			UPixelStreamingDelegates::GetPixelStreamingDelegates();
+		UPixelStreamingDelegates* Delegates = UPixelStreamingDelegates::GetPixelStreamingDelegates();
 		if (Delegates && FModuleManager::Get().IsModuleLoaded("PixelStreaming"))
 		{
-			for (auto &&PlayerId : OldPlayerIds)
+			for (auto&& PlayerId : OldPlayerIds)
 			{
 				bool bWasQualityController = OldQualityController == PlayerId;
 
 				Delegates->OnClosedConnection.Broadcast(PlayerId, bWasQualityController);
-				Delegates->OnClosedConnectionNative.Broadcast(PlayerId,
-															  bWasQualityController);
+				Delegates->OnClosedConnectionNative.Broadcast(PlayerId, bWasQualityController);
 			}
 
 			Delegates->OnAllConnectionsClosed.Broadcast();
@@ -200,8 +229,7 @@ namespace UE::PixelStreaming
 
 	int FPlayerSessions::GetNumPlayers() const { return Players.Num(); }
 
-	bool FPlayerSessions::IsQualityController(
-		FPixelStreamingPlayerId PlayerId) const
+	bool FPlayerSessions::IsQualityController(FPixelStreamingPlayerId PlayerId) const
 	{
 		FScopeLock Lock(&QualityControllerCS);
 		return QualityControllingPlayer == PlayerId;
@@ -209,6 +237,19 @@ namespace UE::PixelStreaming
 
 	void FPlayerSessions::SetQualityController(FPixelStreamingPlayerId PlayerId)
 	{
+
+		TSharedPtr<IPlayerSession> PlayerSession = GetPlayerSession(PlayerId);
+
+		if (!PlayerSession.IsValid())
+		{
+			return;
+		}
+
+		if (PlayerSession->GetSessionType() != FPlayerSession::Type)
+		{
+			return;
+		}
+
 		// The actual assignment of the quality controlling peer
 		{
 			FScopeLock Lock(&QualityControllerCS);
@@ -216,22 +257,20 @@ namespace UE::PixelStreaming
 		}
 
 		// Let any listeners know the quality controller has changed
-		if (UPixelStreamingDelegates *Delegates = UPixelStreamingDelegates::GetPixelStreamingDelegates())
+		if (UPixelStreamingDelegates* Delegates = UPixelStreamingDelegates::GetPixelStreamingDelegates())
 		{
 			Delegates->OnQualityControllerChangedNative.Broadcast(PlayerId);
 		}
 
-		UE_LOG(LogPixelStreaming, Log, TEXT("Quality controller is now PlayerId=%s."),
-			   *PlayerId);
+		UE_LOG(LogPixelStreaming, Log, TEXT("Quality controller is now PlayerId=%s."), *PlayerId);
 
 		// Update quality controller status on the browser side too
-		ForEachSession([&](TSharedPtr<IPlayerSession> Session)
-					   {
+		ForEachSession([&](TSharedPtr<IPlayerSession> Session) {
 			bool bIsQualityController = Session->GetPlayerId() == QualityControllingPlayer;
 			Session->SendQualityControlStatus(bIsQualityController); });
 	}
 
-	void FPlayerSessions::ForEachSession(const TFunction<void(TSharedPtr<IPlayerSession>)> &Func)
+	void FPlayerSessions::ForEachSession(const TFunction<void(TSharedPtr<IPlayerSession>)>& Func)
 	{
 		// We have to be careful here
 		// We dont lock here because we can end up in a situation where Func
@@ -242,7 +281,7 @@ namespace UE::PixelStreaming
 		// and when we actually pull it. luckily GetPlayerSession handles that for us.
 		TSet<FPixelStreamingPlayerId> KeySet;
 		Players.GetKeys(KeySet);
-		for (auto &&PlayerId : KeySet)
+		for (auto&& PlayerId : KeySet)
 		{
 			if (TSharedPtr<IPlayerSession> Session = GetPlayerSession(PlayerId))
 			{
diff --git a/Engine/Plugins/Media/PixelStreaming/Source/PixelStreaming/Private/PlayerSessions.h b/Engine/Plugins/Media/PixelStreaming/Source/PixelStreaming/Private/PlayerSessions.h
--- a/Engine/Plugins/Media/PixelStreaming/Source/PixelStreaming/Private/PlayerSessions.h
+++ b/Engine/Plugins/Media/PixelStreaming/Source/PixelStreaming/Private/PlayerSessions.h
@@ -59,6 +59,15 @@ namespace UE::PixelStreaming
 		void ForEachSession(const TFunction<void(TSharedPtr<IPlayerSession>)>& Func);
 
 	private:
+		void AddSession(FPixelStreamingPlayerId PlayerId, TSharedPtr<IPlayerSession> Session);
+
+		rtc::scoped_refptr<webrtc::PeerConnectionInterface> CreatePeerConnection(
+			rtc::scoped_refptr<webrtc::PeerConnectionFactoryInterface> PeerConnectionFactory,
+			webrtc::PeerConnectionInterface::RTCConfiguration PeerConnectionConfig,
+			TSharedPtr<IPlayerSession> Session) const;
+
+		int NumP2PPeers() const;
+		bool GetFirstP2PPeer(FPixelStreamingPlayerId& OutPlayerId) const;
 		TSharedPtr<IPlayerSession> GetPlayerSession(FPixelStreamingPlayerId PlayerId) const;
 
 		mutable FCriticalSection PlayersCS;
diff --git a/Engine/Plugins/Media/PixelStreaming/Source/PixelStreaming/Private/Poller.cpp b/Engine/Plugins/Media/PixelStreaming/Source/PixelStreaming/Private/Poller.cpp
new file mode 100644
--- /dev/null
+++ b/Engine/Plugins/Media/PixelStreaming/Source/PixelStreaming/Private/Poller.cpp
@@ -0,0 +1,147 @@
+// Copyright Epic Games, Inc. All Rights Reserved.
+#include "Poller.h"
+
+namespace UE::PixelStreaming
+{
+	/*
+	* ------------- UE::PixelStreaming::FPoller ---------------------
+	*/
+	FPoller::FPoller()
+		: Runnable()
+		, Thread(FRunnableThread::Create(&Runnable, TEXT("Pixel Streaming Poller")))
+	{
+	}
+
+	FPoller::~FPoller()
+	{
+		Shutdown();
+		Thread->Kill(true);
+	}
+
+	void FPoller::Shutdown()
+	{
+		Runnable.Stop();
+	}
+
+	void FPoller::AddJob(TFunction<bool()> InIsJobDone, TSharedRef<bool, ESPMode::ThreadSafe> bKeepPolling, TFunction<void()> InJobDoneCallback)
+	{
+		checkf(Runnable.IsRunning(), TEXT("Poller runnable must be running to submit jobs to it."));
+		Runnable.AddJob({ InIsJobDone, bKeepPolling, InJobDoneCallback });
+	}
+
+	/*
+	* ------------------- FPollJob ----------------------
+	*/
+	FPoller::FPollJob::FPollJob(TFunction<bool()> InIsJobDone, TSharedRef<bool, ESPMode::ThreadSafe> bInKeepPolling, TFunction<void()> InJobDoneCallback)
+		: IsJobDone(InIsJobDone)
+		, bKeepPolling(bInKeepPolling)
+		, JobDoneCallback(InJobDoneCallback)
+	{
+	}
+
+	FPoller::FPollJob::FPollJob()
+		: FPollJob([]()->bool { return true; }, MakeShared<bool, ESPMode::ThreadSafe>(false), []() {})
+	{
+	}
+
+	/*
+	* ------------- FPollerRunnable ---------------------
+	*/
+	FPoller::FPollerRunnable::FPollerRunnable()
+		: bIsRunning(false)
+		, JobAddedEvent(FPlatformProcess::GetSynchEventFromPool(false))
+	{
+	}
+
+	bool FPoller::FPollerRunnable::IsRunning() const
+	{
+		return bIsRunning;
+	}
+
+	bool FPoller::FPollerRunnable::Init()
+	{
+		return true;
+	}
+
+	void FPoller::FPollerRunnable::AddJob(FPollJob Job)
+	{
+		// Note: this is threadsafe because use of TQueue
+		JobsToAdd.Enqueue(Job);
+		JobAddedEvent->Trigger();
+	}
+
+	uint32 FPoller::FPollerRunnable::Run()
+	{
+		bIsRunning = true;
+
+		while (bIsRunning)
+		{
+			Poll(false);	
+		}
+
+		return 0;
+	}
+
+	void FPoller::FPollerRunnable::Tick()
+	{
+		Poll(true);
+	}
+
+	void FPoller::FPollerRunnable::Poll(bool bIsRunningSingleThreaded)
+	{
+		// Add any new jobs
+		while (!JobsToAdd.IsEmpty())
+		{
+			FPollJob* Job = JobsToAdd.Peek();
+			if (Job)
+			{
+				PollJobs.Add(*Job);
+			}
+			JobsToAdd.Pop();
+		}
+
+		// Iterate backwards so we can remove elements if the job is completed
+		for (int i = PollJobs.Num() - 1; i >= 0; i--)
+		{
+			FPollJob& Job = PollJobs[i];
+			if (Job.bKeepPolling.Get() == false)
+			{
+				PollJobs.RemoveAt(i);
+				continue;
+			}
+
+			bool bJobCompleted = Job.IsJobDone();
+			if (bJobCompleted)
+			{
+				Job.JobDoneCallback();
+				PollJobs.RemoveAt(i);
+			}
+		}
+
+		if(!bIsRunningSingleThreaded)
+		{
+			//No jobs so sleep this thread waiting for a job to be added.
+			if (PollJobs.Num() == 0)
+			{
+					JobAddedEvent->Wait(); 
+			}
+			else
+			{
+				// Free CPU usage up for other threads.
+				FPlatformProcess::Sleep(0.0f);
+			}
+		}
+	}
+
+	void FPoller::FPollerRunnable::Stop()
+	{
+		bIsRunning = false;
+		JobAddedEvent->Trigger();
+	}
+
+	void FPoller::FPollerRunnable::Exit()
+	{
+		bIsRunning = false;
+		JobAddedEvent->Trigger();
+	}
+} // namespace UE::PixelStreaming
\ No newline at end of file
diff --git a/Engine/Plugins/Media/PixelStreaming/Source/PixelStreaming/Private/Poller.h b/Engine/Plugins/Media/PixelStreaming/Source/PixelStreaming/Private/Poller.h
new file mode 100644
--- /dev/null
+++ b/Engine/Plugins/Media/PixelStreaming/Source/PixelStreaming/Private/Poller.h
@@ -0,0 +1,70 @@
+// Copyright Epic Games, Inc. All Rights Reserved.
+#pragma once
+
+#include "Containers/Queue.h"
+#include "Containers/Array.h"
+#include "RHI.h"
+#include "RHIResources.h"
+#include "Templates/SharedPointer.h"
+#include "HAL/Event.h"
+#include "GenericPlatform/GenericPlatformMisc.h"
+#include "HAL/RunnableThread.h"
+#include "HAL/Runnable.h"
+#include "Misc/SingleThreadRunnable.h"
+
+namespace UE::PixelStreaming
+{
+	/*
+	* A thread with a single purpose, to poll GPU fences until they are done.
+	* Polling jobs can be submitted to this thread using its static methods.
+	* Once a polling job is complete a callback is called on the submitted job.
+	*/
+	class PIXELSTREAMING_API FPoller
+	{
+	public:
+		FPoller();
+		void Shutdown();
+		virtual ~FPoller();
+		void AddJob(TFunction<bool()> InIsJobDone, TSharedRef<bool, ESPMode::ThreadSafe> bKeepPolling, TFunction<void()> InJobDoneCallback);
+
+	private:
+		struct FPollJob
+		{
+			FPollJob();
+			FPollJob(TFunction<bool()> InIsJobDone, TSharedRef<bool, ESPMode::ThreadSafe> bInKeepPolling, TFunction<void()> InJobDoneCallback);
+
+			TFunction<bool()> IsJobDone;
+			TSharedRef<bool, ESPMode::ThreadSafe> bKeepPolling;
+			TFunction<void()> JobDoneCallback;
+		};
+
+		class FPollerRunnable : public FRunnable, public FSingleThreadRunnable
+		{
+		public:
+			FPollerRunnable();
+			void AddJob(FPollJob Job);
+			bool IsRunning() const;
+			// Begin FRunnable interface.
+			virtual bool Init() override;
+			virtual uint32 Run() override;
+			virtual void Stop() override;
+			virtual void Exit() override;
+			virtual FSingleThreadRunnable* GetSingleThreadInterface() override { bIsRunning = true; return this; };
+			// End FRunnable interface
+
+			// Begin FSingleThreadRunnable interface.
+			virtual void Tick() override;
+			// End FSingleThreadRunnable interface
+		private:
+			void Poll(bool bIsRunningSingleThreaded);
+			bool bIsRunning;
+			FEvent* JobAddedEvent;
+			TArray<FPollJob> PollJobs;
+			TQueue<FPollJob, EQueueMode::Mpsc> JobsToAdd;
+		};
+
+	private:
+		FPollerRunnable Runnable;
+		FRunnableThread* Thread;
+	};
+} // namespace UE::PixelStreaming
diff --git a/Engine/Plugins/Media/PixelStreaming/Source/PixelStreaming/Private/SessionDescriptionObservers.h b/Engine/Plugins/Media/PixelStreaming/Source/PixelStreaming/Private/SessionDescriptionObservers.h
--- a/Engine/Plugins/Media/PixelStreaming/Source/PixelStreaming/Private/SessionDescriptionObservers.h
+++ b/Engine/Plugins/Media/PixelStreaming/Source/PixelStreaming/Private/SessionDescriptionObservers.h
@@ -11,78 +11,77 @@
 #include <string>
 #include <memory>
 
-namespace UE {
-	namespace PixelStreaming {
-		//////////////////////////////////////////////////////////////////////////
-		// FSetSessionDescriptionObserver
-		// WebRTC requires an implementation of `webrtc::SetSessionDescriptionObserver` interface as a callback
-		// for setting session description, either on receiving remote `offer` (`PeerConnection::SetRemoteDescription`)
-		// of on sending `answer` (`PeerConnection::SetLocalDescription`)
-		class FSetSessionDescriptionObserver : public webrtc::SetSessionDescriptionObserver
+namespace UE::PixelStreaming
+{
+	//////////////////////////////////////////////////////////////////////////
+	// FSetSessionDescriptionObserver
+	// WebRTC requires an implementation of `webrtc::SetSessionDescriptionObserver` interface as a callback
+	// for setting session description, either on receiving remote `offer` (`PeerConnection::SetRemoteDescription`)
+	// of on sending `answer` (`PeerConnection::SetLocalDescription`)
+	class FSetSessionDescriptionObserver : public webrtc::SetSessionDescriptionObserver
+	{
+	public:
+		using FSuccessCallback = TUniqueFunction<void()>;
+		using FFailureCallback = TUniqueFunction<void(const FString&)>;
+
+		static FSetSessionDescriptionObserver* Create(FSuccessCallback&& successCallback, FFailureCallback&& failureCallback)
 		{
-		public:
-			using FSuccessCallback = TUniqueFunction<void()>;
-			using FFailureCallback = TUniqueFunction<void(const FString&)>;
-
-			static FSetSessionDescriptionObserver* Create(FSuccessCallback&& successCallback, FFailureCallback&& failureCallback)
-			{
-				return new rtc::RefCountedObject<FSetSessionDescriptionObserver>(MoveTemp(successCallback), MoveTemp(failureCallback));
-			}
-
-			FSetSessionDescriptionObserver(FSuccessCallback&& successCallback, FFailureCallback&& failureCallback)
-				: SuccessCallback(MoveTemp(successCallback))
-				, FailureCallback(MoveTemp(failureCallback))
-			{
-			}
-
-			// we don't need to do anything on success
-			void OnSuccess() override
-			{
-				SuccessCallback();
-			}
-
-			// errors usually mean incompatibility between our session configuration (often H.264, its profile and level) and
-			// player, malformed SDP or if player doesn't support UnifiedPlan (whatever was used by proxy)
-			void OnFailure(webrtc::RTCError Error) override
-			{
-				FailureCallback(FString(UTF8_TO_TCHAR(Error.message())));
-			}
-
-		private:
-			FSuccessCallback SuccessCallback;
-			FFailureCallback FailureCallback;
-		};
-
-		class FCreateSessionDescriptionObserver : public webrtc::CreateSessionDescriptionObserver
+			return new rtc::RefCountedObject<FSetSessionDescriptionObserver>(MoveTemp(successCallback), MoveTemp(failureCallback));
+		}
+
+		FSetSessionDescriptionObserver(FSuccessCallback&& successCallback, FFailureCallback&& failureCallback)
+			: SuccessCallback(MoveTemp(successCallback))
+			, FailureCallback(MoveTemp(failureCallback))
+		{
+		}
+
+		// we don't need to do anything on success
+		void OnSuccess() override
+		{
+			SuccessCallback();
+		}
+
+		// errors usually mean incompatibility between our session configuration (often H.264, its profile and level) and
+		// player, malformed SDP or if player doesn't support UnifiedPlan (whatever was used by proxy)
+		void OnFailure(webrtc::RTCError Error) override
+		{
+			FailureCallback(FString(UTF8_TO_TCHAR(Error.message())));
+		}
+
+	private:
+		FSuccessCallback SuccessCallback;
+		FFailureCallback FailureCallback;
+	};
+
+	class FCreateSessionDescriptionObserver : public webrtc::CreateSessionDescriptionObserver
+	{
+	public:
+		using FSuccessCallback = TUniqueFunction<void(webrtc::SessionDescriptionInterface*)>;
+		using FFailureCallback = TUniqueFunction<void(const FString&)>;
+
+		static FCreateSessionDescriptionObserver* Create(FSuccessCallback&& successCallback, FFailureCallback&& failureCallback)
+		{
+			return new rtc::RefCountedObject<FCreateSessionDescriptionObserver>(MoveTemp(successCallback), MoveTemp(failureCallback));
+		}
+
+		FCreateSessionDescriptionObserver(FSuccessCallback&& successCallback, FFailureCallback&& failureCallback)
+			: SuccessCallback(MoveTemp(successCallback))
+			, FailureCallback(MoveTemp(failureCallback))
+		{
+		}
+
+		void OnSuccess(webrtc::SessionDescriptionInterface* SDP) override
+		{
+			SuccessCallback(SDP);
+		}
+
+		void OnFailure(webrtc::RTCError Error) override
 		{
-		public:
-			using FSuccessCallback = TUniqueFunction<void(webrtc::SessionDescriptionInterface*)>;
-			using FFailureCallback = TUniqueFunction<void(const FString&)>;
-
-			static FCreateSessionDescriptionObserver* Create(FSuccessCallback&& successCallback, FFailureCallback&& failureCallback)
-			{
-				return new rtc::RefCountedObject<FCreateSessionDescriptionObserver>(MoveTemp(successCallback), MoveTemp(failureCallback));
-			}
-
-			FCreateSessionDescriptionObserver(FSuccessCallback&& successCallback, FFailureCallback&& failureCallback)
-				: SuccessCallback(MoveTemp(successCallback))
-				, FailureCallback(MoveTemp(failureCallback))
-			{
-			}
-
-			void OnSuccess(webrtc::SessionDescriptionInterface* SDP) override
-			{
-				SuccessCallback(SDP);
-			}
-
-			void OnFailure(webrtc::RTCError Error) override
-			{
-				FailureCallback(ANSI_TO_TCHAR(Error.message()));
-			}
-
-		private:
-			FSuccessCallback SuccessCallback;
-			FFailureCallback FailureCallback;
-		};
-	}
-}
+			FailureCallback(ANSI_TO_TCHAR(Error.message()));
+		}
+
+	private:
+		FSuccessCallback SuccessCallback;
+		FFailureCallback FailureCallback;
+	};
+} // namespace UE::PixelStreaming
\ No newline at end of file
