This commit is a large collection (from various authors) of Pixel Streaming 5.1 work

diff --git a/Engine/Plugins/Media/PixelStreaming/Source/PixelStreaming/Private/Settings.cpp b/Engine/Plugins/Media/PixelStreaming/Source/PixelStreaming/Private/Settings.cpp
--- a/Engine/Plugins/Media/PixelStreaming/Source/PixelStreaming/Private/Settings.cpp
+++ b/Engine/Plugins/Media/PixelStreaming/Source/PixelStreaming/Private/Settings.cpp
@@ -1,6 +1,9 @@
 // Copyright Epic Games, Inc. All Rights Reserved.
 
 #include "Settings.h"
+#include "EngineGlobals.h"
+#include "Engine/Engine.h"
+#include "Engine/GameViewportClient.h"
 #include "PixelStreamingPrivate.h"
 #include "Misc/DefaultValueHelper.h"
 #include "Async/Async.h"
@@ -128,7 +131,7 @@ namespace UE::PixelStreaming::Settings
 		TEXT("Framerate for WebRTC encoding. Default: 60"),
 		ECVF_Default);
 
-	//Note: 1 megabit is the maximum allowed in WebRTC for a start bitrate.
+	// Note: 1 megabit is the maximum allowed in WebRTC for a start bitrate.
 	TAutoConsoleVariable<int32> CVarPixelStreamingWebRTCStartBitrate(
 		TEXT("PixelStreaming.WebRTC.StartBitrate"),
 		1000000,
@@ -223,6 +226,12 @@ namespace UE::PixelStreaming::Settings
 		TEXT("If true disables latency tester being triggerable."),
 		ECVF_Default);
 
+	TAutoConsoleVariable<bool> CVarPixelStreamingDecoupleFrameRate(
+		TEXT("PixelStreaming.DecoupleFrameRate"),
+		true,
+		TEXT("If true, Pixel Streaming will send 60fps to WebRTC regardless of actual application framerate"),
+		ECVF_Default);
+
 	TAutoConsoleVariable<FString> CVarPixelStreamingKeyFilter(
 		TEXT("PixelStreaming.KeyFilter"),
 		"",
@@ -235,7 +244,32 @@ namespace UE::PixelStreaming::Settings
 		TEXT("If true browser can send consoleCommand payloads that execute in UE's console."),
 		ECVF_Default);
 
+	TAutoConsoleVariable<bool> CVarPixelStreamingVPXUseCompute(
+		TEXT("PixelStreaming.VPXUseCompute"),
+		false,
+		TEXT("If true a compute shader will be used to process RGB -> I420. Otherwise libyuv will be used on the CPU."),
+		ECVF_Default);
+
+	TAutoConsoleVariable<FString> CVarPixelStreamingVideoTracks(
+		TEXT("PixelStreaming.VideoTracks"),
+		TEXT("Backbuffer"),
+		TEXT("Comma separated list of video track types to create when a peer joins. Default: Backbuffer"),
+		ECVF_Default);
+
 	TArray<FKey> FilteredKeys;
+	TArray<FName> ActiveTextureSourceTypes;
+
+	void OnActiveTextureSourcesChanged(IConsoleVariable* Var)
+	{
+		FString CommaList = Var->GetString();
+		TArray<FString> StringArray;
+		CommaList.ParseIntoArray(StringArray, TEXT(","), true);
+		ActiveTextureSourceTypes.Empty();
+		for (auto&& TextureSourceTypeStr : StringArray)
+		{
+			ActiveTextureSourceTypes.Add(FName(TextureSourceTypeStr));
+		}
+	}
 
 	void OnFilteredKeysChanged(IConsoleVariable* Var)
 	{
@@ -249,6 +283,21 @@ namespace UE::PixelStreaming::Settings
 		}
 	}
 
+	void OnHudStatsToggled(IConsoleVariable* Var)
+	{
+		bool bHudStatsEnabled = Var->GetBool();
+
+		for (const FWorldContext& WorldContext : GEngine->GetWorldContexts())
+		{
+			if (WorldContext.WorldType == EWorldType::Game || WorldContext.WorldType == EWorldType::PIE)
+			{
+				UWorld* World = WorldContext.World();
+				UGameViewportClient* ViewportClient = World->GetGameViewport();
+				GEngine->SetEngineStat(World, ViewportClient, TEXT("PixelStreaming"), bHudStatsEnabled);
+			}
+		}
+	}
+
 	void OnKeyframeIntervalChanged(IConsoleVariable* Var)
 	{
 		AsyncTask(ENamedThreads::GameThread, [Var]() {
@@ -341,6 +390,17 @@ namespace UE::PixelStreaming::Settings
 			return AVEncoder::FVideoEncoder::H264Profile::BASELINE;
 		return Iter->second;
 	}
+
+	const TArray<FName>& GetActiveTextureSourceTypes()
+	{
+		return ActiveTextureSourceTypes;
+	}
+
+	void SetActiveTextureSourceTypes(TArray<FName> InActiveTextureSourceTypes)
+	{
+		ActiveTextureSourceTypes = InActiveTextureSourceTypes;
+	}
+
 	// End utility functions etc.
 
 	FSimulcastParameters SimulcastParameters;
@@ -355,7 +415,7 @@ namespace UE::PixelStreaming::Settings
 		// If no simulcast parameters are passed use some default values
 		if (!bPassedSimulcastParams)
 		{
-			//StringOptions = FString(TEXT("1.0,5000000,20000000,2.0,1000000,5000000,4.0,50000,1000000"));
+			// StringOptions = FString(TEXT("1.0,5000000,20000000,2.0,1000000,5000000,4.0,50000,1000000"));
 			StringOptions = FString(TEXT("1.0,5000000,100000000,2.0,1000000,5000000"));
 		}
 
@@ -382,8 +442,8 @@ namespace UE::PixelStreaming::Settings
 	}
 
 	/*
-	* Selected Codec.
-	*/
+	 * Selected Codec.
+	 */
 	ECodec GetSelectedCodec()
 	{
 		const FString CodecStr = CVarPixelStreamingEncoderCodec.GetValueOnAnyThread();
@@ -412,82 +472,101 @@ namespace UE::PixelStreaming::Settings
 	}
 
 	/*
-	* Stats logger - as turned on/off by CVarPixelStreamingLogStats
-	*/
+	 * Stats logger - as turned on/off by CVarPixelStreamingLogStats
+	 */
 	void ConsumeStat(FPixelStreamingPlayerId PlayerId, FName StatName, float StatValue)
 	{
 		UE_LOG(LogPixelStreaming, Log, TEXT("[%s](%s) = %f"), *PlayerId, *StatName.ToString(), StatValue);
 	}
 
+	FDelegateHandle LogStatsHandle;
+
 	void OnLogStatsChanged(IConsoleVariable* Var)
 	{
 		bool bLogStats = Var->GetBool();
+
 		UPixelStreamingDelegates* Delegates = UPixelStreamingDelegates::GetPixelStreamingDelegates();
-		if (bLogStats && Delegates)
+		if (!Delegates)
 		{
-			Delegates->OnStatChangedNative.AddStatic(&UE::PixelStreaming::Settings::ConsumeStat);
+			return;
+		}
+
+		if (bLogStats)
+		{
+			LogStatsHandle = Delegates->OnStatChangedNative.AddStatic(&ConsumeStat);
+		}
+		else
+		{
+			Delegates->OnStatChangedNative.Remove(LogStatsHandle);
 		}
 	}
 
 	/*
-	* Settings parsing and initialization.
-	*/
+	 * Settings parsing and initialization.
+	 */
 
 	// Some settings need to be set after streamer is initialized
 	void OnStreamerReady(IPixelStreamingModule& Module)
 	{
-		UE::PixelStreaming::Settings::CVarPixelStreamingLogStats.AsVariable()->SetOnChangedCallback(FConsoleVariableDelegate::CreateStatic(&UE::PixelStreaming::Settings::OnLogStatsChanged));
-		CommandLineParseOption(TEXT("PixelStreamingLogStats"), UE::PixelStreaming::Settings::CVarPixelStreamingLogStats);
+		CVarPixelStreamingLogStats.AsVariable()->SetOnChangedCallback(FConsoleVariableDelegate::CreateStatic(&OnLogStatsChanged));
+		CommandLineParseOption(TEXT("PixelStreamingLogStats"), CVarPixelStreamingLogStats);
 	}
 
 	void InitialiseSettings()
 	{
 		UE_LOG(LogPixelStreaming, Log, TEXT("Initialising Pixel Streaming settings."));
 
-		UE::PixelStreaming::Settings::CVarPixelStreamingKeyFilter.AsVariable()->SetOnChangedCallback(FConsoleVariableDelegate::CreateStatic(&UE::PixelStreaming::Settings::OnFilteredKeysChanged));
-		UE::PixelStreaming::Settings::CVarPixelStreamingEncoderKeyframeInterval.AsVariable()->SetOnChangedCallback(FConsoleVariableDelegate::CreateStatic(&UE::PixelStreaming::Settings::OnKeyframeIntervalChanged));
+		CVarPixelStreamingOnScreenStats.AsVariable()->SetOnChangedCallback(FConsoleVariableDelegate::CreateStatic(&OnHudStatsToggled));
+		CVarPixelStreamingKeyFilter.AsVariable()->SetOnChangedCallback(FConsoleVariableDelegate::CreateStatic(&OnFilteredKeysChanged));
+		CVarPixelStreamingVideoTracks.AsVariable()->SetOnChangedCallback(FConsoleVariableDelegate::CreateStatic(&OnActiveTextureSourcesChanged));
+		OnActiveTextureSourcesChanged(CVarPixelStreamingVideoTracks.AsVariable());
+
+		CVarPixelStreamingEncoderKeyframeInterval.AsVariable()->SetOnChangedCallback(FConsoleVariableDelegate::CreateStatic(&OnKeyframeIntervalChanged));
 
 		// Values parse from commands line
-		CommandLineParseValue(TEXT("PixelStreamingEncoderKeyframeInterval="), UE::PixelStreaming::Settings::CVarPixelStreamingEncoderKeyframeInterval);
-		CommandLineParseValue(TEXT("PixelStreamingEncoderTargetBitrate="), UE::PixelStreaming::Settings::CVarPixelStreamingEncoderTargetBitrate);
-		CommandLineParseValue(TEXT("PixelStreamingEncoderMaxBitrate="), UE::PixelStreaming::Settings::CVarPixelStreamingEncoderMaxBitrate);
-		CommandLineParseValue(TEXT("PixelStreamingEncoderMinQP="), UE::PixelStreaming::Settings::CVarPixelStreamingEncoderMinQP);
-		CommandLineParseValue(TEXT("PixelStreamingEncoderMaxQP="), UE::PixelStreaming::Settings::CVarPixelStreamingEncoderMaxQP);
-		CommandLineParseValue(TEXT("PixelStreamingEncoderRateControl="), UE::PixelStreaming::Settings::CVarPixelStreamingEncoderRateControl);
-		CommandLineParseValue(TEXT("PixelStreamingEncoderMultipass="), UE::PixelStreaming::Settings::CVarPixelStreamingEncoderMultipass);
-		CommandLineParseValue(TEXT("PixelStreamingEncoderCodec="), UE::PixelStreaming::Settings::CVarPixelStreamingEncoderCodec);
-		CommandLineParseValue(TEXT("PixelStreamingH264Profile="), UE::PixelStreaming::Settings::CVarPixelStreamingH264Profile);
-		CommandLineParseValue(TEXT("PixelStreamingDegradationPreference="), UE::PixelStreaming::Settings::CVarPixelStreamingDegradationPreference);
-		CommandLineParseValue(TEXT("PixelStreamingWebRTCDegradationPreference="), UE::PixelStreaming::Settings::CVarPixelStreamingDegradationPreference);
-		CommandLineParseValue(TEXT("PixelStreamingWebRTCFps="), UE::PixelStreaming::Settings::CVarPixelStreamingWebRTCFps);
-		CommandLineParseValue(TEXT("PixelStreamingWebRTCStartBitrate="), UE::PixelStreaming::Settings::CVarPixelStreamingWebRTCStartBitrate);
-		CommandLineParseValue(TEXT("PixelStreamingWebRTCMinBitrate="), UE::PixelStreaming::Settings::CVarPixelStreamingWebRTCMinBitrate);
-		CommandLineParseValue(TEXT("PixelStreamingWebRTCMaxBitrate="), UE::PixelStreaming::Settings::CVarPixelStreamingWebRTCMaxBitrate);
-		CommandLineParseValue(TEXT("PixelStreamingWebRTCLowQpThreshold="), UE::PixelStreaming::Settings::CVarPixelStreamingWebRTCLowQpThreshold);
-		CommandLineParseValue(TEXT("PixelStreamingWebRTCHighQpThreshold="), UE::PixelStreaming::Settings::CVarPixelStreamingWebRTCHighQpThreshold);
-		CommandLineParseValue(TEXT("PixelStreamingFreezeFrameQuality"), UE::PixelStreaming::Settings::CVarPixelStreamingFreezeFrameQuality);
-		CommandLineParseValue(TEXT("PixelStreamingKeyFilter="), UE::PixelStreaming::Settings::CVarPixelStreamingKeyFilter);
-		CommandLineParseValue(TEXT("PixelStreamingFrameScale="), UE::PixelStreaming::Settings::CVarPixelStreamingFrameScale);
+		CommandLineParseValue(TEXT("PixelStreamingEncoderKeyframeInterval="), CVarPixelStreamingEncoderKeyframeInterval);
+		CommandLineParseValue(TEXT("PixelStreamingEncoderTargetBitrate="), CVarPixelStreamingEncoderTargetBitrate);
+		CommandLineParseValue(TEXT("PixelStreamingEncoderMaxBitrate="), CVarPixelStreamingEncoderMaxBitrate);
+		CommandLineParseValue(TEXT("PixelStreamingEncoderMinQP="), CVarPixelStreamingEncoderMinQP);
+		CommandLineParseValue(TEXT("PixelStreamingEncoderMaxQP="), CVarPixelStreamingEncoderMaxQP);
+		CommandLineParseValue(TEXT("PixelStreamingEncoderRateControl="), CVarPixelStreamingEncoderRateControl);
+		CommandLineParseValue(TEXT("PixelStreamingEncoderMultipass="), CVarPixelStreamingEncoderMultipass);
+		CommandLineParseValue(TEXT("PixelStreamingEncoderCodec="), CVarPixelStreamingEncoderCodec);
+		CommandLineParseValue(TEXT("PixelStreamingH264Profile="), CVarPixelStreamingH264Profile);
+		CommandLineParseValue(TEXT("PixelStreamingDegradationPreference="), CVarPixelStreamingDegradationPreference);
+		CommandLineParseValue(TEXT("PixelStreamingWebRTCDegradationPreference="), CVarPixelStreamingDegradationPreference);
+		CommandLineParseValue(TEXT("PixelStreamingWebRTCFps="), CVarPixelStreamingWebRTCFps);
+		CommandLineParseValue(TEXT("PixelStreamingWebRTCStartBitrate="), CVarPixelStreamingWebRTCStartBitrate);
+		CommandLineParseValue(TEXT("PixelStreamingWebRTCMinBitrate="), CVarPixelStreamingWebRTCMinBitrate);
+		CommandLineParseValue(TEXT("PixelStreamingWebRTCMaxBitrate="), CVarPixelStreamingWebRTCMaxBitrate);
+		CommandLineParseValue(TEXT("PixelStreamingWebRTCLowQpThreshold="), CVarPixelStreamingWebRTCLowQpThreshold);
+		CommandLineParseValue(TEXT("PixelStreamingWebRTCHighQpThreshold="), CVarPixelStreamingWebRTCHighQpThreshold);
+		CommandLineParseValue(TEXT("PixelStreamingFreezeFrameQuality"), CVarPixelStreamingFreezeFrameQuality);
+		CommandLineParseValue(TEXT("PixelStreamingKeyFilter="), CVarPixelStreamingKeyFilter);
+		CommandLineParseValue(TEXT("PixelStreamingVideoTracks="), CVarPixelStreamingVideoTracks);
+		CommandLineParseValue(TEXT("PixelStreamingFrameScale="), CVarPixelStreamingFrameScale);
 
 		// Options parse (if these exist they are set to true)
-		CommandLineParseOption(TEXT("AllowPixelStreamingCommands"), UE::PixelStreaming::Settings::CVarPixelStreamingAllowConsoleCommands);
-		CommandLineParseOption(TEXT("PixelStreamingOnScreenStats"), UE::PixelStreaming::Settings::CVarPixelStreamingOnScreenStats);
-		CommandLineParseOption(TEXT("PixelStreamingHudStats"), UE::PixelStreaming::Settings::CVarPixelStreamingOnScreenStats);
-
-		CommandLineParseOption(TEXT("PixelStreamingDebugDumpFrame"), UE::PixelStreaming::Settings::CVarPixelStreamingDebugDumpFrame);
-		CommandLineParseOption(TEXT("PixelStreamingEnableFillerData"), UE::PixelStreaming::Settings::CVarPixelStreamingEnableFillerData);
-		CommandLineParseOption(TEXT("PixelStreamingWebRTCDisableStats"), UE::PixelStreaming::Settings::CVarPixelStreamingWebRTCDisableStats);
-		CommandLineParseOption(TEXT("PixelStreamingWebRTCDisableReceiveAudio"), UE::PixelStreaming::Settings::CVarPixelStreamingWebRTCDisableReceiveAudio);
-		CommandLineParseOption(TEXT("PixelStreamingWebRTCDisableTransmitAudio"), UE::PixelStreaming::Settings::CVarPixelStreamingWebRTCDisableTransmitAudio);
-		CommandLineParseOption(TEXT("PixelStreamingWebRTCDisableAudioSync"), UE::PixelStreaming::Settings::CVarPixelStreamingWebRTCDisableAudioSync);
-		CommandLineParseOption(TEXT("PixelStreamingSendPlayerIdAsInteger"), UE::PixelStreaming::Settings::CVarSendPlayerIdAsInteger);
-		CommandLineParseOption(TEXT("PixelStreamingWebRTCUseLegacyAudioDevice"), UE::PixelStreaming::Settings::CVarPixelStreamingWebRTCUseLegacyAudioDevice);
-		CommandLineParseOption(TEXT("PixelStreamingDisableLatencyTester"), UE::PixelStreaming::Settings::CVarPixelStreamingDisableLatencyTester);
+		CommandLineParseOption(TEXT("AllowPixelStreamingCommands"), CVarPixelStreamingAllowConsoleCommands);
+		CommandLineParseOption(TEXT("PixelStreamingOnScreenStats"), CVarPixelStreamingOnScreenStats);
+		CommandLineParseOption(TEXT("PixelStreamingHudStats"), CVarPixelStreamingOnScreenStats);
+
+		CommandLineParseOption(TEXT("PixelStreamingDebugDumpFrame"), CVarPixelStreamingDebugDumpFrame);
+		CommandLineParseOption(TEXT("PixelStreamingEnableFillerData"), CVarPixelStreamingEnableFillerData);
+		CommandLineParseOption(TEXT("PixelStreamingWebRTCDisableStats"), CVarPixelStreamingWebRTCDisableStats);
+		CommandLineParseOption(TEXT("PixelStreamingWebRTCDisableReceiveAudio"), CVarPixelStreamingWebRTCDisableReceiveAudio);
+		CommandLineParseOption(TEXT("PixelStreamingWebRTCDisableTransmitAudio"), CVarPixelStreamingWebRTCDisableTransmitAudio);
+		CommandLineParseOption(TEXT("PixelStreamingWebRTCDisableAudioSync"), CVarPixelStreamingWebRTCDisableAudioSync);
+		CommandLineParseOption(TEXT("PixelStreamingSendPlayerIdAsInteger"), CVarSendPlayerIdAsInteger);
+		CommandLineParseOption(TEXT("PixelStreamingWebRTCUseLegacyAudioDevice"), CVarPixelStreamingWebRTCUseLegacyAudioDevice);
+		CommandLineParseOption(TEXT("PixelStreamingDisableLatencyTester"), CVarPixelStreamingDisableLatencyTester);
+		CommandLineParseOption(TEXT("PixelStreamingDecoupleFrameRate"), CVarPixelStreamingDecoupleFrameRate);
+		CommandLineParseOption(TEXT("PixelStreamingVPXUseCompute"), CVarPixelStreamingVPXUseCompute);
 
 		ReadSimulcastParameters();
 
 		IPixelStreamingModule& Module = IPixelStreamingModule::Get();
-		Module.OnReady().AddStatic(&UE::PixelStreaming::Settings::OnStreamerReady);
+		Module.OnReady().AddStatic(&OnStreamerReady);
 
 		IConsoleManager::Get().RegisterConsoleCommand(
 			TEXT("PixelStreaming.StartStreaming"),
@@ -503,4 +582,5 @@ namespace UE::PixelStreaming::Settings
 				IPixelStreamingModule::Get().StopStreaming();
 			}));
 	}
+
 } // namespace UE::PixelStreaming::Settings
diff --git a/Engine/Plugins/Media/PixelStreaming/Source/PixelStreaming/Private/Settings.h b/Engine/Plugins/Media/PixelStreaming/Source/PixelStreaming/Private/Settings.h
--- a/Engine/Plugins/Media/PixelStreaming/Source/PixelStreaming/Private/Settings.h
+++ b/Engine/Plugins/Media/PixelStreaming/Source/PixelStreaming/Private/Settings.h
@@ -10,118 +10,113 @@
 #include "VideoEncoder.h"
 #include "WebRTCIncludes.h"
 
-namespace UE
+namespace UE::PixelStreaming::Settings
 {
-	namespace PixelStreaming
+	extern void InitialiseSettings();
+
+	// Begin Encoder CVars
+	extern TAutoConsoleVariable<int32> CVarPixelStreamingEncoderTargetBitrate;
+	extern TAutoConsoleVariable<int32> CVarPixelStreamingEncoderMaxBitrate;
+	extern TAutoConsoleVariable<bool> CVarPixelStreamingDebugDumpFrame;
+	extern TAutoConsoleVariable<int32> CVarPixelStreamingEncoderMinQP;
+	extern TAutoConsoleVariable<int32> CVarPixelStreamingEncoderMaxQP;
+	extern TAutoConsoleVariable<FString> CVarPixelStreamingEncoderRateControl;
+	extern TAutoConsoleVariable<bool> CVarPixelStreamingEnableFillerData;
+	extern TAutoConsoleVariable<FString> CVarPixelStreamingEncoderMultipass;
+	extern TAutoConsoleVariable<FString> CVarPixelStreamingH264Profile;
+	extern TAutoConsoleVariable<int32> CVarPixelStreamingEncoderKeyframeInterval;
+	extern TAutoConsoleVariable<FString> CVarPixelStreamingEncoderCodec;
+	extern TAutoConsoleVariable<FString> CVarPixelStreamingVideoTracks;
+	// End Encoder CVars
+
+	// Begin WebRTC CVars
+	extern TAutoConsoleVariable<FString> CVarPixelStreamingDegradationPreference;
+	extern TAutoConsoleVariable<int32> CVarPixelStreamingWebRTCFps;
+	extern TAutoConsoleVariable<int32> CVarPixelStreamingWebRTCStartBitrate;
+	extern TAutoConsoleVariable<int32> CVarPixelStreamingWebRTCMinBitrate;
+	extern TAutoConsoleVariable<int32> CVarPixelStreamingWebRTCMaxBitrate;
+	extern TAutoConsoleVariable<int> CVarPixelStreamingWebRTCLowQpThreshold;
+	extern TAutoConsoleVariable<int> CVarPixelStreamingWebRTCHighQpThreshold;
+	extern TAutoConsoleVariable<bool> CVarPixelStreamingWebRTCDisableReceiveAudio;
+	extern TAutoConsoleVariable<bool> CVarPixelStreamingWebRTCDisableTransmitAudio;
+	extern TAutoConsoleVariable<bool> CVarPixelStreamingWebRTCDisableAudioSync;
+	extern TAutoConsoleVariable<bool> CVarPixelStreamingWebRTCUseLegacyAudioDevice;
+	extern TAutoConsoleVariable<bool> CVarPixelStreamingWebRTCDisableStats;
+	// End WebRTC CVars
+
+	// Begin Pixel Streaming Plugin CVars
+	extern TAutoConsoleVariable<bool> CVarPixelStreamingAllowConsoleCommands;
+	extern TAutoConsoleVariable<bool> CVarPixelStreamingOnScreenStats;
+	extern TAutoConsoleVariable<bool> CVarPixelStreamingLogStats;
+	extern TAutoConsoleVariable<int32> CVarPixelStreamingFreezeFrameQuality;
+	extern TAutoConsoleVariable<bool> CVarSendPlayerIdAsInteger;
+	extern TAutoConsoleVariable<bool> CVarPixelStreamingDisableLatencyTester;
+	extern TAutoConsoleVariable<bool> CVarPixelStreamingDecoupleFrameRate;
+	extern TAutoConsoleVariable<bool> CVarPixelStreamingVPXUseCompute;
+
+	extern TArray<FKey> FilteredKeys;
+	// Ends Pixel Streaming Plugin CVars
+
+	// Begin TextureSource CVars
+	extern TAutoConsoleVariable<float> CVarPixelStreamingFrameScale;
+	// End TextureSource CVars
+
+	enum ECodec
 	{
-		// Console variables (CVars)
-		namespace Settings
+		H264,
+		VP8,
+		VP9
+	};
+
+	// Begin utility functions etc.
+	bool IsCodecVPX();
+	ECodec GetSelectedCodec();
+	AVEncoder::FVideoEncoder::RateControlMode GetRateControlCVar();
+	AVEncoder::FVideoEncoder::MultipassMode GetMultipassCVar();
+	webrtc::DegradationPreference GetDegradationPreference();
+	AVEncoder::FVideoEncoder::H264Profile GetH264Profile();
+	const TArray<FName>& GetActiveTextureSourceTypes();
+	void SetActiveTextureSourceTypes(TArray<FName> InActiveTextureSourceTypes);
+	// End utility functions etc.
+
+	struct FSimulcastParameters
+	{
+		struct FLayer
 		{
-			extern void InitialiseSettings();
-
-			// Begin Encoder CVars
-			extern TAutoConsoleVariable<int32> CVarPixelStreamingEncoderTargetBitrate;
-			extern TAutoConsoleVariable<int32> CVarPixelStreamingEncoderMaxBitrate;
-			extern TAutoConsoleVariable<bool> CVarPixelStreamingDebugDumpFrame;
-			extern TAutoConsoleVariable<int32> CVarPixelStreamingEncoderMinQP;
-			extern TAutoConsoleVariable<int32> CVarPixelStreamingEncoderMaxQP;
-			extern TAutoConsoleVariable<FString> CVarPixelStreamingEncoderRateControl;
-			extern TAutoConsoleVariable<bool> CVarPixelStreamingEnableFillerData;
-			extern TAutoConsoleVariable<FString> CVarPixelStreamingEncoderMultipass;
-			extern TAutoConsoleVariable<FString> CVarPixelStreamingH264Profile;
-			extern TAutoConsoleVariable<int32> CVarPixelStreamingEncoderKeyframeInterval;
-			extern TAutoConsoleVariable<FString> CVarPixelStreamingEncoderCodec;
-			// End Encoder CVars
-
-			// Begin WebRTC CVars
-			extern TAutoConsoleVariable<FString> CVarPixelStreamingDegradationPreference;
-			extern TAutoConsoleVariable<int32> CVarPixelStreamingWebRTCFps;
-			extern TAutoConsoleVariable<int32> CVarPixelStreamingWebRTCStartBitrate;
-			extern TAutoConsoleVariable<int32> CVarPixelStreamingWebRTCMinBitrate;
-			extern TAutoConsoleVariable<int32> CVarPixelStreamingWebRTCMaxBitrate;
-			extern TAutoConsoleVariable<int> CVarPixelStreamingWebRTCLowQpThreshold;
-			extern TAutoConsoleVariable<int> CVarPixelStreamingWebRTCHighQpThreshold;
-			extern TAutoConsoleVariable<bool> CVarPixelStreamingWebRTCDisableReceiveAudio;
-			extern TAutoConsoleVariable<bool> CVarPixelStreamingWebRTCDisableTransmitAudio;
-			extern TAutoConsoleVariable<bool> CVarPixelStreamingWebRTCDisableAudioSync;
-			extern TAutoConsoleVariable<bool> CVarPixelStreamingWebRTCUseLegacyAudioDevice;
-			extern TAutoConsoleVariable<bool> CVarPixelStreamingWebRTCDisableStats;
-			// End WebRTC CVars
-
-			// Begin Pixel Streaming Plugin CVars
-			extern TAutoConsoleVariable<bool> CVarPixelStreamingAllowConsoleCommands;
-			extern TAutoConsoleVariable<bool> CVarPixelStreamingOnScreenStats;
-			extern TAutoConsoleVariable<bool> CVarPixelStreamingLogStats;
-
-			extern TAutoConsoleVariable<int32> CVarPixelStreamingFreezeFrameQuality;
-			extern TAutoConsoleVariable<bool> CVarSendPlayerIdAsInteger;
-			extern TAutoConsoleVariable<bool> CVarPixelStreamingDisableLatencyTester;
-			extern TArray<FKey> FilteredKeys;
-			// Ends Pixel Streaming Plugin CVars
-
-			// Begin TextureSource CVars
-			extern TAutoConsoleVariable<float> CVarPixelStreamingFrameScale;
-			// End TextureSource CVars
-
-			enum ECodec
-			{
-				H264,
-				VP8,
-				VP9
-			};
-
-			// Begin utility functions etc.
-			bool IsCodecVPX();
-			ECodec GetSelectedCodec();
-			AVEncoder::FVideoEncoder::RateControlMode GetRateControlCVar();
-			AVEncoder::FVideoEncoder::MultipassMode GetMultipassCVar();
-			webrtc::DegradationPreference GetDegradationPreference();
-			AVEncoder::FVideoEncoder::H264Profile GetH264Profile();
-			// End utility functions etc.
-
-			struct FSimulcastParameters
-			{
-				struct FLayer
-				{
-					float Scaling;
-					int MinBitrate;
-					int MaxBitrate;
-				};
-
-				TArray<FLayer> Layers;
-			};
-
-			extern FSimulcastParameters SimulcastParameters;
-
-			// Begin Command line args
-
-			inline bool IsPixelStreamingHideCursor()
-			{
-				return FParse::Param(FCommandLine::Get(), TEXT("PixelStreamingHideCursor"));
-			}
-
-			inline bool GetSignallingServerIP(FString& OutSignallingServerIP)
-			{
-				return FParse::Value(FCommandLine::Get(), TEXT("PixelStreamingIP="), OutSignallingServerIP);
-			}
-
-			inline bool GetSignallingServerPort(uint16& OutSignallingServerPort)
-			{
-				return FParse::Value(FCommandLine::Get(), TEXT("PixelStreamingPort="), OutSignallingServerPort);
-			}
-
-			inline bool GetControlScheme(FString& OutControlScheme)
-			{
-				return FParse::Value(FCommandLine::Get(), TEXT("PixelStreamingControlScheme="), OutControlScheme);
-			}
-
-			inline bool GetFastPan(float& OutFastPan)
-			{
-				return FParse::Value(FCommandLine::Get(), TEXT("PixelStreamingFastPan="), OutFastPan);
-			}
-
-			// End Command line args
-
-		} // namespace Settings
-	}	  // namespace PixelStreaming
-} // namespace UE
+			float Scaling;
+			int MinBitrate;
+			int MaxBitrate;
+		};
+
+		TArray<FLayer> Layers;
+	};
+
+	extern FSimulcastParameters SimulcastParameters;
+
+	// Begin Command line args
+	inline bool IsPixelStreamingHideCursor()
+	{
+		return FParse::Param(FCommandLine::Get(), TEXT("PixelStreamingHideCursor"));
+	}
+
+	inline bool GetSignallingServerIP(FString& OutSignallingServerIP)
+	{
+		return FParse::Value(FCommandLine::Get(), TEXT("PixelStreamingIP="), OutSignallingServerIP);
+	}
+
+	inline bool GetSignallingServerPort(uint16& OutSignallingServerPort)
+	{
+		return FParse::Value(FCommandLine::Get(), TEXT("PixelStreamingPort="), OutSignallingServerPort);
+	}
+
+	inline bool GetControlScheme(FString& OutControlScheme)
+	{
+		return FParse::Value(FCommandLine::Get(), TEXT("PixelStreamingControlScheme="), OutControlScheme);
+	}
+
+	inline bool GetFastPan(float& OutFastPan)
+	{
+		return FParse::Value(FCommandLine::Get(), TEXT("PixelStreamingFastPan="), OutFastPan);
+	}
+	// End Command line args
+} // namespace UE::PixelStreaming::Settings
diff --git a/Engine/Plugins/Media/PixelStreaming/Source/PixelStreaming/Private/SignallingServerConnection.cpp b/Engine/Plugins/Media/PixelStreaming/Source/PixelStreaming/Private/SignallingServerConnection.cpp
--- a/Engine/Plugins/Media/PixelStreaming/Source/PixelStreaming/Private/SignallingServerConnection.cpp
+++ b/Engine/Plugins/Media/PixelStreaming/Source/PixelStreaming/Private/SignallingServerConnection.cpp
@@ -10,14 +10,17 @@
 #include "TimerManager.h"
 #include "PixelStreamingDelegates.h"
 #include "PixelStreamingProtocolDefs.h"
+#include "Utils.h"
 
 DECLARE_LOG_CATEGORY_EXTERN(LogPixelStreamingSS, Log, VeryVerbose);
 DEFINE_LOG_CATEGORY(LogPixelStreamingSS);
 
 namespace UE::PixelStreaming
 {
-	FSignallingServerConnection::FSignallingServerConnection(FSignallingServerConnectionObserver& InObserver, FString InStreamerId)
-		: Observer(InObserver), StreamerId(InStreamerId)
+	FSignallingServerConnection::FSignallingServerConnection(const FWebSocketFactory& InWebSocketFactory, FSignallingServerConnectionObserver& InObserver, FString InStreamerId)
+		: WebSocketFactory(InWebSocketFactory)
+		, Observer(InObserver)
+		, StreamerId(InStreamerId)
 	{
 		RegisterHandler("identify", [this](FJsonObjectPtr JsonMsg) { OnIdRequested(); });
 		RegisterHandler("config", [this](FJsonObjectPtr JsonMsg) { OnConfig(JsonMsg); });
@@ -46,52 +49,53 @@ namespace UE::PixelStreaming
 		Json->SetStringField(TEXT("id"), StreamerId);
 		FString Msg = ToString(Json, false);
 		UE_LOG(LogPixelStreamingSS, Verbose, TEXT("-> SS: endpointId\n%s"), *Msg);
-		WS->Send(Msg);
+		SendMessage(Msg);
 	}
 
 	void FSignallingServerConnection::Connect(const FString& Url)
 	{
 		// Already have a websocket connection, no need to make another one
-		if (WS)
+		if (WebSocket)
 		{
 			return;
 		}
 
-		WS = FWebSocketsModule::Get().CreateWebSocket(Url, TEXT(""));
+		WebSocket = WebSocketFactory(Url);
+		checkf(WebSocket, TEXT("Web Socket Factory failed to return a valid Web Socket."));
 
-		OnConnectedHandle = WS->OnConnected().AddLambda([this]() { OnConnected(); });
-		OnConnectionErrorHandle = WS->OnConnectionError().AddLambda([this](const FString& Error) { OnConnectionError(Error); });
-		OnClosedHandle = WS->OnClosed().AddLambda([this](int32 StatusCode, const FString& Reason, bool bWasClean) { OnClosed(StatusCode, Reason, bWasClean); });
-		OnMessageHandle = WS->OnMessage().AddLambda([this](const FString& Msg) { OnMessage(Msg); });
+		OnConnectedHandle = WebSocket->OnConnected().AddLambda([this]() { OnConnected(); });
+		OnConnectionErrorHandle = WebSocket->OnConnectionError().AddLambda([this](const FString& Error) { OnConnectionError(Error); });
+		OnClosedHandle = WebSocket->OnClosed().AddLambda([this](int32 StatusCode, const FString& Reason, bool bWasClean) { OnClosed(StatusCode, Reason, bWasClean); });
+		OnMessageHandle = WebSocket->OnMessage().AddLambda([this](const FString& Msg) { OnMessage(Msg); });
 
 		UE_LOG(LogPixelStreamingSS, Log, TEXT("Connecting to SS %s"), *Url);
-		WS->Connect();
+		WebSocket->Connect();
 	}
 
 	void FSignallingServerConnection::Disconnect()
 	{
-		if (!WS)
+		if (!WebSocket)
 		{
 			return;
 		}
 
 		if (!IsEngineExitRequested())
 		{
-			GWorld->GetTimerManager().ClearTimer(TimerHandle_KeepAlive);
+			StopKeepAliveTimer();
 		}
 
-		WS->OnConnected().Remove(OnConnectedHandle);
-		WS->OnConnectionError().Remove(OnConnectionErrorHandle);
-		WS->OnClosed().Remove(OnClosedHandle);
-		WS->OnMessage().Remove(OnMessageHandle);
+		WebSocket->OnConnected().Remove(OnConnectedHandle);
+		WebSocket->OnConnectionError().Remove(OnConnectionErrorHandle);
+		WebSocket->OnClosed().Remove(OnClosedHandle);
+		WebSocket->OnMessage().Remove(OnMessageHandle);
 
-		WS->Close();
-		WS = nullptr;
+		WebSocket->Close();
+		WebSocket = nullptr;
 	}
 
 	bool FSignallingServerConnection::IsConnected() const
 	{
-		return WS != nullptr && WS.IsValid() && WS->IsConnected();
+		return WebSocket != nullptr && WebSocket.IsValid() && WebSocket->IsConnected();
 	}
 
 	void FSignallingServerConnection::SendOffer(FPixelStreamingPlayerId PlayerId, const webrtc::SessionDescriptionInterface& SDP)
@@ -107,7 +111,7 @@ namespace UE::PixelStreaming
 
 		UE_LOG(LogPixelStreamingSS, Verbose, TEXT("-> SS: offer\n%s"), *SdpStr);
 
-		WS->Send(ToString(OfferJson, false));
+		SendMessage(ToString(OfferJson, false));
 	}
 
 	void FSignallingServerConnection::SendAnswer(FPixelStreamingPlayerId PlayerId, const webrtc::SessionDescriptionInterface& SDP)
@@ -123,7 +127,7 @@ namespace UE::PixelStreaming
 
 		UE_LOG(LogPixelStreamingSS, Verbose, TEXT("-> SS: answer\n%s"), *SdpStr);
 
-		WS->Send(ToString(AnswerJson, false));
+		SendMessage(ToString(AnswerJson, false));
 	}
 
 	void FSignallingServerConnection::SetPlayerIdJson(FJsonObjectPtr& JsonObject, FPixelStreamingPlayerId PlayerId)
@@ -159,6 +163,25 @@ namespace UE::PixelStreaming
 		return false;
 	}
 
+	void FSignallingServerConnection::StartKeepAliveTimer()
+	{
+		// GWorld dereferencing needs to happen on the game thread
+		// we dont need to wait since its just setting the timer
+		DoOnGameThread([&]() {
+			GWorld->GetTimerManager().SetTimer(TimerHandle_KeepAlive, std::bind(&FSignallingServerConnection::KeepAlive, this), KEEP_ALIVE_INTERVAL, true);
+		});
+	}
+
+	void FSignallingServerConnection::StopKeepAliveTimer()
+	{
+		// GWorld dereferencing needs to happen on the game thread
+		// we need to wait because if we're destructing this object we dont
+		// want to call the callback mid/post destruction
+		DoOnGameThreadAndWait(MAX_uint32, [&]() {
+			GWorld->GetTimerManager().ClearTimer(TimerHandle_KeepAlive);
+		});
+	}
+
 	void FSignallingServerConnection::SendIceCandidate(const webrtc::IceCandidateInterface& IceCandidate)
 	{
 		FJsonObjectPtr IceCandidateJson = MakeShared<FJsonObject>();
@@ -178,7 +201,7 @@ namespace UE::PixelStreaming
 
 		UE_LOG(LogPixelStreamingSS, Verbose, TEXT("-> SS: ice-candidate\n%s"), *ToString(IceCandidateJson));
 
-		WS->Send(ToString(IceCandidateJson, false));
+		SendMessage(ToString(IceCandidateJson, false));
 	}
 
 	void FSignallingServerConnection::SendIceCandidate(FPixelStreamingPlayerId PlayerId, const webrtc::IceCandidateInterface& IceCandidate)
@@ -201,7 +224,7 @@ namespace UE::PixelStreaming
 
 		UE_LOG(LogPixelStreamingSS, Verbose, TEXT("-> SS: iceCandidate\n%s"), *ToString(IceCandidateJson));
 
-		WS->Send(ToString(IceCandidateJson, false));
+		SendMessage(ToString(IceCandidateJson, false));
 	}
 
 	void FSignallingServerConnection::KeepAlive()
@@ -210,11 +233,7 @@ namespace UE::PixelStreaming
 		double unixTime = FDateTime::UtcNow().ToUnixTimestamp();
 		Json->SetStringField(TEXT("type"), TEXT("ping"));
 		Json->SetNumberField(TEXT("time"), unixTime);
-		FString Msg = ToString(Json, false);
-		if (WS != nullptr && WS.IsValid() && WS->IsConnected())
-		{
-			WS->Send(Msg);
-		}
+		SendMessage(ToString(Json, false));
 	}
 
 	void FSignallingServerConnection::SendDisconnectPlayer(FPixelStreamingPlayerId PlayerId, const FString& Reason)
@@ -228,7 +247,7 @@ namespace UE::PixelStreaming
 		FString Msg = ToString(Json, false);
 		UE_LOG(LogPixelStreamingSS, Verbose, TEXT("-> SS: iceCandidate\n%s"), *Msg);
 
-		WS->Send(Msg);
+		SendMessage(Msg);
 	}
 
 	void FSignallingServerConnection::SendAnswer(const webrtc::SessionDescriptionInterface& SDP)
@@ -241,7 +260,7 @@ namespace UE::PixelStreaming
 
 		UE_LOG(LogPixelStreamingSS, Verbose, TEXT("-> SS: answer\n%s"), *SdpStr);
 
-		WS->Send(ToString(AnswerJson, false));
+		SendMessage(ToString(AnswerJson, false));
 	}
 
 	void FSignallingServerConnection::OnConnected()
@@ -250,8 +269,7 @@ namespace UE::PixelStreaming
 
 		Observer.OnSignallingServerConnected();
 
-		//Send message to keep connection alive every 60 seconds
-		GWorld->GetTimerManager().SetTimer(TimerHandle_KeepAlive, std::bind(&FSignallingServerConnection::KeepAlive, this), KEEP_ALIVE_INTERVAL, true);
+		StartKeepAliveTimer();
 
 		if (UPixelStreamingDelegates* Delegates = UPixelStreamingDelegates::GetPixelStreamingDelegates())
 		{
@@ -266,7 +284,7 @@ namespace UE::PixelStreaming
 
 		Observer.OnSignallingServerDisconnected();
 
-		GWorld->GetTimerManager().ClearTimer(TimerHandle_KeepAlive);
+		StopKeepAliveTimer();
 
 		if (UPixelStreamingDelegates* Delegates = UPixelStreamingDelegates::GetPixelStreamingDelegates())
 		{
@@ -281,7 +299,7 @@ namespace UE::PixelStreaming
 
 		Observer.OnSignallingServerDisconnected();
 
-		GWorld->GetTimerManager().ClearTimer(TimerHandle_KeepAlive);
+		StopKeepAliveTimer();
 
 		if (UPixelStreamingDelegates* Delegates = UPixelStreamingDelegates::GetPixelStreamingDelegates())
 		{
@@ -332,15 +350,16 @@ namespace UE::PixelStreaming
 		// As `RTCConfiguration` doesn't implement parsing from a string (or `ToString` method),
 		// we just get `stun`/`turn` URLs from it and ignore other options
 
-		const TSharedPtr<FJsonObject>* PeerConnectionOptions;
+		const TSharedPtr<FJsonObject>* PeerConnectionOptions = nullptr;
 		if (!Json->TryGetObjectField(TEXT("peerConnectionOptions"), PeerConnectionOptions))
 		{
 			FatalError(TEXT("Cannot find `peerConnectionOptions` field in SS config\n%s"), *ToString(Json));
+			return;
 		}
 
 		webrtc::PeerConnectionInterface::RTCConfiguration RTCConfig;
 
-		const TArray<TSharedPtr<FJsonValue>>* IceServers;
+		const TArray<TSharedPtr<FJsonValue>>* IceServers = nullptr;
 		if ((*PeerConnectionOptions)->TryGetArrayField(TEXT("iceServers"), IceServers))
 		{
 			for (const TSharedPtr<FJsonValue>& IceServerVal : *IceServers)
@@ -556,9 +575,17 @@ namespace UE::PixelStreaming
 		SendDisconnectPlayer(PlayerId, Msg);
 	}
 
+	void FSignallingServerConnection::SendMessage(const FString& Msg)
+	{
+		if (WebSocket && WebSocket->IsConnected())
+		{
+			WebSocket->Send(Msg);
+		}
+	}
+
 	void FSignallingServerConnection::FatalError(const FString& Msg)
 	{
 		UE_LOG(LogPixelStreamingSS, Error, TEXT("%s"), *Msg);
-		WS->Close(4000, Msg);
+		WebSocket->Close(4000, Msg);
 	}
 } // namespace UE::PixelStreaming
diff --git a/Engine/Plugins/Media/PixelStreaming/Source/PixelStreaming/Private/SignallingServerConnection.h b/Engine/Plugins/Media/PixelStreaming/Source/PixelStreaming/Private/SignallingServerConnection.h
--- a/Engine/Plugins/Media/PixelStreaming/Source/PixelStreaming/Private/SignallingServerConnection.h
+++ b/Engine/Plugins/Media/PixelStreaming/Source/PixelStreaming/Private/SignallingServerConnection.h
@@ -36,10 +36,14 @@ namespace UE::PixelStreaming
 		virtual void OnPeerDataChannels(int32 SendStreamId, int32 RecvStreamId) { unimplemented(); }
 	};
 
+	
+
 	class FSignallingServerConnection final
 	{
 	public:
-		explicit FSignallingServerConnection(FSignallingServerConnectionObserver& Observer, FString InStreamerId);
+		using FWebSocketFactory = TFunction<TSharedPtr<IWebSocket>(const FString&)>;
+
+		FSignallingServerConnection(const FWebSocketFactory& InWebSocketFactory, FSignallingServerConnectionObserver& Observer, FString InStreamerId = "");
 		~FSignallingServerConnection();
 
 		void Connect(const FString& Url);
@@ -79,6 +83,11 @@ namespace UE::PixelStreaming
 		void SetPlayerIdJson(FJsonObjectPtr& JsonObject, FPixelStreamingPlayerId PlayerId);
 		bool GetPlayerIdJson(const FJsonObjectPtr& Json, FPixelStreamingPlayerId& OutPlayerId, const FString& FieldId = TEXT("playerId"));
 
+		void StartKeepAliveTimer();
+		void StopKeepAliveTimer();
+
+		void SendMessage(const FString& Msg);
+
 		template <typename FmtType, typename... T>
 		void PlayerError(FPixelStreamingPlayerId PlayerId, const FmtType& Msg, T... args)
 		{
@@ -96,9 +105,11 @@ namespace UE::PixelStreaming
 		void FatalError(const FString& Msg);
 
 	private:
-		TMap<FString, TFunction<void(FJsonObjectPtr)>> MessageHandlers;
-
+		FWebSocketFactory WebSocketFactory;
 		FSignallingServerConnectionObserver& Observer;
+		FString StreamerId;
+
+		TSharedPtr<IWebSocket> WebSocket;
 
 		FDelegateHandle OnConnectedHandle;
 		FDelegateHandle OnConnectionErrorHandle;
@@ -109,8 +120,6 @@ namespace UE::PixelStreaming
 		FTimerHandle TimerHandle_KeepAlive;
 		const float KEEP_ALIVE_INTERVAL = 60.0f;
 
-		TSharedPtr<IWebSocket> WS;
-
-		FString StreamerId;
+		TMap<FString, TFunction<void(FJsonObjectPtr)>> MessageHandlers;
 	};
 } // namespace UE::PixelStreaming
\ No newline at end of file
diff --git a/Engine/Plugins/Media/PixelStreaming/Source/PixelStreaming/Private/SimulcastEncoderAdapter.cpp b/Engine/Plugins/Media/PixelStreaming/Source/PixelStreaming/Private/SimulcastEncoderAdapter.cpp
--- a/Engine/Plugins/Media/PixelStreaming/Source/PixelStreaming/Private/SimulcastEncoderAdapter.cpp
+++ b/Engine/Plugins/Media/PixelStreaming/Source/PixelStreaming/Private/SimulcastEncoderAdapter.cpp
@@ -1,7 +1,8 @@
 // Copyright Epic Games, Inc. All Rights Reserved.
 #include "SimulcastEncoderAdapter.h"
-#include "PixelStreamingFrameBuffer.h"
-#include "IPixelStreamingTextureSource.h"
+#include "FrameBuffer.h"
+#include "modules/video_coding/codecs/vp8/include/vp8.h"
+#include "modules/video_coding/codecs/vp9/include/vp9.h"
 
 namespace
 {
@@ -92,91 +93,64 @@ namespace
 	};
 } // namespace
 
-UE::PixelStreaming::FSimulcastEncoderAdapter::FSimulcastEncoderAdapter(FSimulcastEncoderFactory& InSimulcastFactory, const webrtc::SdpVideoFormat& format)
-	: Initialized(false)
-	, SimulcastEncoderFactory(InSimulcastFactory)
-	, VideoFormat(format)
-	, EncodedCompleteCallback(nullptr)
+namespace UE::PixelStreaming
 {
-	memset(&CurrentCodec, 0, sizeof(webrtc::VideoCodec));
-}
+	FSimulcastEncoderAdapter::FSimulcastEncoderAdapter(FSimulcastEncoderFactory& InSimulcastFactory, const webrtc::SdpVideoFormat& format)
+		: Initialized(false)
+		, SimulcastEncoderFactory(InSimulcastFactory)
+		, VideoFormat(format)
+		, EncodedCompleteCallback(nullptr)
+	{
+		memset(&CurrentCodec, 0, sizeof(webrtc::VideoCodec));
+	}
 
-UE::PixelStreaming::FSimulcastEncoderAdapter::~FSimulcastEncoderAdapter()
-{
-	check(!IsInitialized());
-}
+	FSimulcastEncoderAdapter::~FSimulcastEncoderAdapter()
+	{
+		check(!IsInitialized());
+	}
 
-int UE::PixelStreaming::FSimulcastEncoderAdapter::Release()
-{
+	int FSimulcastEncoderAdapter::Release()
 	{
-		// Lock during deleting an encoder
-		FScopeLock Lock(&StreamInfosGuard);
-		while (!StreamInfos.empty())
 		{
-			std::unique_ptr<VideoEncoder> Encoder = std::move(StreamInfos.back().Encoder);
-			// Even though it seems very unlikely, there are no guarantees that the
-			// encoder will not call back after being Release()'d. Therefore, we first
-			// disable the callbacks here.
-			Encoder->RegisterEncodeCompleteCallback(nullptr);
-			Encoder->Release();
-			StreamInfos.pop_back(); // Deletes callback adapter.
+			// Lock during deleting an encoder
+			FScopeLock Lock(&StreamInfosGuard);
+			while (!StreamInfos.empty())
+			{
+				std::unique_ptr<VideoEncoder> Encoder = std::move(StreamInfos.back().Encoder);
+				// Even though it seems very unlikely, there are no guarantees that the
+				// encoder will not call back after being Release()'d. Therefore, we first
+				// disable the callbacks here.
+				Encoder->RegisterEncodeCompleteCallback(nullptr);
+				Encoder->Release();
+				StreamInfos.pop_back(); // Deletes callback adapter.
+			}
 		}
-	}
 
-	Initialized = false;
+		Initialized = false;
 
-	return WEBRTC_VIDEO_CODEC_OK;
-}
-
-int UE::PixelStreaming::FSimulcastEncoderAdapter::InitEncode(const webrtc::VideoCodec* codec_settings, const webrtc::VideoEncoder::Settings& settings)
-{
-	int NumberOfStreams = GetNumberOfStreams(*codec_settings);
-	int NumActiveStreams = GetNumActiveStreams(*codec_settings);
-
-	CurrentCodec = *codec_settings;
+		return WEBRTC_VIDEO_CODEC_OK;
+	}
 
-	const auto MinMax = std::minmax_element(std::begin(CurrentCodec.simulcastStream), std::begin(CurrentCodec.simulcastStream) + NumberOfStreams, StreamResolutionCompare);
-	LowestResolutionStreamIndex = std::distance(std::begin(CurrentCodec.simulcastStream), MinMax.first);
-	HighestResolutionStreamIndex = std::distance(std::begin(CurrentCodec.simulcastStream), MinMax.second);
+	int FSimulcastEncoderAdapter::InitEncode(const webrtc::VideoCodec* codec_settings, const webrtc::VideoEncoder::Settings& settings)
+	{
+		int NumberOfStreams = GetNumberOfStreams(*codec_settings);
+		int NumActiveStreams = GetNumActiveStreams(*codec_settings);
 
-	const webrtc::SdpVideoFormat Format(CurrentCodec.codecType == webrtc::kVideoCodecVP8 ? "VP8" : "H264", VideoFormat.parameters);
+		CurrentCodec = *codec_settings;
 
-	if (NumActiveStreams == 1)
-	{
-		// with one stream we just proxy the pixelstreaming encoder
-		FVideoEncoderFactory* EncoderFactory = SimulcastEncoderFactory.GetOrCreateEncoderFactory(0);
-		std::unique_ptr<VideoEncoder> Encoder = EncoderFactory->CreateVideoEncoder(Format);
+		const auto MinMax = std::minmax_element(std::begin(CurrentCodec.simulcastStream), std::begin(CurrentCodec.simulcastStream) + NumberOfStreams, StreamResolutionCompare);
+		LowestResolutionStreamIndex = std::distance(std::begin(CurrentCodec.simulcastStream), MinMax.first);
+		HighestResolutionStreamIndex = std::distance(std::begin(CurrentCodec.simulcastStream), MinMax.second);
 
-		int ReturnCode = Encoder->InitEncode(&CurrentCodec, settings);
-		if (ReturnCode < 0)
-		{
-			// Explicitly destroy the current encoder; because we haven't registered
-			// a StreamInfo for it yet, Release won't do anything about it.
-			Encoder.reset();
-			Release();
-			return ReturnCode;
-		}
+		const webrtc::SdpVideoFormat Format(CurrentCodec.codecType == webrtc::kVideoCodecVP8 ? "VP8" : "H264", VideoFormat.parameters);
 
-		std::unique_ptr<webrtc::EncodedImageCallback> Callback(new AdapterEncodedImageCallback(this, 0));
-		Encoder->RegisterEncodeCompleteCallback(Callback.get());
-		StreamInfos.emplace_back(
-			std::move(Encoder), std::move(Callback),
-			std::make_unique<webrtc::FramerateController>(CurrentCodec.maxFramerate),
-			CurrentCodec.width, CurrentCodec.height, true, true);
-	}
-	else
-	{
-		for (int i = 0; i < NumberOfStreams; ++i)
+		if (NumActiveStreams == 1)
 		{
-			webrtc::VideoCodec StreamCodec;
-			uint32_t StartBitrateKbps = CurrentCodec.simulcastStream[i].targetBitrate;
-			PopulateStreamCodec(CurrentCodec, i, StartBitrateKbps, &StreamCodec);
-
-			FVideoEncoderFactory* EncoderFactory = SimulcastEncoderFactory.GetOrCreateEncoderFactory(i);
+			// with one stream we just proxy the pixelstreaming encoder
+			FVideoEncoderFactory* EncoderFactory = SimulcastEncoderFactory.GetOrCreateEncoderFactory(0);
 			std::unique_ptr<VideoEncoder> Encoder = EncoderFactory->CreateVideoEncoder(Format);
 
-			// std::unique_ptr<VideoEncoder> Encoder = SimulcastEncoderFactory->CreateVideoEncoder(Format);
-			int ReturnCode = Encoder->InitEncode(&StreamCodec, settings);
+			int ReturnCode = Encoder->InitEncode(&CurrentCodec, settings);
 			if (ReturnCode < 0)
 			{
 				// Explicitly destroy the current encoder; because we haven't registered
@@ -186,297 +160,405 @@ int UE::PixelStreaming::FSimulcastEncoderAdapter::InitEncode(const webrtc::Video
 				return ReturnCode;
 			}
 
-			std::unique_ptr<webrtc::EncodedImageCallback> Callback(new AdapterEncodedImageCallback(this, i));
+			std::unique_ptr<webrtc::EncodedImageCallback> Callback(new AdapterEncodedImageCallback(this, 0));
 			Encoder->RegisterEncodeCompleteCallback(Callback.get());
 			StreamInfos.emplace_back(
 				std::move(Encoder), std::move(Callback),
-				std::make_unique<webrtc::FramerateController>(StreamCodec.maxFramerate),
-				StreamCodec.width, StreamCodec.height, true, true);
+				std::make_unique<webrtc::FramerateController>(CurrentCodec.maxFramerate),
+				CurrentCodec.width, CurrentCodec.height, true, true);
 		}
-	}
-
-	Initialized = true;
-
-	return WEBRTC_VIDEO_CODEC_OK;
-}
-
-int UE::PixelStreaming::FSimulcastEncoderAdapter::EncodeStream(const webrtc::VideoFrame& InputImage, TSharedPtr<IPixelStreamingTextureSource> LayerFrameSource, size_t StreamIdx, bool bSendKeyFrame)
-{
-	// grab the simulcast frame source, extract the frame source for this layer and wrap that in a new frame buffer
-	rtc::scoped_refptr<FLayerFrameBuffer> LayerFrameBuffer = new rtc::RefCountedObject<FLayerFrameBuffer>(LayerFrameSource);
-	webrtc::VideoFrame NewFrame(InputImage);
-	NewFrame.set_video_frame_buffer(LayerFrameBuffer);
-
-	const uint32_t FrameTimestampMs = 1000 * NewFrame.timestamp() / 90000;
-
-	// If adapter is passed through and only one sw encoder does simulcast,
-	// frame types for all streams should be passed to the encoder unchanged.
-	// Otherwise a single per-encoder frame type is passed.
-	std::vector<webrtc::VideoFrameType> StreamFrameTypes(StreamInfos.size() == 1 ? GetNumberOfStreams(CurrentCodec) : 1);
-	if (bSendKeyFrame)
-	{
-		std::fill(StreamFrameTypes.begin(), StreamFrameTypes.end(), webrtc::VideoFrameType::kVideoFrameKey);
-		StreamInfos[StreamIdx].KeyFrameRequest = false;
-	}
-	else
-	{
-		if (StreamInfos[StreamIdx].FramerateController->DropFrame(FrameTimestampMs))
+		else
 		{
-			return WEBRTC_VIDEO_CODEC_OK;
+			for (int i = 0; i < NumberOfStreams; ++i)
+			{
+				webrtc::VideoCodec StreamCodec;
+				uint32_t StartBitrateKbps = CurrentCodec.simulcastStream[i].targetBitrate;
+				PopulateStreamCodec(CurrentCodec, i, StartBitrateKbps, &StreamCodec);
+
+				FVideoEncoderFactory* EncoderFactory = SimulcastEncoderFactory.GetOrCreateEncoderFactory(i);
+				std::unique_ptr<VideoEncoder> Encoder = EncoderFactory->CreateVideoEncoder(Format);
+
+				int ReturnCode = Encoder->InitEncode(&StreamCodec, settings);
+				if (ReturnCode < 0)
+				{
+					// Explicitly destroy the current encoder; because we haven't registered
+					// a StreamInfo for it yet, Release won't do anything about it.
+					Encoder.reset();
+					Release();
+					return ReturnCode;
+				}
+
+				std::unique_ptr<webrtc::EncodedImageCallback> Callback(new AdapterEncodedImageCallback(this, i));
+				Encoder->RegisterEncodeCompleteCallback(Callback.get());
+				StreamInfos.emplace_back(
+					std::move(Encoder), std::move(Callback),
+					std::make_unique<webrtc::FramerateController>(StreamCodec.maxFramerate),
+					StreamCodec.width, StreamCodec.height, true, true);
+			}
 		}
-		std::fill(StreamFrameTypes.begin(), StreamFrameTypes.end(), webrtc::VideoFrameType::kVideoFrameDelta);
-	}
-	StreamInfos[StreamIdx].FramerateController->AddFrame(FrameTimestampMs);
 
-	return StreamInfos[StreamIdx].Encoder->Encode(NewFrame, &StreamFrameTypes);
-}
+		Initialized = true;
 
-int UE::PixelStreaming::FSimulcastEncoderAdapter::Encode(const webrtc::VideoFrame& input_image, const std::vector<webrtc::VideoFrameType>* frame_types)
-{
-	if (!IsInitialized())
-	{
-		return WEBRTC_VIDEO_CODEC_UNINITIALIZED;
+		return WEBRTC_VIDEO_CODEC_OK;
 	}
 
-	if (EncodedCompleteCallback == nullptr)
+	int FSimulcastEncoderAdapter::EncodeVP8(const webrtc::VideoFrame& InputImage, const std::vector<webrtc::VideoFrameType>* FrameTypes, bool bSendKeyFrame)
 	{
-		return WEBRTC_VIDEO_CODEC_UNINITIALIZED;
+		rtc::scoped_refptr<webrtc::I420BufferInterface> SourceBuffer;
+		for (size_t StreamIdx = 0; StreamIdx < StreamInfos.size(); ++StreamIdx)
+		{
+			// Don't encode frames in resolutions that we don't intend to send.
+			if (!StreamInfos[StreamIdx].bSendStream) {
+				continue;
+			} 
+			const uint32_t FrameTimestampMs = 1000 * InputImage.timestamp() / 90000;  // kVideoPayloadTypeFrequency;
+			// If adapter is passed through and only one sw encoder does simulcast,
+			// frame types for all streams should be passed to the encoder unchanged.
+			// Otherwise a single per-encoder frame type is passed.
+			std::vector<webrtc::VideoFrameType> StreamFrameTypes(StreamInfos.size() == 1 ? GetNumberOfStreams(CurrentCodec) : 1);
+			if (bSendKeyFrame)
+			{
+				std::fill(StreamFrameTypes.begin(), StreamFrameTypes.end(), webrtc::VideoFrameType::kVideoFrameKey);
+				StreamInfos[StreamIdx].KeyFrameRequest = false;
+			}
+			else
+			{
+				if (StreamInfos[StreamIdx].FramerateController->DropFrame(FrameTimestampMs))
+				{
+					continue;
+				}
+				std::fill(StreamFrameTypes.begin(), StreamFrameTypes.end(), webrtc::VideoFrameType::kVideoFrameDelta);
+			}
+			StreamInfos[StreamIdx].FramerateController->AddFrame(FrameTimestampMs);
+			int DestWidth = StreamInfos[StreamIdx].Width;
+    		int DestHeight = StreamInfos[StreamIdx].Height;
+			if((DestWidth == InputImage.width() && DestHeight == InputImage.height())
+				|| (InputImage.video_frame_buffer()->type() == webrtc::VideoFrameBuffer::Type::kNative && StreamInfos[StreamIdx].Encoder->GetEncoderInfo().supports_native_handle))
+			{
+				int ret = StreamInfos[StreamIdx].Encoder->Encode(InputImage, &StreamFrameTypes);
+				if(ret != WEBRTC_VIDEO_CODEC_OK)
+				{
+					return ret;
+				}
+			}
+			else
+			{
+				if(SourceBuffer == nullptr)
+				{
+					SourceBuffer = InputImage.video_frame_buffer()->ToI420();
+				}
+				rtc::scoped_refptr<webrtc::I420Buffer> DestBuffer = webrtc::I420Buffer::Create(DestWidth, DestHeight);
+				DestBuffer->ScaleFrom(*SourceBuffer);
+				webrtc::VideoFrame Frame(InputImage);
+				Frame.set_video_frame_buffer(DestBuffer);
+				Frame.set_rotation(webrtc::kVideoRotation_0);
+				Frame.set_update_rect(
+					webrtc::VideoFrame::UpdateRect{0, 0, Frame.width(), Frame.height()});
+				int ret = StreamInfos[StreamIdx].Encoder->Encode(Frame, &StreamFrameTypes);
+				if (ret != WEBRTC_VIDEO_CODEC_OK) 
+				{
+					return ret;
+				}
+			}
+		}
+		return WEBRTC_VIDEO_CODEC_OK;
 	}
 
-	// ignore any init frames
-	// NOTE: It seems most of the time they never even get here because the internals of WebRTC decide to drop it. It seems it thinks
-	// the encoder is paused or something.
-	const FPixelStreamingFrameBuffer* FrameBuffer = static_cast<FPixelStreamingFrameBuffer*>(input_image.video_frame_buffer().get());
-	if (FrameBuffer->GetFrameBufferType() == UE::PixelStreaming::EFrameBufferType::Initialize)
+	int FSimulcastEncoderAdapter::EncodeVP9(const webrtc::VideoFrame& InputImage, const std::vector<webrtc::VideoFrameType>* FrameTypes, bool bSendKeyFrame)
 	{
 		return WEBRTC_VIDEO_CODEC_OK;
 	}
 
-	// separate out our pixelstreaming image sources
-	check(FrameBuffer->GetFrameBufferType() == UE::PixelStreaming::EFrameBufferType::Simulcast);
-	const FSimulcastFrameBuffer* SimulcastFrameBuffer = static_cast<const FSimulcastFrameBuffer*>(FrameBuffer);
-
-	// All active streams should generate a key frame if
-	// a key frame is requested by any stream.
-	bool bSendKeyFrame = false;
-	if (frame_types)
+	int FSimulcastEncoderAdapter::EncodeH264(const webrtc::VideoFrame& InputImage, const std::vector<webrtc::VideoFrameType>* FrameTypes, bool bSendKeyFrame)
 	{
-		for (size_t i = 0; i < frame_types->size(); ++i)
+		// ignore any init frames
+		// NOTE: It seems most of the time they never even get here because the internals of WebRTC decide to drop it. It seems it thinks
+		// the encoder is paused or something.
+		const FPixelStreamingFrameBuffer* FrameBuffer = static_cast<FPixelStreamingFrameBuffer*>(InputImage.video_frame_buffer().get());
+		if (FrameBuffer->GetFrameBufferType() == EPixelStreamingFrameBufferType::Initialize)
 		{
-			if (frame_types->at(i) == webrtc::VideoFrameType::kVideoFrameKey)
+			return WEBRTC_VIDEO_CODEC_OK;
+ 
+		}
+		// separate out our pixelstreaming image sources
+		check(FrameBuffer->GetFrameBufferType() == EPixelStreamingFrameBufferType::Simulcast);
+		const FSimulcastFrameBuffer* SimulcastFrameBuffer = static_cast<const FSimulcastFrameBuffer*>(FrameBuffer);
+		for (size_t StreamIdx = 0; StreamIdx < StreamInfos.size(); ++StreamIdx)
+		{
+			// Don't encode frames in resolutions that we don't intend to send.
+			if (!StreamInfos[StreamIdx].bSendStream)
 			{
-				bSendKeyFrame = true;
-				break;
+				continue;
+			}
+			TSharedPtr<FTextureTripleBuffer> LayerFrameSource = SimulcastFrameBuffer->GetLayerFrameSource(StreamInfos.size() == 1 ? SimulcastFrameBuffer->GetNumLayers() - 1 : StreamIdx);
+			webrtc::VideoFrame NewFrame(InputImage);
+			// grab the simulcast frame source, extract the frame source for this layer and wrap that in a new frame buffer appropriate for H264
+			rtc::scoped_refptr<FLayerFrameBuffer> LayerFrameBuffer = new rtc::RefCountedObject<FLayerFrameBuffer>(LayerFrameSource);
+			NewFrame.set_video_frame_buffer(LayerFrameBuffer);
+			
+			const uint32_t FrameTimestampMs = 1000 * NewFrame.timestamp() / 90000;
+			// If adapter is passed through and only one sw encoder does simulcast,
+			// frame types for all streams should be passed to the encoder unchanged.
+			// Otherwise a single per-encoder frame type is passed.
+			std::vector<webrtc::VideoFrameType> StreamFrameTypes(StreamInfos.size() == 1 ? GetNumberOfStreams(CurrentCodec) : 1);
+			if (bSendKeyFrame)
+			{
+				std::fill(StreamFrameTypes.begin(), StreamFrameTypes.end(), webrtc::VideoFrameType::kVideoFrameKey);
+				StreamInfos[StreamIdx].KeyFrameRequest = false;
+			}
+			else
+			{
+				if (StreamInfos[StreamIdx].FramerateController->DropFrame(FrameTimestampMs))
+				{
+					return WEBRTC_VIDEO_CODEC_OK;
+				}
+				std::fill(StreamFrameTypes.begin(), StreamFrameTypes.end(), webrtc::VideoFrameType::kVideoFrameDelta);
+			}
+			StreamInfos[StreamIdx].FramerateController->AddFrame(FrameTimestampMs);
+			int RtcError = StreamInfos[StreamIdx].Encoder->Encode(NewFrame, &StreamFrameTypes);
+			if (RtcError != WEBRTC_VIDEO_CODEC_OK)
+			{
+				return RtcError;
 			}
+ 
 		}
+		return WEBRTC_VIDEO_CODEC_OK;
 	}
 
-	for (size_t StreamIdx = 0; StreamIdx < StreamInfos.size(); ++StreamIdx)
+	int FSimulcastEncoderAdapter::Encode(const webrtc::VideoFrame& input_image, const std::vector<webrtc::VideoFrameType>* frame_types)
 	{
-		if (StreamInfos[StreamIdx].KeyFrameRequest && StreamInfos[StreamIdx].bSendStream)
+		if (!IsInitialized())
 		{
-			bSendKeyFrame = true;
-			break;
+			return WEBRTC_VIDEO_CODEC_UNINITIALIZED;
 		}
-	}
-
-	for (size_t StreamIdx = 0; StreamIdx < StreamInfos.size(); ++StreamIdx)
-	{
-		// Don't encode frames in resolutions that we don't intend to send.
-		if (!StreamInfos[StreamIdx].bSendStream)
+		if (EncodedCompleteCallback == nullptr)
+		{
+			return WEBRTC_VIDEO_CODEC_UNINITIALIZED;
+		}
+  
+		// All active streams should generate a key frame if
+		// a key frame is requested by any stream.
+		bool bSendKeyFrame = false;
+		if (frame_types)
 		{
-			continue;
+			for (size_t i = 0; i < frame_types->size(); ++i)
+			{
+				if (frame_types->at(i) == webrtc::VideoFrameType::kVideoFrameKey)
+				{
+					bSendKeyFrame = true;
+					break;
+				}
+			}
 		}
 
-		TSharedPtr<IPixelStreamingTextureSource> LayerFrameSource = SimulcastFrameBuffer->GetLayerFrameSource(StreamInfos.size() == 1 ? SimulcastFrameBuffer->GetNumLayers() - 1 : StreamIdx);
+		for (size_t StreamIdx = 0; StreamIdx < StreamInfos.size(); ++StreamIdx)
+		{
+			if (StreamInfos[StreamIdx].KeyFrameRequest && StreamInfos[StreamIdx].bSendStream)
+			{
+				bSendKeyFrame = true;
+				break;
+			}
+		}
 
-		int RtcError = EncodeStream(input_image, LayerFrameSource, StreamIdx, bSendKeyFrame);
-		if (RtcError != WEBRTC_VIDEO_CODEC_OK)
+		switch(CurrentCodec.codecType)
 		{
-			return RtcError;
+			case webrtc::kVideoCodecVP8:
+				return EncodeVP8(input_image, frame_types, bSendKeyFrame);
+			case webrtc::kVideoCodecVP9:
+				return EncodeVP9(input_image, frame_types, bSendKeyFrame);
+			case webrtc::kVideoCodecH264:
+				return EncodeH264(input_image, frame_types, bSendKeyFrame);
+			default:
+				return WEBRTC_VIDEO_CODEC_ERROR;
 		}
 	}
-	return WEBRTC_VIDEO_CODEC_OK;
-}
-
-int UE::PixelStreaming::FSimulcastEncoderAdapter::RegisterEncodeCompleteCallback(webrtc::EncodedImageCallback* callback)
-{
-	EncodedCompleteCallback = callback;
-	return WEBRTC_VIDEO_CODEC_OK;
-}
 
-void UE::PixelStreaming::FSimulcastEncoderAdapter::SetRates(const RateControlParameters& parameters)
-{
-	if (!IsInitialized())
+	int FSimulcastEncoderAdapter::RegisterEncodeCompleteCallback(webrtc::EncodedImageCallback* callback)
 	{
-		RTC_LOG(LS_WARNING) << "SetRates while not initialized";
-		return;
+		EncodedCompleteCallback = callback;
+		return WEBRTC_VIDEO_CODEC_OK;
 	}
 
-	if (parameters.framerate_fps < 1.0)
+	void FSimulcastEncoderAdapter::SetRates(const RateControlParameters& parameters)
 	{
-		RTC_LOG(LS_WARNING) << "Invalid framerate: " << parameters.framerate_fps;
-		return;
-	}
-
-	CurrentCodec.maxFramerate = static_cast<uint32_t>(parameters.framerate_fps + 0.5);
+		if (!IsInitialized())
+		{
+			RTC_LOG(LS_WARNING) << "SetRates while not initialized";
+			return;
+		}
 
-	if (StreamInfos.size() == 1)
-	{
-		// Not doing simulcast.
-		StreamInfos[0].Encoder->SetRates(parameters);
-		return;
-	}
+		if (parameters.framerate_fps < 1.0)
+		{
+			RTC_LOG(LS_WARNING) << "Invalid framerate: " << parameters.framerate_fps;
+			return;
+		}
 
-	for (size_t StreamIdx = 0; StreamIdx < StreamInfos.size(); ++StreamIdx)
-	{
-		uint32_t StreamBitrateKbps = parameters.bitrate.GetSpatialLayerSum(StreamIdx) / 1000;
+		CurrentCodec.maxFramerate = static_cast<uint32_t>(parameters.framerate_fps + 0.5);
 
-		// Need a key frame if we have not sent this stream before.
-		if (StreamBitrateKbps > 0 && !StreamInfos[StreamIdx].bSendStream)
+		if (StreamInfos.size() == 1)
 		{
-			StreamInfos[StreamIdx].KeyFrameRequest = true;
+			// Not doing simulcast.
+			StreamInfos[0].Encoder->SetRates(parameters);
+			return;
 		}
-		StreamInfos[StreamIdx].bSendStream = StreamBitrateKbps > 0;
 
-		// Slice the temporal layers out of the full allocation and pass it on to
-		// the encoder handling the current simulcast stream.
-		RateControlParameters StreamParameters = parameters;
-		StreamParameters.bitrate = webrtc::VideoBitrateAllocation();
-		for (int i = 0; i < webrtc::kMaxTemporalStreams; ++i)
+		for (size_t StreamIdx = 0; StreamIdx < StreamInfos.size(); ++StreamIdx)
 		{
-			if (parameters.bitrate.HasBitrate(StreamIdx, i))
+			uint32_t StreamBitrateKbps = parameters.bitrate.GetSpatialLayerSum(StreamIdx) / 1000;
+
+			// Need a key frame if we have not sent this stream before.
+			if (StreamBitrateKbps > 0 && !StreamInfos[StreamIdx].bSendStream)
 			{
-				StreamParameters.bitrate.SetBitrate(0, i, parameters.bitrate.GetBitrate(StreamIdx, i));
+				StreamInfos[StreamIdx].KeyFrameRequest = true;
 			}
-		}
+			StreamInfos[StreamIdx].bSendStream = StreamBitrateKbps > 0;
 
-		// Assign link allocation proportionally to spatial layer allocation.
-		if (!parameters.bandwidth_allocation.IsZero() && parameters.bitrate.get_sum_bps() > 0)
-		{
-			StreamParameters.bandwidth_allocation =
-				webrtc::DataRate::BitsPerSec((parameters.bandwidth_allocation.bps() * StreamParameters.bitrate.get_sum_bps()) / parameters.bitrate.get_sum_bps());
-			// Make sure we don't allocate bandwidth lower than target bitrate.
-			if (StreamParameters.bandwidth_allocation.bps() < StreamParameters.bitrate.get_sum_bps())
+			// Slice the temporal layers out of the full allocation and pass it on to
+			// the encoder handling the current simulcast stream.
+			RateControlParameters StreamParameters = parameters;
+			StreamParameters.bitrate = webrtc::VideoBitrateAllocation();
+			for (int i = 0; i < webrtc::kMaxTemporalStreams; ++i)
+			{
+				if (parameters.bitrate.HasBitrate(StreamIdx, i))
+				{
+					StreamParameters.bitrate.SetBitrate(0, i, parameters.bitrate.GetBitrate(StreamIdx, i));
+				}
+			}
+
+			// Assign link allocation proportionally to spatial layer allocation.
+			if (!parameters.bandwidth_allocation.IsZero() && parameters.bitrate.get_sum_bps() > 0)
 			{
 				StreamParameters.bandwidth_allocation =
-					webrtc::DataRate::BitsPerSec(StreamParameters.bitrate.get_sum_bps());
+					webrtc::DataRate::BitsPerSec((parameters.bandwidth_allocation.bps() * StreamParameters.bitrate.get_sum_bps()) / parameters.bitrate.get_sum_bps());
+				// Make sure we don't allocate bandwidth lower than target bitrate.
+				if (StreamParameters.bandwidth_allocation.bps() < StreamParameters.bitrate.get_sum_bps())
+				{
+					StreamParameters.bandwidth_allocation =
+						webrtc::DataRate::BitsPerSec(StreamParameters.bitrate.get_sum_bps());
+				}
 			}
-		}
 
-		StreamParameters.framerate_fps = std::min<double>(parameters.framerate_fps, StreamInfos[StreamIdx].FramerateController->GetTargetRate());
+			StreamParameters.framerate_fps = std::min<double>(parameters.framerate_fps, StreamInfos[StreamIdx].FramerateController->GetTargetRate());
 
-		StreamInfos[StreamIdx].Encoder->SetRates(StreamParameters);
+			StreamInfos[StreamIdx].Encoder->SetRates(StreamParameters);
+		}
 	}
-}
 
-void UE::PixelStreaming::FSimulcastEncoderAdapter::OnPacketLossRateUpdate(float packet_loss_rate)
-{
-	for (StreamInfo& Info : StreamInfos)
+	void FSimulcastEncoderAdapter::OnPacketLossRateUpdate(float packet_loss_rate)
 	{
-		Info.Encoder->OnPacketLossRateUpdate(packet_loss_rate);
+		for (StreamInfo& Info : StreamInfos)
+		{
+			Info.Encoder->OnPacketLossRateUpdate(packet_loss_rate);
+		}
 	}
-}
 
-void UE::PixelStreaming::FSimulcastEncoderAdapter::OnRttUpdate(int64_t rtt_ms)
-{
-	for (StreamInfo& Info : StreamInfos)
+	void FSimulcastEncoderAdapter::OnRttUpdate(int64_t rtt_ms)
 	{
-		Info.Encoder->OnRttUpdate(rtt_ms);
+		for (StreamInfo& Info : StreamInfos)
+		{
+			Info.Encoder->OnRttUpdate(rtt_ms);
+		}
 	}
-}
 
-void UE::PixelStreaming::FSimulcastEncoderAdapter::OnLossNotification(const LossNotification& loss_notification)
-{
-	for (StreamInfo& Info : StreamInfos)
+	void FSimulcastEncoderAdapter::OnLossNotification(const LossNotification& loss_notification)
 	{
-		Info.Encoder->OnLossNotification(loss_notification);
+		for (StreamInfo& Info : StreamInfos)
+		{
+			Info.Encoder->OnLossNotification(loss_notification);
+		}
 	}
-}
-
-webrtc::EncodedImageCallback::Result UE::PixelStreaming::FSimulcastEncoderAdapter::OnEncodedImage(
-	size_t stream_idx,
-	const webrtc::EncodedImage& encodedImage,
-	const webrtc::CodecSpecificInfo* codecSpecificInfo,
-	const webrtc::RTPFragmentationHeader* fragmentation)
-{
-	webrtc::EncodedImage StreamImage(encodedImage);
-	webrtc::CodecSpecificInfo StreamCodecSpecific = *codecSpecificInfo;
 
-	StreamImage.SetSpatialIndex(stream_idx);
-
-	return EncodedCompleteCallback->OnEncodedImage(StreamImage, &StreamCodecSpecific, fragmentation);
-}
+	webrtc::EncodedImageCallback::Result FSimulcastEncoderAdapter::OnEncodedImage(
+		size_t stream_idx,
+		const webrtc::EncodedImage& encodedImage,
+		const webrtc::CodecSpecificInfo* codecSpecificInfo,
+		const webrtc::RTPFragmentationHeader* fragmentation)
+	{
+		webrtc::EncodedImage StreamImage(encodedImage);
+		webrtc::CodecSpecificInfo StreamCodecSpecific = *codecSpecificInfo;
 
-bool UE::PixelStreaming::FSimulcastEncoderAdapter::IsInitialized() const
-{
-	return Initialized;
-}
+		StreamImage.SetSpatialIndex(stream_idx);
 
-webrtc::VideoEncoder::EncoderInfo UE::PixelStreaming::FSimulcastEncoderAdapter::GetEncoderInfo() const
-{
-	if (StreamInfos.size() == 1)
-	{
-		// Not using simulcast adapting functionality, just pass through.
-		return StreamInfos[0].Encoder->GetEncoderInfo();
+		return EncodedCompleteCallback->OnEncodedImage(StreamImage, &StreamCodecSpecific, fragmentation);
 	}
 
-	VideoEncoder::EncoderInfo EncoderInfo;
-	EncoderInfo.implementation_name = "PixelStreamingSimulcastEncoderAdapter";
-	EncoderInfo.requested_resolution_alignment = 1;
-	EncoderInfo.supports_native_handle = true;
-	EncoderInfo.scaling_settings.thresholds = absl::nullopt;
-	if (StreamInfos.empty())
+	bool FSimulcastEncoderAdapter::IsInitialized() const
 	{
-		return EncoderInfo;
+		return Initialized;
 	}
 
-	EncoderInfo.scaling_settings = VideoEncoder::ScalingSettings::kOff;
-	int NumActiveStreams = GetNumActiveStreams(CurrentCodec);
-
-	for (size_t i = 0; i < StreamInfos.size(); ++i)
+	webrtc::VideoEncoder::EncoderInfo FSimulcastEncoderAdapter::GetEncoderInfo() const
 	{
-		VideoEncoder::EncoderInfo EncoderImplInfo = StreamInfos[i].Encoder->GetEncoderInfo();
+		if (StreamInfos.size() == 1)
+		{
+			// Not using simulcast adapting functionality, just pass through.
+			return StreamInfos[0].Encoder->GetEncoderInfo();
+		}
 
-		if (i == 0)
+		VideoEncoder::EncoderInfo EncoderInfo;
+		EncoderInfo.implementation_name = "PixelStreamingSimulcastEncoderAdapter";
+		EncoderInfo.requested_resolution_alignment = 1;
+		EncoderInfo.supports_native_handle = true;
+		EncoderInfo.scaling_settings.thresholds = absl::nullopt;
+		if (StreamInfos.empty())
 		{
-			// Encoder name indicates names of all sub-encoders.
-			EncoderInfo.implementation_name += " (";
-			EncoderInfo.implementation_name += EncoderImplInfo.implementation_name;
-
-			EncoderInfo.supports_native_handle = EncoderImplInfo.supports_native_handle;
-			EncoderInfo.has_trusted_rate_controller = EncoderImplInfo.has_trusted_rate_controller;
-			EncoderInfo.is_hardware_accelerated = EncoderImplInfo.is_hardware_accelerated;
-			EncoderInfo.has_internal_source = EncoderImplInfo.has_internal_source;
+			return EncoderInfo;
 		}
-		else
+
+		EncoderInfo.scaling_settings = VideoEncoder::ScalingSettings::kOff;
+		int NumActiveStreams = GetNumActiveStreams(CurrentCodec);
+
+		for (size_t i = 0; i < StreamInfos.size(); ++i)
 		{
-			EncoderInfo.implementation_name += ", ";
-			EncoderInfo.implementation_name += EncoderImplInfo.implementation_name;
+			VideoEncoder::EncoderInfo EncoderImplInfo = StreamInfos[i].Encoder->GetEncoderInfo();
 
-			// Native handle supported if any encoder supports it.
-			EncoderInfo.supports_native_handle |= EncoderImplInfo.supports_native_handle;
+			if (i == 0)
+			{
+				// Encoder name indicates names of all sub-encoders.
+				EncoderInfo.implementation_name += " (";
+				EncoderInfo.implementation_name += EncoderImplInfo.implementation_name;
+
+				EncoderInfo.supports_native_handle = EncoderImplInfo.supports_native_handle;
+				EncoderInfo.has_trusted_rate_controller = EncoderImplInfo.has_trusted_rate_controller;
+				EncoderInfo.is_hardware_accelerated = EncoderImplInfo.is_hardware_accelerated;
+				EncoderInfo.has_internal_source = EncoderImplInfo.has_internal_source;
+			}
+			else
+			{
+				EncoderInfo.implementation_name += ", ";
+				EncoderInfo.implementation_name += EncoderImplInfo.implementation_name;
 
-			// Trusted rate controller only if all encoders have it.
-			EncoderInfo.has_trusted_rate_controller &= EncoderImplInfo.has_trusted_rate_controller;
+				// Native handle supported if any encoder supports it.
+				EncoderInfo.supports_native_handle |= EncoderImplInfo.supports_native_handle;
 
-			// Uses hardware support if any of the encoders uses it.
-			// For example, if we are having issues with down-scaling due to
-			// pipelining delay in HW encoders we need higher encoder usage
-			// thresholds in CPU adaptation.
-			EncoderInfo.is_hardware_accelerated |= EncoderImplInfo.is_hardware_accelerated;
+				// Trusted rate controller only if all encoders have it.
+				EncoderInfo.has_trusted_rate_controller &= EncoderImplInfo.has_trusted_rate_controller;
 
-			// Has internal source only if all encoders have it.
-			EncoderInfo.has_internal_source &= EncoderImplInfo.has_internal_source;
-		}
-		EncoderInfo.fps_allocation[i] = EncoderImplInfo.fps_allocation[0];
-		EncoderInfo.requested_resolution_alignment = cricket::LeastCommonMultiple(EncoderInfo.requested_resolution_alignment, EncoderImplInfo.requested_resolution_alignment);
-		if (NumActiveStreams == 1 && CurrentCodec.simulcastStream[i].active)
-		{
-			EncoderInfo.scaling_settings = EncoderImplInfo.scaling_settings;
+				// Uses hardware support if any of the encoders uses it.
+				// For example, if we are having issues with down-scaling due to
+				// pipelining delay in HW encoders we need higher encoder usage
+				// thresholds in CPU adaptation.
+				EncoderInfo.is_hardware_accelerated |= EncoderImplInfo.is_hardware_accelerated;
+
+				// Has internal source only if all encoders have it.
+				EncoderInfo.has_internal_source &= EncoderImplInfo.has_internal_source;
+			}
+
+			// Nasty hack to allow us to manually convert VPX frames to I420 later in the encode block
+			if(CurrentCodec.codecType == webrtc::kVideoCodecVP8)
+			{
+				EncoderInfo.supports_native_handle = true;
+			}
+
+			EncoderInfo.fps_allocation[i] = EncoderImplInfo.fps_allocation[0];
+			EncoderInfo.requested_resolution_alignment = cricket::LeastCommonMultiple(EncoderInfo.requested_resolution_alignment, EncoderImplInfo.requested_resolution_alignment);
+			if (NumActiveStreams == 1 && CurrentCodec.simulcastStream[i].active)
+			{
+				EncoderInfo.scaling_settings = EncoderImplInfo.scaling_settings;
+			}
 		}
-	}
-	EncoderInfo.implementation_name += ")";
+		EncoderInfo.implementation_name += ")";
 
-	return EncoderInfo;
-}
+		return EncoderInfo;
+	}
+} // namespace UE::PixelStreaming
\ No newline at end of file
diff --git a/Engine/Plugins/Media/PixelStreaming/Source/PixelStreaming/Private/SimulcastEncoderAdapter.h b/Engine/Plugins/Media/PixelStreaming/Source/PixelStreaming/Private/SimulcastEncoderAdapter.h
--- a/Engine/Plugins/Media/PixelStreaming/Source/PixelStreaming/Private/SimulcastEncoderAdapter.h
+++ b/Engine/Plugins/Media/PixelStreaming/Source/PixelStreaming/Private/SimulcastEncoderAdapter.h
@@ -4,8 +4,7 @@
 #include "CoreMinimal.h"
 #include "WebRTCIncludes.h"
 #include "EncoderFactory.h"
-
-class IPixelStreamingTextureSource;
+#include "TextureTripleBuffer.h"
 
 namespace webrtc
 {
@@ -13,80 +12,81 @@ namespace webrtc
 	class VideoEncoderFactory;
 } // namespace webrtc
 
-namespace UE
+namespace UE::PixelStreaming
 {
-	namespace PixelStreaming
+	// This is highly modified version of webrtc::simulcast_encoder_adapter
+	class FSimulcastEncoderAdapter : public webrtc::VideoEncoder
 	{
-		// This is highly modified version of webrtc::simulcast_encoder_adapter
-		class FSimulcastEncoderAdapter : public webrtc::VideoEncoder
-		{
-		public:
-			FSimulcastEncoderAdapter(FSimulcastEncoderFactory& InSimulcastFactory, const webrtc::SdpVideoFormat& format);
-			~FSimulcastEncoderAdapter() override;
+	public:
+		FSimulcastEncoderAdapter(FSimulcastEncoderFactory& InSimulcastFactory, const webrtc::SdpVideoFormat& format);
+		~FSimulcastEncoderAdapter() override;
 
-			// Implements VideoEncoder.
-			int Release() override;
-			int InitEncode(const webrtc::VideoCodec* codec_settings, const webrtc::VideoEncoder::Settings& settings) override;
-			int Encode(const webrtc::VideoFrame& input_image, const std::vector<webrtc::VideoFrameType>* frame_types) override;
-			int RegisterEncodeCompleteCallback(webrtc::EncodedImageCallback* callback) override;
+		// Implements VideoEncoder.
+		int Release() override;
+		int InitEncode(const webrtc::VideoCodec* codec_settings, const webrtc::VideoEncoder::Settings& settings) override;
+		int Encode(const webrtc::VideoFrame& input_image, const std::vector<webrtc::VideoFrameType>* frame_types) override;
+		int RegisterEncodeCompleteCallback(webrtc::EncodedImageCallback* callback) override;
 
-			void SetRates(const RateControlParameters& parameters) override;
-			void OnPacketLossRateUpdate(float packet_loss_rate) override;
-			void OnRttUpdate(int64_t rtt_ms) override;
-			void OnLossNotification(const LossNotification& loss_notification) override;
+		void SetRates(const RateControlParameters& parameters) override;
+		void OnPacketLossRateUpdate(float packet_loss_rate) override;
+		void OnRttUpdate(int64_t rtt_ms) override;
+		void OnLossNotification(const LossNotification& loss_notification) override;
 
-			// Eventual handler for the contained encoders' EncodedImageCallbacks, but
-			// called from an internal helper that also knows the correct stream
-			// index.
-			webrtc::EncodedImageCallback::Result OnEncodedImage(
-				size_t stream_idx,
-				const webrtc::EncodedImage& encoded_image,
-				const webrtc::CodecSpecificInfo* codec_specific_info,
-				const webrtc::RTPFragmentationHeader* fragmentation);
+		// Eventual handler for the contained encoders' EncodedImageCallbacks, but
+		// called from an internal helper that also knows the correct stream
+		// index.
+		webrtc::EncodedImageCallback::Result OnEncodedImage(
+			size_t stream_idx,
+			const webrtc::EncodedImage& encoded_image,
+			const webrtc::CodecSpecificInfo* codec_specific_info,
+			const webrtc::RTPFragmentationHeader* fragmentation);
 
-			EncoderInfo GetEncoderInfo() const override;
+		EncoderInfo GetEncoderInfo() const override;
 
-		private:
-			struct StreamInfo
+	private:
+		struct StreamInfo
+		{
+			StreamInfo(std::unique_ptr<webrtc::VideoEncoder> InEncoder,
+				std::unique_ptr<webrtc::EncodedImageCallback> InCallback,
+				std::unique_ptr<webrtc::FramerateController> InFramerateController,
+				uint16_t InWidth,
+				uint16_t InHeight,
+				bool bInKeyFrame,
+				bool bInSendStream)
+				: Encoder(std::move(InEncoder))
+				, Callback(std::move(InCallback))
+				, FramerateController(std::move(InFramerateController))
+				, Width(InWidth)
+				, Height(InHeight)
+				, KeyFrameRequest(bInKeyFrame)
+				, bSendStream(bInSendStream)
 			{
-				StreamInfo(std::unique_ptr<webrtc::VideoEncoder> InEncoder,
-					std::unique_ptr<webrtc::EncodedImageCallback> InCallback,
-					std::unique_ptr<webrtc::FramerateController> InFramerateController,
-					uint16_t InWidth,
-					uint16_t InHeight,
-					bool bInKeyFrame,
-					bool bInSendStream)
-					: Encoder(std::move(InEncoder))
-					, Callback(std::move(InCallback))
-					, FramerateController(std::move(InFramerateController))
-					, Width(InWidth)
-					, Height(InHeight)
-					, KeyFrameRequest(bInKeyFrame)
-					, bSendStream(bInSendStream)
-				{
-				}
-				std::unique_ptr<webrtc::VideoEncoder> Encoder;
-				std::unique_ptr<webrtc::EncodedImageCallback> Callback;
-				std::unique_ptr<webrtc::FramerateController> FramerateController;
-				uint16_t Width;
-				uint16_t Height;
-				bool KeyFrameRequest;
-				bool bSendStream;
-			};
+			}
+			std::unique_ptr<webrtc::VideoEncoder> Encoder;
+			std::unique_ptr<webrtc::EncodedImageCallback> Callback;
+			std::unique_ptr<webrtc::FramerateController> FramerateController;
+			uint16_t Width;
+			uint16_t Height;
+			bool KeyFrameRequest;
+			bool bSendStream;
+		};
 
-			bool IsInitialized() const;
+		bool IsInitialized() const;
 
-			TAtomic<bool> Initialized;
-			FSimulcastEncoderFactory& SimulcastEncoderFactory;
-			const webrtc::SdpVideoFormat VideoFormat;
-			webrtc::VideoCodec CurrentCodec;
-			FCriticalSection StreamInfosGuard;
-			std::vector<StreamInfo> StreamInfos;
-			webrtc::EncodedImageCallback* EncodedCompleteCallback;
+		TAtomic<bool> Initialized;
+		FSimulcastEncoderFactory& SimulcastEncoderFactory;
+		const webrtc::SdpVideoFormat VideoFormat;
+		webrtc::VideoCodec CurrentCodec;
+		FCriticalSection StreamInfosGuard;
+		std::vector<StreamInfo> StreamInfos;
+		webrtc::EncodedImageCallback* EncodedCompleteCallback;
 
-			int LowestResolutionStreamIndex;
-			int HighestResolutionStreamIndex;
-			int EncodeStream(const webrtc::VideoFrame& InputImage, TSharedPtr<IPixelStreamingTextureSource> LayerFrameSource, size_t StreamIdx, bool bSendKeyFrame);
-		};
-	} // namespace PixelStreaming
-} // namespace UE
+		int LowestResolutionStreamIndex;
+		int HighestResolutionStreamIndex;
+
+		//int EncodeStream(const webrtc::VideoFrame& InputImage, TSharedPtr<FTextureTripleBuffer> LayerFrameSource, size_t StreamIdx, bool bSendKeyFrame);
+		int EncodeVP8(const webrtc::VideoFrame& InputImage, const std::vector<webrtc::VideoFrameType>* FrameTypes, bool bSendKeyFrame);
+		int EncodeVP9(const webrtc::VideoFrame& InputImage, const std::vector<webrtc::VideoFrameType>* FrameTypes, bool bSendKeyFrame);
+		int EncodeH264(const webrtc::VideoFrame& InputImage, const std::vector<webrtc::VideoFrameType>* FrameTypes, bool bSendKeyFrame);
+	};
+} // namespace UE::PixelStreaming
diff --git a/Engine/Plugins/Media/PixelStreaming/Source/PixelStreaming/Private/Stats.cpp b/Engine/Plugins/Media/PixelStreaming/Source/PixelStreaming/Private/Stats.cpp
--- a/Engine/Plugins/Media/PixelStreaming/Source/PixelStreaming/Private/Stats.cpp
+++ b/Engine/Plugins/Media/PixelStreaming/Source/PixelStreaming/Private/Stats.cpp
@@ -7,6 +7,8 @@
 #include "PlayerSessions.h"
 #include "Settings.h"
 #include "IPixelStreamingStatsConsumer.h"
+#include "Engine/Console.h"
+#include "ConsoleSettings.h"
 
 namespace UE::PixelStreaming
 {
@@ -23,6 +25,8 @@ namespace UE::PixelStreaming
 		checkf(FStats::Instance == nullptr, TEXT("There should only ever been one PixelStreaming stats object."));
 		checkf(Sessions, TEXT("To make stats object the sessions object must not be nullptr."));
 		FStats::Instance = this;
+
+		UConsole::RegisterConsoleAutoCompleteEntries.AddRaw(this, &FStats::UpdateConsoleAutoComplete_GameThread);
 	}
 
 	void FStats::QueryPeerStat(FPixelStreamingPlayerId PlayerId, FName StatToQuery, TFunction<void(bool, double)> QueryCallback) const
@@ -59,7 +63,7 @@ namespace UE::PixelStreaming
 	{
 		if (IsInGameThread())
 		{
-			//todo: return bool
+			// todo: return bool
 			bool bUpdated = StorePeerStat_GameThread(PlayerId, Stat);
 			if (bUpdated)
 			{
@@ -268,6 +272,21 @@ namespace UE::PixelStreaming
 		Callbacks->Add(Callback);
 	}
 
+	void FStats::UpdateConsoleAutoComplete_GameThread(TArray<FAutoCompleteCommand>& AutoCompleteList)
+	{
+		checkf(IsInGameThread(), TEXT("This method was not called from the game thread."));
+
+		const UConsoleSettings* ConsoleSettings = GetDefault<UConsoleSettings>();
+
+		AutoCompleteList.AddDefaulted();
+
+		FAutoCompleteCommand& AutoCompleteCommand = AutoCompleteList.Last();
+
+		AutoCompleteCommand.Command = TEXT("Stat PixelStreaming");
+		AutoCompleteCommand.Desc = TEXT("Displays stats about Pixel Streaming on screen.");
+		AutoCompleteCommand.Color = ConsoleSettings->AutoCompleteCommandColor;
+	}
+
 	int32 FStats::OnRenderStats(UWorld* World, FViewport* Viewport, FCanvas* Canvas, int32 X, int32 Y, const FVector* ViewLocation, const FRotator* ViewRotation)
 	{
 		if (GAreScreenMessagesEnabled)
@@ -327,7 +346,8 @@ namespace UE::PixelStreaming
 
 	bool FStats::OnToggleStats(UWorld* World, FCommonViewportClient* ViewportClient, const TCHAR* Stream)
 	{
-		return false;
+		// todo: all about the toggle func
+		return true;
 	}
 
 	void FStats::PollPixelStreamingSettings()
@@ -378,25 +398,12 @@ namespace UE::PixelStreaming
 			return;
 		}
 
-		// Check if user has enabled on screen stats
-		if (!bRegisterEngineStats && Settings::CVarPixelStreamingOnScreenStats.GetValueOnGameThread())
+		if (!bRegisterEngineStats)
 		{
 			GAreScreenMessagesEnabled = true;
 			UEngine::FEngineStatRender RenderFunc = UEngine::FEngineStatRender::CreateRaw(this, &FStats::OnRenderStats);
 			UEngine::FEngineStatToggle ToggleFunc = UEngine::FEngineStatToggle::CreateRaw(this, &FStats::OnToggleStats);
 			GEngine->AddEngineStat(PixelStreamingStatName, PixelStreamingStatCategory, PixelStreamingStatDescription, RenderFunc, ToggleFunc, false);
-
-			// Turn on the Engine stat for Pixel Streaming
-			for (const FWorldContext& WorldContext : GEngine->GetWorldContexts())
-			{
-				if (WorldContext.WorldType == EWorldType::Game || WorldContext.WorldType == EWorldType::PIE)
-				{
-					UWorld* World = WorldContext.World();
-					UGameViewportClient* ViewportClient = World->GetGameViewport();
-					GEngine->SetEngineStat(World, ViewportClient, TEXT("PixelStreaming"), true);
-				}
-			}
-
 			bRegisterEngineStats = true;
 		}
 	}
diff --git a/Engine/Plugins/Media/PixelStreaming/Source/PixelStreaming/Private/Stats.h b/Engine/Plugins/Media/PixelStreaming/Source/PixelStreaming/Private/Stats.h
--- a/Engine/Plugins/Media/PixelStreaming/Source/PixelStreaming/Private/Stats.h
+++ b/Engine/Plugins/Media/PixelStreaming/Source/PixelStreaming/Private/Stats.h
@@ -6,6 +6,7 @@
 #include "PixelStreamingPlayerId.h"
 #include "CanvasItem.h"
 #include "UnrealEngine.h"
+#include "ConsoleSettings.h"
 
 class IPixelStreamingStatsConsumer;
 
@@ -111,6 +112,7 @@ namespace UE::PixelStreaming
 		bool StoreApplicationStat_GameThread(FStatData Stat);
 		void AddOnPeerStatChangedCallback_GameThread(FPixelStreamingPlayerId PlayerId, FName StatToListenOn, TWeakPtr<IPixelStreamingStatsConsumer> Callback);
 		void FireStatChanged_GameThread(FPixelStreamingPlayerId PlayerId, FName StatName, float StatValue);
+		void UpdateConsoleAutoComplete_GameThread(TArray<FAutoCompleteCommand>& AutoCompleteList);
 
 	private:
 		static FStats* Instance;
diff --git a/Engine/Plugins/Media/PixelStreaming/Source/PixelStreaming/Private/Streamer.cpp b/Engine/Plugins/Media/PixelStreaming/Source/PixelStreaming/Private/Streamer.cpp
--- a/Engine/Plugins/Media/PixelStreaming/Source/PixelStreaming/Private/Streamer.cpp
+++ b/Engine/Plugins/Media/PixelStreaming/Source/PixelStreaming/Private/Streamer.cpp
@@ -26,7 +26,6 @@ namespace UE::PixelStreaming
 	FStreamer::FStreamer(const FString& InStreamerId)
 		: StreamerId(InStreamerId)
 		, WebRtcSignallingThread(MakeUnique<rtc::Thread>(rtc::SocketServer::CreateDefault()))
-		, SignallingServerConnection(MakeUnique<FSignallingServerConnection>(*this, InStreamerId))
 		, Stats(&PlayerSessions)
 	{
 		RedirectWebRtcLogsToUnreal(rtc::LoggingSeverity::LS_VERBOSE);
@@ -36,6 +35,8 @@ namespace UE::PixelStreaming
 		// required for communication with Signalling Server and must be called in the game thread, while it's used in signalling thread
 		FModuleManager::LoadModuleChecked<FWebSocketsModule>("WebSockets");
 
+		FSignallingServerConnection::FWebSocketFactory WebSocketFactory = [](const FString& Url) { return FWebSocketsModule::Get().CreateWebSocket(Url, TEXT("")); };
+		SignallingServerConnection = MakeUnique<FSignallingServerConnection>(WebSocketFactory, *this, InStreamerId);
 		StartWebRtcSignallingThread();
 	}
 
@@ -557,25 +558,33 @@ namespace UE::PixelStreaming
 
 		bool bIsSFU = (Flags & Protocol::EPlayerFlags::PSPFlag_IsSFU) != Protocol::EPlayerFlags::PSPFlag_None;
 
+		const TArray<FName>& ActiveTextureSourceTypes = IPixelStreamingModule::Get().GetActiveTextureSourceTypes();
+
+		// Todo (Luke): Add multitrack support here.
+		verifyf(ActiveTextureSourceTypes.Num() == 1, TEXT("Only a single track is supported for streaming currently."));
+
+		FName SourceType = ActiveTextureSourceTypes[0];
+
 		rtc::scoped_refptr<webrtc::PeerConnectionFactoryInterface> PCFactory = bIsSFU ? SFUPeerConnectionFactory : P2PPeerConnectionFactory;
 
 		// Create a video source for this player
 		rtc::scoped_refptr<FVideoSourceBase> VideoSource;
 		if (bIsSFU)
 		{
-			VideoSource = new FVideoSourceSFU();
+			VideoSource = new FVideoSourceSFU(SourceType);
 		}
 		else
 		{
-			VideoSource = new FVideoSourceP2P(Session->GetPlayerId(), &PlayerSessions);
+			FPixelStreamingPlayerId PlayerId = Session->GetPlayerId();
+			VideoSource = new FVideoSourceP2P(SourceType, [PlayerId, this]() { return PlayerSessions.IsQualityController(PlayerId); });
 		}
 
 		Session->SetVideoSource(VideoSource);
-		VideoSource->Initialize();
 
 		// Create video track
 		rtc::scoped_refptr<webrtc::VideoTrackInterface> VideoTrack = PCFactory->CreateVideoTrack(TCHAR_TO_UTF8(*InVideoTrackLabel), VideoSource.get());
 		VideoTrack->set_enabled(true);
+		VideoSource->Initialize();
 
 		// Set some content hints based on degradation prefs, WebRTC uses these internally.
 		webrtc::DegradationPreference DegradationPref = Settings::GetDegradationPreference();
@@ -807,4 +816,5 @@ namespace UE::PixelStreaming
 	{
 		SignallingServerConnection->SendDisconnectPlayer(PlayerId, TEXT("Player was kicked"));
 	}
+
 } // namespace UE::PixelStreaming
diff --git a/Engine/Plugins/Media/PixelStreaming/Source/PixelStreaming/Private/Streamer.h b/Engine/Plugins/Media/PixelStreaming/Source/PixelStreaming/Private/Streamer.h
--- a/Engine/Plugins/Media/PixelStreaming/Source/PixelStreaming/Private/Streamer.h
+++ b/Engine/Plugins/Media/PixelStreaming/Source/PixelStreaming/Private/Streamer.h
@@ -7,7 +7,7 @@
 #include "RHI.h"
 #include "PlayerSessions.h"
 #include "FixedFPSPump.h"
-#include "GPUFencePoller.h"
+#include "Poller.h"
 #include "Stats.h"
 
 namespace UE::PixelStreaming
@@ -16,7 +16,7 @@ namespace UE::PixelStreaming
 	class FVideoEncoderFactory;
 	class FSimulcastEncoderFactory;
 	class FSetSessionDescriptionObserver;
-	class FGPUFencePoller;
+	class FPoller;
 	class IPlayerSession;
 
 	class FStreamer : public FSignallingServerConnectionObserver
@@ -74,8 +74,6 @@ namespace UE::PixelStreaming
 		std::vector<webrtc::RtpEncodingParameters> CreateRTPEncodingParams(int Flags);
 		void SendAnswer(TSharedPtr<IPlayerSession> Session, TUniquePtr<webrtc::SessionDescriptionInterface> Sdp);
 		void OnDataChannelOpen(FPixelStreamingPlayerId PlayerId, webrtc::DataChannelInterface* DataChannel);
-		void OnQualityControllerChanged(FPixelStreamingPlayerId PlayerId);
-		void PostPlayerDeleted(FPixelStreamingPlayerId PlayerId, bool bWasQualityController);
 		void SetLocalDescription(webrtc::PeerConnectionInterface* PeerConnection, FSetSessionDescriptionObserver* Observer, webrtc::SessionDescriptionInterface* SDP);
 		FString GetAudioStreamID() const;
 		FString GetVideoStreamID() const;
@@ -115,9 +113,6 @@ namespace UE::PixelStreaming
 		FPlayerSessions PlayerSessions;
 		FStats Stats;
 
-		FFixedFPSPump PumpThread;
-		FGPUFencePoller FencePollerThread;
-
 		TUniquePtr<webrtc::SessionDescriptionInterface> SFULocalDescription;
 		TUniquePtr<webrtc::SessionDescriptionInterface> SFURemoteDescription;
 	};
diff --git a/Engine/Plugins/Media/PixelStreaming/Source/PixelStreaming/Private/Tests/ClientRTCTests.cpp b/Engine/Plugins/Media/PixelStreaming/Source/PixelStreaming/Private/Tests/ClientRTCTests.cpp
new file mode 100644
--- /dev/null
+++ b/Engine/Plugins/Media/PixelStreaming/Source/PixelStreaming/Private/Tests/ClientRTCTests.cpp
@@ -0,0 +1,198 @@
+// Copyright Epic Games, Inc. All Rights Reserved.
+#include "Misc/AutomationTest.h"
+#include "Tests/AutomationCommon.h"
+#include "ClientRTC.h"
+#include "IPixelStreamingModule.h"
+#include "PixelStreamingPrivate.h"
+
+#if WITH_DEV_AUTOMATION_TESTS
+
+// not using the macro here so we can define all this context and states
+class FClientRTCMessageLatentCommand : public IAutomationLatentCommand
+{
+public:
+	FClientRTCMessageLatentCommand(const FString& InStreamingSignallingUrl, const FString& InClientSignallingUrl)
+		: StreamingSignallingUrl(InStreamingSignallingUrl)
+		, ClientSignallingUrl(InClientSignallingUrl)
+	{
+	}
+
+	virtual ~FClientRTCMessageLatentCommand()
+	{
+		ClientSession.Reset();
+		IPixelStreamingModule& Module = IPixelStreamingModule::Get();
+		Module.StopStreaming();
+	}
+
+	enum class EState
+	{
+		StreamerInit,
+		StreamerWait,
+		StreamerReady,
+
+		ClientWait,
+		ClientReady,
+
+		MessageWait,
+
+		Finished,
+	};
+
+	// return true when we're done
+	bool Update() override
+	{
+		switch (State)
+		{
+			case EState::StreamerInit:
+				return StreamerInit();
+			case EState::StreamerWait:
+				return Wait();
+			case EState::StreamerReady:
+				return StreamerReady();
+			case EState::ClientWait:
+				return Wait();
+			case EState::ClientReady:
+				return ClientReady();
+			case EState::MessageWait:
+				return Wait();
+			default:
+				UE_LOG(LogPixelStreaming, Error, TEXT("Unknown test state."));
+				// fall through
+			case EState::Finished:
+				return true;
+		}
+	}
+
+private:
+	static constexpr uint32 TimeoutStreamerStartup = 5000;
+	static constexpr uint32 TimeoutClientConnect = 5000;
+	static constexpr uint32 TimeoutMessageWait = 5000;
+
+	EState State = EState::StreamerInit;
+
+	FString StreamingSignallingUrl;
+	FString ClientSignallingUrl;
+
+	TSharedPtr<UE::PixelStreaming::FClientRTC> ClientSession;
+
+	FDateTime WaitStartTime;
+	uint32 CurrentTimeout = 0;
+	FString TimeoutMessage;
+
+	bool StreamerInit()
+	{
+		IPixelStreamingModule& Module = IPixelStreamingModule::Get();
+		Module.OnStreamingStarted().AddLambda([&](IPixelStreamingModule&) {
+			State = EState::StreamerReady;
+		});
+		Module.StartStreaming(StreamingSignallingUrl);
+
+		StartTimeout(TimeoutStreamerStartup, "Streaming startup timeout.");
+		State = EState::StreamerWait;
+		return false;
+	}
+
+	bool StreamerReady()
+	{
+		ClientSession = MakeShared<UE::PixelStreaming::FClientRTC>();
+		ClientSession->OnConnected.AddLambda([](UE::PixelStreaming::FClientRTC& Session) {
+		});
+		ClientSession->OnDisconnected.AddLambda([](UE::PixelStreaming::FClientRTC& Session) {
+		});
+		ClientSession->OnDataChannelOpen.AddLambda([&](UE::PixelStreaming::FClientRTC& Session) {
+			State = EState::ClientReady;
+		});
+		ClientSession->Connect(ClientSignallingUrl);
+
+		StartTimeout(TimeoutClientConnect, "Client connect timeout.");
+		State = EState::ClientWait;
+		return false;
+	}
+
+	bool ClientReady()
+	{
+		const FString TestMessage(TEXT("Hello Streamer"));
+
+		ClientSession->OnDataMessage.AddLambda([this, TestMessage](UE::PixelStreaming::FClientRTC&, uint8 Type, const FString& Descriptor) {
+			if (Descriptor != TestMessage)
+			{
+				UE_LOG(LogPixelStreaming, Error, TEXT("Received data message did not match."));
+			}
+			State = EState::Finished;
+		});
+
+		if (!ClientSession->SendMessage(UE::PixelStreaming::Protocol::EToStreamerMsg::TestEcho, TestMessage))
+		{
+			UE_LOG(LogPixelStreaming, Error, TEXT("Failed to send data message."));
+			return true;
+		}
+
+		StartTimeout(TimeoutMessageWait, "Timeout waiting for data message.");
+		State = EState::MessageWait;
+		return false;
+	}
+
+	void StartTimeout(uint32 Timeout, const FString& Message)
+	{
+		WaitStartTime = FDateTime::Now();
+		CurrentTimeout = Timeout;
+		TimeoutMessage = Message;
+	}
+
+	bool Wait()
+	{
+		const FTimespan WaitDelta = FDateTime::Now() - WaitStartTime;
+		if (WaitDelta.GetTotalMilliseconds() > CurrentTimeout)
+		{
+			UE_LOG(LogPixelStreaming, Error, TEXT("%s"), *TimeoutMessage);
+			return true;
+		}
+		return false;
+	}
+};
+
+bool ClientRTCInitTest()
+{
+	std::unique_ptr<UE::PixelStreaming::FClientRTC> client = std::make_unique<UE::PixelStreaming::FClientRTC>();
+	const bool success = client->GetState() == UE::PixelStreaming::FClientRTC::EState::Disconnected;
+	checkf(success, TEXT("ClientRTC initialized to an incorrect state."));
+	return success;
+}
+
+bool ClientRTCMessageTest()
+{
+	const FString ServerSignallingUrl("ws://localhost:8888");
+	const FString ClientSignallingUrl("ws://localhost");
+	ADD_LATENT_AUTOMATION_COMMAND(FClientRTCMessageLatentCommand(ServerSignallingUrl, ClientSignallingUrl));
+
+	return true;
+}
+
+// test name -> test function
+TMap<FString, TFunction<bool()>> ClientRTCTestNames{
+	{ "Init", ClientRTCInitTest },
+	{ "Connect", ClientRTCMessageTest }
+};
+
+// The base test that will just contain all our actual tests
+IMPLEMENT_COMPLEX_AUTOMATION_TEST(FTestClientRTC, "PixelStreaming.ClientRTC", EAutomationTestFlags::ApplicationContextMask | EAutomationTestFlags::EngineFilter)
+void FTestClientRTC::GetTests(TArray<FString>& OutBeautifiedNames, TArray<FString>& OutTestCommands) const
+{
+	for (auto&& [name, func] : ClientRTCTestNames)
+	{
+		OutBeautifiedNames.Add(name);
+		OutTestCommands.Add(name);
+	}
+}
+
+bool FTestClientRTC::RunTest(const FString& Parameters)
+{
+	if (ClientRTCTestNames.Contains(Parameters))
+	{
+		return ClientRTCTestNames[Parameters]();
+	}
+	UE_LOG(LogPixelStreaming, Error, TEXT("Unknown test %S."), *Parameters);
+	return false;
+}
+
+#endif // WITH_DEV_AUTOMATION_TESTS
