This commit is a large collection (from various authors) of Pixel Streaming 5.1 work

diff --git a/Engine/Plugins/Media/PixelStreaming/Source/PixelStreaming/Private/Tests/SignallingServerConnectionTests.cpp b/Engine/Plugins/Media/PixelStreaming/Source/PixelStreaming/Private/Tests/SignallingServerConnectionTests.cpp
new file mode 100644
--- /dev/null
+++ b/Engine/Plugins/Media/PixelStreaming/Source/PixelStreaming/Private/Tests/SignallingServerConnectionTests.cpp
@@ -0,0 +1,291 @@
+// Copyright Epic Games, Inc. All Rights Reserved.
+#include "Misc/AutomationTest.h"
+#include "Tests/AutomationCommon.h"
+#include "SignallingServerConnection.h"
+#include "Serialization/JsonSerializer.h"
+#include "IWebSocket.h"
+
+namespace
+{
+	using namespace UE::PixelStreaming;
+
+	// For validating callbacks from the signalling server connection
+	class FTestSSConnectionObserver : public FSignallingServerConnectionObserver
+	{
+	public:
+		virtual ~FTestSSConnectionObserver() = default;
+
+		virtual void OnSignallingServerConnected() override {}
+		virtual void OnSignallingServerDisconnected() override {}
+
+		virtual void OnConfig(const webrtc::PeerConnectionInterface::RTCConfiguration& Config) override { SetConfig = Config; }
+
+		//// Streamer-only
+		virtual void OnSessionDescription(FPixelStreamingPlayerId PlayerId, webrtc::SdpType Type, const FString& Sdp) override
+		{
+			SetSDP = { Type, Sdp, PlayerId };
+		}
+		virtual void OnRemoteIceCandidate(FPixelStreamingPlayerId PlayerId, const FString& SdpMid, int SdpMLineIndex, const FString& Sdp)
+		{
+			SetIceCandidate = { SdpMid, SdpMLineIndex, Sdp, PlayerId };
+		}
+		//virtual void OnPlayerConnected(FPixelStreamingPlayerId PlayerId, int Flags) { unimplemented(); }
+		//virtual void OnPlayerDisconnected(FPixelStreamingPlayerId PlayerId) { unimplemented(); }
+		//virtual void OnStreamerDataChannels(FPixelStreamingPlayerId SFUId, FPixelStreamingPlayerId PlayerId, int32 SendStreamId, int32 RecvStreamId) { unimplemented(); }
+
+		//// Player-only
+		virtual void OnSessionDescription(webrtc::SdpType Type, const FString& Sdp) override
+		{
+			SetSDP = { Type, Sdp };
+		}
+		virtual void OnRemoteIceCandidate(const FString& SdpMid, int SdpMLineIndex, const FString& Sdp)
+		{
+			SetIceCandidate = { SdpMid, SdpMLineIndex, Sdp };
+		}
+		//virtual void OnPlayerCount(uint32 Count) { unimplemented(); }
+		//virtual void OnPeerDataChannels(int32 SendStreamId, int32 RecvStreamId) { unimplemented(); }
+
+		void Reset()
+		{
+			SetConfig.Reset();
+			SetSDP.Reset();
+			SetIceCandidate.Reset();
+		}
+
+		TOptional<webrtc::PeerConnectionInterface::RTCConfiguration> SetConfig;
+
+		struct FSessionDescriptionData
+		{
+			webrtc::SdpType Type;
+			FString SDP;
+			TOptional<FPixelStreamingPlayerId> PlayerId;
+		};
+		TOptional<FSessionDescriptionData> SetSDP;
+
+		struct FIceCandidateData
+		{
+			FString SdpMid;
+			int SdpMLineIndex;
+			FString SDP;
+			TOptional<FPixelStreamingPlayerId> PlayerId;
+		};
+		TOptional<FIceCandidateData> SetIceCandidate;
+	};
+
+	// For faking a web socket connection
+	class FMockWebSocket : public IWebSocket
+	{
+	public:
+		FMockWebSocket() = default;
+		virtual ~FMockWebSocket() = default;
+		virtual void Connect() override
+		{
+			bConnected = true;
+			OnConnectedEvent.Broadcast();
+		}
+		virtual void Close(int32 Code = 1000, const FString& Reason = FString()) override { bConnected = false; }
+		virtual bool IsConnected() override { return bConnected; }
+		virtual void Send(const FString& Data) override { OnMockResponseEvent.Broadcast(Data); }
+		virtual void Send(const void* Data, SIZE_T Size, bool bIsBinary = false) override {}
+		virtual FWebSocketConnectedEvent& OnConnected() override { return OnConnectedEvent; }
+		virtual FWebSocketConnectionErrorEvent& OnConnectionError() override { return OnErrorEvent; }
+		virtual FWebSocketClosedEvent& OnClosed() override { return OnClosedEvent; }
+		virtual FWebSocketMessageEvent& OnMessage() override { return OnMessageEvent; }
+		virtual FWebSocketRawMessageEvent& OnRawMessage() override { return OnRawMessageEvent; }
+		virtual FWebSocketMessageSentEvent& OnMessageSent() override { return OnMessageSentEvent; }
+
+		void MockSend(const FString& Msg) const { OnMessageEvent.Broadcast(Msg); }
+
+		// used to check outgoing messages
+		DECLARE_EVENT_OneParam(FMockWebSocket, FMockResponseEvent, const FString& /* Message */);
+		FMockResponseEvent OnMockResponseEvent;
+
+	private:
+		FWebSocketConnectedEvent OnConnectedEvent;
+		FWebSocketConnectionErrorEvent OnErrorEvent;
+		FWebSocketClosedEvent OnClosedEvent;
+		FWebSocketMessageEvent OnMessageEvent;
+		FWebSocketRawMessageEvent OnRawMessageEvent;
+		FWebSocketMessageSentEvent OnMessageSentEvent;
+
+		bool bConnected = false;
+	};
+
+	class FMockSessionDescriptionInterface : public webrtc::SessionDescriptionInterface
+	{
+	public:
+		virtual ~FMockSessionDescriptionInterface() = default;
+		virtual cricket::SessionDescription* description() override { return nullptr; }
+		virtual const cricket::SessionDescription* description() const override { return nullptr; }
+		virtual std::string session_id() const override { return ""; }
+		virtual std::string session_version() const override { return ""; }
+		virtual std::string type() const override { return ""; }
+		virtual bool AddCandidate(const webrtc::IceCandidateInterface* candidate) override { return true; }
+		virtual size_t RemoveCandidates(const std::vector<cricket::Candidate>& candidates) override { return 0; }
+		virtual size_t number_of_mediasections() const override { return 0; }
+		virtual const webrtc::IceCandidateCollection* candidates(size_t mediasection_index) const override { return nullptr; }
+		virtual bool ToString(std::string* out) const override { return true; }
+	};
+
+	class FMockIceCandidateInterface : public webrtc::IceCandidateInterface
+	{
+	public:
+		virtual ~FMockIceCandidateInterface() = default;
+		virtual std::string sdp_mid() const override { return ""; }
+		virtual int sdp_mline_index() const override { return 0; }
+		virtual const cricket::Candidate& candidate() const override { return MockCandidate; }
+		virtual std::string server_url() const override { return ""; }
+		virtual bool ToString(std::string* out) const override { return true; }
+
+	private:
+		cricket::Candidate MockCandidate;
+	};
+
+	BEGIN_DEFINE_SPEC(FSignallingServerConnectionSpec, "PixelStreaming.SignallingServerConnection", EAutomationTestFlags::ProductFilter | EAutomationTestFlags::ApplicationContextMask)
+	TUniquePtr<FSignallingServerConnection> SignallingServerConnection; // the object being tested
+	FTestSSConnectionObserver Observer;
+	FSignallingServerConnection::FWebSocketFactory WebSocketFactory;
+	TSharedPtr<FMockWebSocket> SocketEndpoint; // the socket created for the SignallingServerConnection. Used to fake send data.
+	FDelegateHandle OnMessageDelegateHandle; // so we can remove the delegate after tests
+	TSharedPtr<FJsonObject> LastJsonMsg; // the parsed json of the last message received from the SignallingServerConnection
+	const FString ServerID = "MOCK SERVER";
+	const FPixelStreamingPlayerId MockPlayerId = "MOCK";
+	const FMockSessionDescriptionInterface MockSDP;
+	const FMockIceCandidateInterface MockCandidate;
+	void OnMessageReceived(const FString& Response);
+	END_DEFINE_SPEC(FSignallingServerConnectionSpec)
+	void FSignallingServerConnectionSpec::Define()
+	{
+		Describe("Mocked connection tests", [this]() {
+			BeforeEach([this]() {
+				WebSocketFactory = [&](const FString&) {
+					SocketEndpoint = MakeShared<FMockWebSocket>();
+					return SocketEndpoint;
+				};
+				SignallingServerConnection = MakeUnique<FSignallingServerConnection>(WebSocketFactory, Observer, ServerID);
+				SignallingServerConnection->Connect("Fake Url");
+				OnMessageDelegateHandle = SocketEndpoint->OnMockResponseEvent.AddLambda([this](const FString& message) {
+					OnMessageReceived(message);
+				});
+			});
+
+			AfterEach([this]() {
+				LastJsonMsg.Reset();
+				SocketEndpoint->OnMockResponseEvent.Remove(OnMessageDelegateHandle);
+				SocketEndpoint.Reset();
+				SignallingServerConnection.Reset();
+				WebSocketFactory.Reset();
+			});
+
+			It("should be connected", [this]() {
+				TestTrue("Connected", SignallingServerConnection->IsConnected());
+			});
+
+			It("should be able to send an offer", [this]() {
+				// these calls should probably set some data on these mock objects
+				// and then check that the received data is correct, like the mock
+				// sends below
+				SignallingServerConnection->SendOffer(MockPlayerId, MockSDP);
+			});
+
+			It("should be able to send an answer with no player it", [this]() {
+				SignallingServerConnection->SendAnswer(MockSDP);
+			});
+
+			It("should be able to send an answer with a player id", [this]() {
+				SignallingServerConnection->SendAnswer(MockPlayerId, MockSDP);
+			});
+
+			It("should be able to send an ice candidate with no player id", [this]() {
+				SignallingServerConnection->SendIceCandidate(MockCandidate);
+			});
+
+			It("should be able to send an ice candidate with a player id", [this]() {
+				SignallingServerConnection->SendIceCandidate(MockPlayerId, MockCandidate);
+			});
+
+			It("should be able to send a player disconnect message", [this]() {
+				SignallingServerConnection->SendDisconnectPlayer(MockPlayerId, "No Reason");
+			});
+
+			It("should reply to an identity request", [this]() {
+				SocketEndpoint->MockSend(R"({"type" : "identify"})");
+				TestEqual("Type", LastJsonMsg->GetStringField("type"), "endpointId");
+				TestEqual("Id", LastJsonMsg->GetStringField("id"), ServerID);
+			});
+
+			It("should receive configurations", [this]() {
+				SocketEndpoint->MockSend(R"({ "type": "config", "peerConnectionOptions": {} })");
+				TestTrue("Config IsSet", Observer.SetConfig.IsSet());
+			});
+
+			It("should be able to receive offers with no player id", [this]() {
+				SocketEndpoint->MockSend(R"({ "type": "offer", "sdp": "sdp goes here" })");
+				if (TestTrue("SDP IsSet", Observer.SetSDP.IsSet()))
+				{
+					TestEqual("Offer type", Observer.SetSDP.GetValue().Type, webrtc::SdpType::kOffer);
+					TestFalse("Player id IsSet", Observer.SetSDP.GetValue().PlayerId.IsSet());
+					TestEqual("SDP value", Observer.SetSDP.GetValue().SDP, "sdp goes here");
+				}
+			});
+
+			It("should be able to receive offers with a player id", [this]() {
+				SocketEndpoint->MockSend(R"({ "type": "offer", "playerId": "101", "sdp": "sdp goes here" })");
+				if (TestTrue("SDP IsSet", Observer.SetSDP.IsSet()))
+				{
+					TestEqual("Offer type", Observer.SetSDP.GetValue().Type, webrtc::SdpType::kOffer);
+					if (TestTrue("Player id IsSet", Observer.SetSDP.GetValue().PlayerId.IsSet()))
+					{
+						TestEqual("Player id", Observer.SetSDP.GetValue().PlayerId.GetValue(), "101");
+					}
+					TestEqual("SDP value", Observer.SetSDP.GetValue().SDP, "sdp goes here");
+				}
+			});
+
+			It("should be able to receive an answer", [this]() {
+				SocketEndpoint->MockSend(R"({ "type": "answer", "playerId": "102", "sdp": "sdp goes here" })");
+				if (TestTrue("SDP IsSet", Observer.SetSDP.IsSet()))
+				{
+					TestEqual("Answer type", Observer.SetSDP.GetValue().Type, webrtc::SdpType::kAnswer);
+					TestEqual("SDP value", Observer.SetSDP.GetValue().SDP, "sdp goes here");
+					if (TestTrue("Player id IsSet", Observer.SetSDP.GetValue().PlayerId.IsSet()))
+					{
+						TestEqual("Player id", Observer.SetSDP.GetValue().PlayerId.GetValue(), "102");
+					}
+				}
+			});
+
+			It("should be able to receive an ice candidate with no player id", [this]() {
+				SocketEndpoint->MockSend(R"({ "type": "iceCandidate", "candidate": { "sdpMid": "mid", "sdpMlineIndex": 10, "candidate": "sdp string" } })");
+				if (TestTrue("Ice Candidate IsSet", Observer.SetIceCandidate.IsSet()))
+				{
+					TestEqual("sdp mid", Observer.SetIceCandidate.GetValue().SdpMid, "mid");
+					TestEqual("sdp mline index", Observer.SetIceCandidate.GetValue().SdpMLineIndex, 10);
+					TestEqual("sdp value", Observer.SetIceCandidate.GetValue().SDP, "sdp string");
+					TestFalse("PlayerId IsSet", Observer.SetIceCandidate.GetValue().PlayerId.IsSet());
+				}
+			});
+
+			It("should be able to receive an ice candidate with a player id", [this]() {
+				SocketEndpoint->MockSend(R"({ "type": "iceCandidate", "playerId": "103", "candidate": { "sdpMid": "mid", "sdpMlineIndex": 10, "candidate": "sdp string" } })");
+				if (TestTrue("Ice Candidate IsSet", Observer.SetIceCandidate.IsSet()))
+				{
+					TestEqual("sdp mid", Observer.SetIceCandidate.GetValue().SdpMid, "mid");
+					TestEqual("sdp mline index", Observer.SetIceCandidate.GetValue().SdpMLineIndex, 10);
+					TestEqual("sdp value", Observer.SetIceCandidate.GetValue().SDP, "sdp string");
+					if (TestTrue("Player id IsSet", Observer.SetIceCandidate.GetValue().PlayerId.IsSet()))
+					{
+						TestEqual("Player id", Observer.SetIceCandidate.GetValue().PlayerId.GetValue(), "103");
+					}
+				}
+			});
+		});
+	}
+
+	// used to store the parsed json response from the signalling server connection
+	void FSignallingServerConnectionSpec::OnMessageReceived(const FString& Response)
+	{
+		const auto JsonReader = TJsonReaderFactory<TCHAR>::Create(Response);
+		TestTrue("Deserialised json", FJsonSerializer::Deserialize(JsonReader, LastJsonMsg));
+	}
+} // namespace
diff --git a/Engine/Plugins/Media/PixelStreaming/Source/PixelStreaming/Private/TextureSourceBackbuffer.cpp b/Engine/Plugins/Media/PixelStreaming/Source/PixelStreaming/Private/TextureSourceBackbuffer.cpp
new file mode 100644
--- /dev/null
+++ b/Engine/Plugins/Media/PixelStreaming/Source/PixelStreaming/Private/TextureSourceBackbuffer.cpp
@@ -0,0 +1,87 @@
+// Copyright Epic Games, Inc. All Rights Reserved.
+
+#include "TextureSourceBackbuffer.h"
+#include "Settings.h"
+#include "Utils.h"
+
+#include "RHI.h"
+#include "RHIResources.h"
+
+namespace
+{
+	class FBackbufferFrameCapturer : public IPixelStreamingFrameCapturer
+	{
+	public:
+		FGPUFenceRHIRef Fence;
+
+	public:
+		virtual void CaptureTexture(FPixelStreamingTextureWrapper& TextureToCopy, TSharedPtr<FPixelStreamingTextureWrapper> DestinationTexture) override;
+		virtual bool IsCaptureFinished() override;
+		virtual void OnCaptureFinished(TSharedPtr<FPixelStreamingTextureWrapper> CapturedTexture) override;
+	};
+
+	void FBackbufferFrameCapturer::CaptureTexture(FPixelStreamingTextureWrapper& TextureToCopy, TSharedPtr<FPixelStreamingTextureWrapper> DestinationTexture)
+	{
+		UE::PixelStreaming::CopyTextureToRHI(TextureToCopy.GetTexture<FTextureRHIRef>(), DestinationTexture->GetTexture<FTextureRHIRef>(), Fence);
+	}
+
+	bool FBackbufferFrameCapturer::IsCaptureFinished()
+	{
+		return Fence->Poll();
+	}
+
+	void FBackbufferFrameCapturer::OnCaptureFinished(TSharedPtr<FPixelStreamingTextureWrapper> CapturedTexture)
+	{
+		Fence->Clear();
+	}
+} // namespace
+
+namespace UE::PixelStreaming
+{
+
+	/* ----------------- FBackbufferReadyDelegateRouter ----------------- */
+
+	FBackbufferReadyDelegateRouter::FBackbufferReadyDelegateRouter(FTextureSourceBackbuffer* InParent)
+		: Parent(InParent)
+	{
+	}
+
+	void FBackbufferReadyDelegateRouter::RouteOnBackbufferReadyCall(SWindow& SlateWindow, const FTexture2DRHIRef& FrameBuffer)
+	{
+		FPixelStreamingTextureWrapper NewFrame = FPixelStreamingTextureWrapper(FrameBuffer);
+		Parent->OnNewTexture.Broadcast(NewFrame, FrameBuffer->GetDesc().Extent.X, FrameBuffer->GetDesc().Extent.Y);
+	}
+
+	/* ----------------- FTextureSourceBackbuffer ----------------- */
+
+	FTextureSourceBackbuffer::FTextureSourceBackbuffer()
+		// Store this because accessing slate application is difficult during destructor
+		: BackbufferReadyDelegate(FSlateApplication::Get().GetRenderer()->OnBackBufferReadyToPresent())
+		, DelegateRouter(MakeShared<FBackbufferReadyDelegateRouter>(this))
+	{
+		BackbufferReadyHandle = BackbufferReadyDelegate.AddSP(DelegateRouter, &FBackbufferReadyDelegateRouter::RouteOnBackbufferReadyCall);
+	}
+
+	FTextureSourceBackbuffer::~FTextureSourceBackbuffer()
+	{
+		BackbufferReadyDelegate.Remove(BackbufferReadyHandle);
+	}
+
+	TSharedPtr<FPixelStreamingTextureWrapper> FTextureSourceBackbuffer::CreateBlankStagingTexture(uint32 Width, uint32 Height)
+	{
+		return MakeShared<FPixelStreamingTextureWrapper>(UE::PixelStreaming::CreateRHITexture(Width, Height));
+	}
+
+	TSharedPtr<IPixelStreamingFrameCapturer> FTextureSourceBackbuffer::CreateFrameCapturer()
+	{
+		TSharedPtr<FBackbufferFrameCapturer> Capturer = MakeShared<FBackbufferFrameCapturer>();
+		Capturer->Fence = GDynamicRHI->RHICreateGPUFence(TEXT("VideoCapturerCopyFence"));
+		return Capturer;
+	}
+
+	rtc::scoped_refptr<webrtc::I420Buffer> FTextureSourceBackbuffer::ToWebRTCI420Buffer(TSharedPtr<FPixelStreamingTextureWrapper> Texture)
+	{
+		verifyf(false, TEXT("The I420 conversion should not be required for the GPU code path. If this is getting hit something is wrong."));
+		return nullptr;
+	}
+} // namespace UE::PixelStreaming
diff --git a/Engine/Plugins/Media/PixelStreaming/Source/PixelStreaming/Private/TextureSourceBackbuffer.h b/Engine/Plugins/Media/PixelStreaming/Source/PixelStreaming/Private/TextureSourceBackbuffer.h
new file mode 100644
--- /dev/null
+++ b/Engine/Plugins/Media/PixelStreaming/Source/PixelStreaming/Private/TextureSourceBackbuffer.h
@@ -0,0 +1,39 @@
+// Copyright Epic Games, Inc. All Rights Reserved.
+
+#pragma once
+
+#include "PixelStreamingTextureSource.h"
+#include "Rendering/SlateRenderer.h"
+#include "Framework/Application/SlateApplication.h"
+#include "Templates/SharedPointer.h"
+
+namespace UE::PixelStreaming
+{
+	class FBackbufferReadyDelegateRouter
+	{
+	public:
+		FBackbufferReadyDelegateRouter(class FTextureSourceBackbuffer* InParent);
+		void RouteOnBackbufferReadyCall(SWindow& SlateWindow, const FTexture2DRHIRef& FrameBuffer);
+
+	private:
+		class FTextureSourceBackbuffer* Parent = nullptr;
+	};
+
+	/*
+	* Copies from the UE backbuffer into texture that stay fully resident on the GPU.
+	*/
+	class FTextureSourceBackbuffer : public FPixelStreamingTextureSource
+	{
+	private:
+		FSlateRenderer::FOnBackBufferReadyToPresent& BackbufferReadyDelegate;
+		FDelegateHandle BackbufferReadyHandle;
+		TSharedRef<FBackbufferReadyDelegateRouter> DelegateRouter;
+
+	public:
+		FTextureSourceBackbuffer();
+		~FTextureSourceBackbuffer();
+		virtual TSharedPtr<FPixelStreamingTextureWrapper> CreateBlankStagingTexture(uint32 Width, uint32 Height) override;
+		virtual TSharedPtr<IPixelStreamingFrameCapturer> CreateFrameCapturer() override;
+		virtual rtc::scoped_refptr<webrtc::I420Buffer> ToWebRTCI420Buffer(TSharedPtr<FPixelStreamingTextureWrapper> Texture) override;
+	};
+} // namespace UE::PixelStreaming
diff --git a/Engine/Plugins/Media/PixelStreaming/Source/PixelStreaming/Private/TextureSourceCPUI420.cpp b/Engine/Plugins/Media/PixelStreaming/Source/PixelStreaming/Private/TextureSourceCPUI420.cpp
new file mode 100644
--- /dev/null
+++ b/Engine/Plugins/Media/PixelStreaming/Source/PixelStreaming/Private/TextureSourceCPUI420.cpp
@@ -0,0 +1,102 @@
+// Copyright Epic Games, Inc. All Rights Reserved.
+
+#include "TextureSourceCPUI420.h"
+#include "libyuv/convert.h"
+#include "Templates/SharedPointer.h"
+#include "RHI.h"
+#include "RHIResources.h"
+#include "Utils.h"
+
+namespace
+{
+	struct FGPUToCPUReadbackTexture
+	{
+		FGPUToCPUReadbackTexture(FTextureRHIRef InReadbackTexture)
+			: ReadbackTexture(InReadbackTexture)
+		{
+		}
+
+		FTextureRHIRef ReadbackTexture;
+		TArray<FColor> ReadbackResult;
+	};
+
+	class FCPUMemFrameCapturer : public IPixelStreamingFrameCapturer
+	{
+	public:
+		TSharedPtr<FRHIGPUTextureReadback> GPUTextureReadback;
+
+	public:
+		virtual void CaptureTexture(FPixelStreamingTextureWrapper& TextureToCopy, TSharedPtr<FPixelStreamingTextureWrapper> DestinationTexture) override;
+		virtual bool IsCaptureFinished() override;
+		virtual void OnCaptureFinished(TSharedPtr<FPixelStreamingTextureWrapper> CapturedTexture) override;
+	};
+
+	void FCPUMemFrameCapturer::CaptureTexture(FPixelStreamingTextureWrapper& TextureToCopy, TSharedPtr<FPixelStreamingTextureWrapper> DestinationTexture)
+	{
+		FTextureRHIRef SourceTexture = TextureToCopy.GetTexture<FTextureRHIRef>();
+		TSharedRef<FGPUToCPUReadbackTexture> GPUToCPUReadbackTexture = DestinationTexture->GetTexture<TSharedRef<FGPUToCPUReadbackTexture>>();
+		FTextureRHIRef ReadbackTexture = GPUToCPUReadbackTexture->ReadbackTexture;
+		TArray<FColor>& RawPixels = GPUToCPUReadbackTexture->ReadbackResult;
+		UE::PixelStreaming::CopyTextureToReadbackTexture(SourceTexture, GPUTextureReadback, RawPixels.GetData());
+
+	}
+
+	bool FCPUMemFrameCapturer::IsCaptureFinished()
+	{
+		return GPUTextureReadback->IsReady();
+	}
+
+	void FCPUMemFrameCapturer::OnCaptureFinished(TSharedPtr<FPixelStreamingTextureWrapper> CapturedTexture)
+	{
+	}
+} // namespace
+
+namespace UE::PixelStreaming
+{
+	TSharedPtr<FPixelStreamingTextureWrapper> FTextureSourceCPUI420::CreateBlankStagingTexture(uint32 Width, uint32 Height)
+	{
+
+		TSharedRef<FGPUToCPUReadbackTexture> GPUToCPUTexture = MakeShared<FGPUToCPUReadbackTexture>(
+			UE::PixelStreaming::CreateCPUReadbackTexture(Width, Height));
+
+		// Allocate some CPU accessible memory for the texture we intend to read back from the GPU
+		GPUToCPUTexture->ReadbackResult.Init(FColor(), Width * Height * 4);
+
+		return MakeShared<FPixelStreamingTextureWrapper>(GPUToCPUTexture);
+	}
+
+	TSharedPtr<IPixelStreamingFrameCapturer> FTextureSourceCPUI420::CreateFrameCapturer()
+	{
+		TSharedPtr<FCPUMemFrameCapturer> Capturer = MakeShared<FCPUMemFrameCapturer>();
+		Capturer->GPUTextureReadback = MakeShared<FRHIGPUTextureReadback>(TEXT("CopyTextureToReadbackTexture"));
+		return Capturer;
+	}
+
+	rtc::scoped_refptr<webrtc::I420Buffer> FTextureSourceCPUI420::ToWebRTCI420Buffer(TSharedPtr<FPixelStreamingTextureWrapper> Texture)
+	{
+		TSharedRef<FGPUToCPUReadbackTexture> GPUToCPUReadbackTexture = Texture->GetTexture<TSharedRef<FGPUToCPUReadbackTexture>>();
+		FTextureRHIRef& TextureThatWasRead = GPUToCPUReadbackTexture->ReadbackTexture;
+		TArray<FColor>& RawPixels = GPUToCPUReadbackTexture->ReadbackResult;
+
+		uint32 TextureWidth = TextureThatWasRead->GetDesc().Extent.X;
+		uint32 TextureHeight = TextureThatWasRead->GetDesc().Extent.Y;
+
+		const uint8* BGRAIn = (const uint8*)RawPixels.GetData();
+
+		rtc::scoped_refptr<webrtc::I420Buffer> Buffer = webrtc::I420Buffer::Create(TextureWidth, TextureHeight);
+
+		libyuv::ARGBToI420(
+			BGRAIn,
+			TextureWidth * 4,
+			Buffer->MutableDataY(),
+			Buffer->StrideY(),
+			Buffer->MutableDataU(),
+			Buffer->StrideU(),
+			Buffer->MutableDataV(),
+			Buffer->StrideV(),
+			TextureWidth,
+			TextureHeight);
+
+		return Buffer;
+	}
+} // namespace UE::PixelStreaming
diff --git a/Engine/Plugins/Media/PixelStreaming/Source/PixelStreaming/Private/TextureSourceCPUI420.h b/Engine/Plugins/Media/PixelStreaming/Source/PixelStreaming/Private/TextureSourceCPUI420.h
new file mode 100644
--- /dev/null
+++ b/Engine/Plugins/Media/PixelStreaming/Source/PixelStreaming/Private/TextureSourceCPUI420.h
@@ -0,0 +1,18 @@
+// Copyright Epic Games, Inc. All Rights Reserved.
+
+#pragma once
+
+#include "TextureSourceBackbuffer.h"
+
+namespace UE::PixelStreaming
+{
+	/*
+	* Copies from the UE backbuffer into textures that are readback to CPU and system memory.
+	*/
+	class FTextureSourceCPUI420 : public FTextureSourceBackbuffer
+	{
+		virtual TSharedPtr<FPixelStreamingTextureWrapper> CreateBlankStagingTexture(uint32 Width, uint32 Height) override;
+		virtual TSharedPtr<IPixelStreamingFrameCapturer> CreateFrameCapturer() override;
+		virtual rtc::scoped_refptr<webrtc::I420Buffer> ToWebRTCI420Buffer(TSharedPtr<FPixelStreamingTextureWrapper> Texture) override;
+	};
+} // namespace UE::PixelStreaming
diff --git a/Engine/Plugins/Media/PixelStreaming/Source/PixelStreaming/Private/TextureSourceComputeI420.cpp b/Engine/Plugins/Media/PixelStreaming/Source/PixelStreaming/Private/TextureSourceComputeI420.cpp
new file mode 100644
--- /dev/null
+++ b/Engine/Plugins/Media/PixelStreaming/Source/PixelStreaming/Private/TextureSourceComputeI420.cpp
@@ -0,0 +1,224 @@
+// Copyright Epic Games, Inc. All Rights Reserved.
+
+#include "TextureSourceComputeI420.h"
+#include "RGBToYUVShader.h"
+#include "Utils.h"
+
+namespace
+{
+	struct FComputeToI420Texture
+	{
+		FComputeToI420Texture(const FIntPoint& InPlaneYDimensions, const FIntPoint& InPlaneUVDimensions);
+		~FComputeToI420Texture();
+
+		// dimensions of the texures
+		FIntPoint PlaneYDimensions;
+		FIntPoint PlaneUVDimensions;
+
+		// used as targets for the compute shader
+		FTextureRHIRef TextureY;
+		FTextureRHIRef TextureU;
+		FTextureRHIRef TextureV;
+
+		// the UAVs of the targets
+		FUnorderedAccessViewRHIRef TextureYUAV;
+		FUnorderedAccessViewRHIRef TextureUUAV;
+		FUnorderedAccessViewRHIRef TextureVUAV;
+
+		// cpu readable copies of the targets above
+		FTextureRHIRef StagingTextureY;
+		FTextureRHIRef StagingTextureU;
+		FTextureRHIRef StagingTextureV;
+
+		// memory mapped pointers of the staging textures
+		void* MappedY = nullptr;
+		void* MappedU = nullptr;
+		void* MappedV = nullptr;
+
+		int32 YStride = 0;
+		int32 UStride = 0;
+		int32 VStride = 0;
+
+		// the final webrtc buffer
+		rtc::scoped_refptr<webrtc::I420Buffer> Buffer;
+	};
+
+	class FI420FrameCapturer : public IPixelStreamingFrameCapturer
+	{
+	public:
+		FGPUFenceRHIRef Fence;
+		FTextureRHIRef TempFBCopy;
+
+	public:
+		virtual void CaptureTexture(FPixelStreamingTextureWrapper& TextureToCopy, TSharedPtr<FPixelStreamingTextureWrapper> DestinationTexture) override;
+		virtual bool IsCaptureFinished() override;
+		virtual void OnCaptureFinished(TSharedPtr<FPixelStreamingTextureWrapper> CapturedTexture) override;
+	};
+
+	FComputeToI420Texture::FComputeToI420Texture(const FIntPoint& InPlaneYDimensions, const FIntPoint& InPlaneUVDimensions)
+		: PlaneYDimensions(InPlaneYDimensions)
+		, PlaneUVDimensions(InPlaneUVDimensions)
+	{
+		Buffer = webrtc::I420Buffer::Create(PlaneYDimensions.X, PlaneYDimensions.Y);
+
+		FRHITextureCreateDesc TextureDescY = FRHITextureCreateDesc::Create2D(
+			TEXT("Compute YUV Target"),
+			{ PlaneYDimensions.X, PlaneYDimensions.Y },
+			EPixelFormat::PF_R8,
+			FClearValueBinding::None,
+			TexCreate_UAV,
+			1,
+			1,
+			0,
+			ERHIAccess::UAVCompute);
+
+		FRHITextureCreateDesc TextureDescUV = FRHITextureCreateDesc::Create2D(
+			TEXT("Compute YUV Target"),
+			{ PlaneUVDimensions.X, PlaneUVDimensions.Y },
+			EPixelFormat::PF_R8,
+			FClearValueBinding::None,
+			TexCreate_UAV,
+			1,
+			1,
+			0,
+			ERHIAccess::UAVCompute);
+
+		TextureY = GDynamicRHI->RHICreateTexture(TextureDescY);
+		TextureU = GDynamicRHI->RHICreateTexture(TextureDescUV);
+		TextureV = GDynamicRHI->RHICreateTexture(TextureDescUV);
+
+		FRHITextureCreateDesc StagingDescY = FRHITextureCreateDesc::Create2D(
+			TEXT("YUV Output CPU Texture"),
+			{ PlaneYDimensions.X, PlaneYDimensions.Y },
+			EPixelFormat::PF_R8,
+			FClearValueBinding::None,
+			TexCreate_CPUReadback,
+			1,
+			1,
+			0,
+			ERHIAccess::Unknown);
+
+		FRHITextureCreateDesc StagingDescUV = FRHITextureCreateDesc::Create2D(
+			TEXT("YUV Output CPU Texture"),
+			{ PlaneUVDimensions.X, PlaneUVDimensions.Y },
+			EPixelFormat::PF_R8,
+			FClearValueBinding::None,
+			TexCreate_CPUReadback,
+			1,
+			1,
+			0,
+			ERHIAccess::Unknown);
+
+		StagingTextureY = GDynamicRHI->RHICreateTexture(StagingDescY);
+		StagingTextureU = GDynamicRHI->RHICreateTexture(StagingDescUV);
+		StagingTextureV = GDynamicRHI->RHICreateTexture(StagingDescUV);
+
+		TextureYUAV = GDynamicRHI->RHICreateUnorderedAccessView(TextureY, 0, 0, 0);
+		TextureUUAV = GDynamicRHI->RHICreateUnorderedAccessView(TextureU, 0, 0, 0);
+		TextureVUAV = GDynamicRHI->RHICreateUnorderedAccessView(TextureV, 0, 0, 0);
+
+		int32 OutWidth, OutHeight;
+		GDynamicRHI->RHIMapStagingSurface(StagingTextureY, nullptr, MappedY, OutWidth, OutHeight);
+		YStride = OutWidth;
+		GDynamicRHI->RHIMapStagingSurface(StagingTextureU, nullptr, MappedU, OutWidth, OutHeight);
+		UStride = OutWidth;
+		GDynamicRHI->RHIMapStagingSurface(StagingTextureV, nullptr, MappedV, OutWidth, OutHeight);
+		VStride = OutWidth;
+	}
+
+	FComputeToI420Texture::~FComputeToI420Texture()
+	{
+		GDynamicRHI->RHIUnmapStagingSurface(StagingTextureY);
+		GDynamicRHI->RHIUnmapStagingSurface(StagingTextureU);
+		GDynamicRHI->RHIUnmapStagingSurface(StagingTextureV);
+	}
+
+	void FI420FrameCapturer::CaptureTexture(FPixelStreamingTextureWrapper& TextureToCopy, TSharedPtr<FPixelStreamingTextureWrapper> DestinationTexture)
+	{
+		// use a compute shader to extract the yuv planes of the supplied texture
+		// TODO (Aidan) move this to RDG
+		
+		FTextureRHIRef SourceTexture = TextureToCopy.GetTexture<FTextureRHIRef>();
+		TSharedRef<FComputeToI420Texture> I420ReadbackTexture = DestinationTexture->GetTexture<TSharedRef<FComputeToI420Texture>>();
+		FRHICommandListImmediate& RHICmdList = FRHICommandListExecutor::GetImmediateCommandList();
+
+		// Todo (Luke): Add support for changing capture resolutions at runtime using frame scale etc.
+		verifyf(SourceTexture->GetDesc().Extent.X == I420ReadbackTexture->PlaneYDimensions.X && SourceTexture->GetDesc().Extent.Y == I420ReadbackTexture->PlaneYDimensions.Y, TEXT("Does not support resolution changes at runtime yet."));
+
+		if (!TempFBCopy)
+		{
+			FRHITextureCreateDesc Desc = FRHITextureCreateDesc::Create2D(
+				TEXT("PixelStreamingBlankTexture"),
+				{ (int32)SourceTexture->GetDesc().Extent.X, (int32)SourceTexture->GetDesc().Extent.Y },
+				SourceTexture->GetDesc().Format,
+				FClearValueBinding::None,
+				TexCreate_ShaderResource,
+				1,
+				1,
+				0,
+				ERHIAccess::Unknown);
+
+			TempFBCopy = GDynamicRHI->RHICreateTexture(Desc);
+		}
+		RHICmdList.Transition(FRHITransitionInfo(SourceTexture, ERHIAccess::Unknown, ERHIAccess::CopySrc));
+		RHICmdList.Transition(FRHITransitionInfo(TempFBCopy, ERHIAccess::Unknown, ERHIAccess::CopyDest));
+		RHICmdList.CopyTexture(SourceTexture, TempFBCopy, {});
+		RHICmdList.Transition(FRHITransitionInfo(TempFBCopy, ERHIAccess::Unknown, ERHIAccess::SRVMask));
+		
+		FRGBToYUVShaderParameters ShaderParameters;
+		ShaderParameters.SourceTexture = TempFBCopy;
+		ShaderParameters.DestPlaneYDimensions = I420ReadbackTexture->PlaneYDimensions;
+		ShaderParameters.DestPlaneUVDimensions = I420ReadbackTexture->PlaneUVDimensions;
+		ShaderParameters.DestPlaneY = I420ReadbackTexture->TextureYUAV;
+		ShaderParameters.DestPlaneU = I420ReadbackTexture->TextureUUAV;
+		ShaderParameters.DestPlaneV = I420ReadbackTexture->TextureVUAV;
+		FRGBToYUVShader::Dispatch(RHICmdList, ShaderParameters);
+
+		RHICmdList.CopyToResolveTarget(I420ReadbackTexture->TextureY, I420ReadbackTexture->StagingTextureY, {});
+		RHICmdList.CopyToResolveTarget(I420ReadbackTexture->TextureU, I420ReadbackTexture->StagingTextureU, {});
+		RHICmdList.CopyToResolveTarget(I420ReadbackTexture->TextureV, I420ReadbackTexture->StagingTextureV, {});
+
+		RHICmdList.WriteGPUFence(Fence);
+	}
+
+	bool FI420FrameCapturer::IsCaptureFinished()
+	{
+		return Fence->Poll();
+	}
+
+	void FI420FrameCapturer::OnCaptureFinished(TSharedPtr<FPixelStreamingTextureWrapper> CapturedTexture)
+	{
+		// once the compute shader is complete we can copy the data from the pre mapped destination textures
+
+		TSharedRef<FComputeToI420Texture> I420ReadbackTexture = CapturedTexture->GetTexture<TSharedRef<FComputeToI420Texture>>();
+		UE::PixelStreaming::MemCpyStride(I420ReadbackTexture->Buffer->MutableDataY(), I420ReadbackTexture->MappedY, I420ReadbackTexture->Buffer->StrideY(), I420ReadbackTexture->YStride, I420ReadbackTexture->PlaneYDimensions.Y);
+		UE::PixelStreaming::MemCpyStride(I420ReadbackTexture->Buffer->MutableDataU(), I420ReadbackTexture->MappedU, I420ReadbackTexture->Buffer->StrideU(), I420ReadbackTexture->UStride, I420ReadbackTexture->PlaneUVDimensions.Y);
+		UE::PixelStreaming::MemCpyStride(I420ReadbackTexture->Buffer->MutableDataV(), I420ReadbackTexture->MappedV, I420ReadbackTexture->Buffer->StrideV(), I420ReadbackTexture->VStride, I420ReadbackTexture->PlaneUVDimensions.Y);
+		Fence->Clear();
+	}
+} // namespace
+
+namespace UE::PixelStreaming
+{
+	TSharedPtr<FPixelStreamingTextureWrapper> FTextureSourceComputeI420::CreateBlankStagingTexture(uint32 Width, uint32 Height)
+	{
+		const int32 W = static_cast<int32>(Width);
+		const int32 H = static_cast<int32>(Height);
+		TSharedRef<FComputeToI420Texture> Texture = MakeShared<FComputeToI420Texture>(FIntPoint(W, H), FIntPoint(W / 2, H / 2));
+		return MakeShared<FPixelStreamingTextureWrapper>(Texture);
+	}
+
+	TSharedPtr<IPixelStreamingFrameCapturer> FTextureSourceComputeI420::CreateFrameCapturer()
+	{
+		TSharedPtr<FI420FrameCapturer> Capturer = MakeShared<FI420FrameCapturer>();
+		Capturer->Fence = GDynamicRHI->RHICreateGPUFence(TEXT("VideoCapturerCopyFence"));
+		return Capturer;
+	}
+
+	rtc::scoped_refptr<webrtc::I420Buffer> FTextureSourceComputeI420::ToWebRTCI420Buffer(TSharedPtr<FPixelStreamingTextureWrapper> Texture)
+	{
+		TSharedRef<FComputeToI420Texture> I420ReadbackTexture = Texture->GetTexture<TSharedRef<FComputeToI420Texture>>();
+		return I420ReadbackTexture->Buffer;
+	}
+
+} // namespace UE::PixelStreaming
diff --git a/Engine/Plugins/Media/PixelStreaming/Source/PixelStreaming/Private/TextureSourceComputeI420.h b/Engine/Plugins/Media/PixelStreaming/Source/PixelStreaming/Private/TextureSourceComputeI420.h
new file mode 100644
--- /dev/null
+++ b/Engine/Plugins/Media/PixelStreaming/Source/PixelStreaming/Private/TextureSourceComputeI420.h
@@ -0,0 +1,20 @@
+// Copyright Epic Games, Inc. All Rights Reserved.
+
+#pragma once
+
+#include "TextureSourceBackbuffer.h"
+
+namespace UE::PixelStreaming
+{
+	/*
+	* Copies from the UE backbuffer and converts it to I420 using compute shaders
+	*/
+	class FTextureSourceComputeI420 : public FTextureSourceBackbuffer
+	{
+	public:
+		virtual TSharedPtr<FPixelStreamingTextureWrapper> CreateBlankStagingTexture(uint32 Width, uint32 Height) override;
+		virtual TSharedPtr<IPixelStreamingFrameCapturer> CreateFrameCapturer() override;
+		virtual rtc::scoped_refptr<webrtc::I420Buffer> ToWebRTCI420Buffer(TSharedPtr<FPixelStreamingTextureWrapper> Texture) override;
+	};
+
+} // namespace UE::PixelStreaming
diff --git a/Engine/Plugins/Media/PixelStreaming/Source/PixelStreaming/Private/TextureSourceFactory.h b/Engine/Plugins/Media/PixelStreaming/Source/PixelStreaming/Private/TextureSourceFactory.h
new file mode 100644
--- /dev/null
+++ b/Engine/Plugins/Media/PixelStreaming/Source/PixelStreaming/Private/TextureSourceFactory.h
@@ -0,0 +1,80 @@
+// Copyright Epic Games, Inc. All Rights Reserved.
+
+#pragma once
+
+#include "Containers/Map.h"
+#include "GenericPlatform/GenericPlatformMisc.h"
+#include "IPixelStreamingTextureSourceFactory.h"
+#include "TextureSourceBackbuffer.h"
+#include "TextureSourceComputeI420.h"
+#include "TextureSourceCPUI420.h"
+#include "Templates/UniquePtr.h"
+#include "Settings.h"
+
+namespace UE::PixelStreaming
+{
+	/*
+	* A factory where "TextureSourceCreators" are registered as lambdas with unique FNames.
+	* The premise is that outside implementers can register texture source creators using the Pixel Streaming module.
+	*/
+	class FTextureSourceFactory : public IPixelStreamingTextureSourceFactory
+	{
+	public:
+		FTextureSourceFactory() :
+			TextureSourceCreators()
+		{
+			RegisterInternalSources();
+		}
+
+		virtual ~FTextureSourceFactory() override {}
+
+		virtual TUniquePtr<FPixelStreamingTextureSource> CreateTextureSource(FName SourceType) override
+		{
+			verifyf(TextureSourceCreators.Contains(SourceType), TEXT("There was no texture source registered for the type: %s"), *SourceType.ToString())
+			TFunction<TUniquePtr<FPixelStreamingTextureSource>()>& Creator = TextureSourceCreators[SourceType];
+			return Creator();
+		}
+
+		void RegisterInternalSources()
+		{
+			// Backbuffer
+			RegisterTextureSourceType(FName("Backbuffer"), []() {
+				TUniquePtr<FPixelStreamingTextureSource> TextureSource;
+
+				if (Settings::IsCodecVPX())
+				{
+					if (Settings::CVarPixelStreamingVPXUseCompute.GetValueOnAnyThread())
+					{
+						TextureSource = MakeUnique<FTextureSourceComputeI420>();
+					}
+					else
+					{
+						TextureSource = MakeUnique<FTextureSourceCPUI420>();
+					}
+				}
+				else
+				{
+					TextureSource = MakeUnique<FTextureSourceBackbuffer>();
+				}
+
+				return MoveTemp(TextureSource);
+			});
+
+			// Some other internal texture source here...
+		}
+
+		virtual void RegisterTextureSourceType(FName SourceType, TFunction<TUniquePtr<FPixelStreamingTextureSource>()> CreatorFunc) override
+		{
+			TextureSourceCreators.Add(SourceType, CreatorFunc);
+		}
+
+		virtual void UnregisterTextureSourceType(FName SourceType) override
+		{
+			TextureSourceCreators.Remove(SourceType);
+		}
+
+	private:
+		TMap<FName, TFunction<TUniquePtr<FPixelStreamingTextureSource>()>> TextureSourceCreators;
+	};
+
+} // namespace UE::PixelStreaming
\ No newline at end of file
diff --git a/Engine/Plugins/Media/PixelStreaming/Source/PixelStreaming/Private/TextureTripleBuffer.cpp b/Engine/Plugins/Media/PixelStreaming/Source/PixelStreaming/Private/TextureTripleBuffer.cpp
new file mode 100644
--- /dev/null
+++ b/Engine/Plugins/Media/PixelStreaming/Source/PixelStreaming/Private/TextureTripleBuffer.cpp
@@ -0,0 +1,152 @@
+// Copyright Epic Games, Inc. All Rights Reserved.
+
+#pragma once
+#include "TextureTripleBuffer.h"
+#include "Stats.h"
+#include "Utils.h"
+#include "Settings.h"
+#include "Async/Async.h"
+#include "RHI.h"
+#include "Rendering/SlateRenderer.h"
+#include "Framework/Application/SlateApplication.h"
+#include "IPixelStreamingModule.h"
+
+namespace
+{
+	template <typename T>
+	void Swap(TSharedPtr<T>& PtrA, TSharedPtr<T>& PtrB)
+	{
+		TSharedPtr<T> Temp = MoveTemp(PtrA);
+		PtrA = MoveTemp(PtrB);
+		PtrB = MoveTemp(Temp);
+	}
+} // namespace
+
+namespace UE::PixelStreaming
+{
+	FTextureTripleBuffer::FTextureTripleBuffer(float InFrameScale, TUniquePtr<FPixelStreamingTextureSource> InTextureGenerator)
+		: FrameScale(InFrameScale)
+		, TextureGenerator(MoveTemp(InTextureGenerator))
+	{
+		TextureGenerator->OnNewTexture.AddRaw(this, &FTextureTripleBuffer::OnNewTexture);
+	}
+
+	FTextureTripleBuffer::~FTextureTripleBuffer()
+	{
+		*bEnabled = false;
+	}
+
+	TSharedPtr<FPixelStreamingTextureWrapper> FTextureTripleBuffer::GetCurrent()
+	{
+		if (bIsTempDirty)
+		{
+			FScopeLock Lock(&CriticalSection);
+			Swap(ReadBuffer, TempBuffer);
+			bIsTempDirty = false;
+		}
+		return ReadBuffer;
+	}
+
+	void FTextureTripleBuffer::SetEnabled(bool bInEnabled)
+	{
+		*bEnabled = bInEnabled;
+		// This source has been disabled, so set `bInitialized` to false so `OnNewTexture`
+		// will make new textures next time it is called.
+		if (bInitialized && bInEnabled == false)
+		{
+			bInitialized = false;
+		}
+	}
+
+	void FTextureTripleBuffer::Initialize(uint32 Width, uint32 Height)
+	{
+		SourceWidth = Width;
+		SourceHeight = Height;
+
+		WriteBuffer1 = CreateWriteBuffer(SourceWidth, SourceHeight);
+		WriteBuffer2 = CreateWriteBuffer(SourceWidth, SourceHeight);
+		bWriteParity = true;
+
+		TempBuffer = TextureGenerator->CreateBlankStagingTexture(SourceWidth, SourceHeight);
+		ReadBuffer = TextureGenerator->CreateBlankStagingTexture(SourceWidth, SourceHeight);
+		bIsTempDirty = false;
+
+		bInitialized = true;
+	}
+
+	void FTextureTripleBuffer::OnNewTexture(FPixelStreamingTextureWrapper& NewFrame, uint32 FrameWidth, uint32 FrameHeight)
+	{
+		// Note: this is the logic that makes mid-stream resolution changing work.
+		bool bFrameSizeMismatch = (FrameWidth * FrameScale != SourceWidth) || (FrameHeight * FrameScale != SourceHeight);
+
+		if (!bInitialized || bFrameSizeMismatch)
+		{
+			Initialize(FrameWidth * FrameScale, FrameHeight * FrameScale);
+		}
+
+		if (!IsEnabled())
+		{
+			return;
+		}
+
+		TSharedPtr<FWriteBuffer> WriteBuffer = bWriteParity ? WriteBuffer1 : WriteBuffer2;
+		bWriteParity = !bWriteParity;
+
+		// for safety we just make sure that the buffer is not currently waiting for a copy
+		if (WriteBuffer->bAvailable)
+		{
+			WriteBuffer->bAvailable = false;
+
+			if (IsInRenderingThread())
+			{
+				FRHICommandListImmediate& RHICmdList = FRHICommandListExecutor::GetImmediateCommandList();
+				RHICmdList.EnqueueLambda([WriteBuffer](FRHICommandListImmediate& RHICmdList) {
+					if (WriteBuffer.IsValid())
+					{
+						WriteBuffer->PreWaitingOnCopy = FPlatformTime::Cycles64();
+					}
+				});
+			}
+			else
+			{
+				WriteBuffer->PreWaitingOnCopy = FPlatformTime::Cycles64();
+			}
+
+			WriteBuffer->FrameCapturer->CaptureTexture(NewFrame, WriteBuffer->CapturedTexture);
+
+			IPixelStreamingModule::Get().AddPollerTask(
+				[this, WriteBuffer]() {
+					WriteBuffer->FrameCapturer->OnCaptureFinished(WriteBuffer->CapturedTexture);
+					// This lambda is called only once the GPUFence is done
+					{
+						FScopeLock Lock(&CriticalSection);
+						Swap(TempBuffer, WriteBuffer->CapturedTexture);
+						WriteBuffer->bAvailable = true;
+						bIsTempDirty = true;
+					}
+
+					// Capture timing information about the copy operation
+					uint64 PostWaitingOnCopy = FPlatformTime::Cycles64();
+					UE::PixelStreaming::FStats* Stats = UE::PixelStreaming::FStats::Get();
+					if (Stats)
+					{
+						double CaptureLatencyMs = FPlatformTime::ToMilliseconds64(PostWaitingOnCopy - WriteBuffer->PreWaitingOnCopy);
+						Stats->StoreApplicationStat(UE::PixelStreaming::FStatData(FName(*FString::Printf(TEXT("Layer (x%.2f) Capture time (ms)"), FrameScale)), CaptureLatencyMs, 2, true));
+					}
+				},
+				[WriteBuffer]() -> bool {
+					return WriteBuffer.IsValid() ? WriteBuffer->FrameCapturer->IsCaptureFinished() : false;
+				},
+				bEnabled);
+		}
+	}
+
+	TSharedPtr<FWriteBuffer> FTextureTripleBuffer::CreateWriteBuffer(uint32 Width, uint32 Height)
+	{
+		TSharedPtr<FWriteBuffer> Buffer = MakeShared<FWriteBuffer>();
+		Buffer->bAvailable = true;
+		Buffer->FrameCapturer = TextureGenerator->CreateFrameCapturer();
+		Buffer->CapturedTexture = TextureGenerator->CreateBlankStagingTexture(Width, Height);
+		return Buffer;
+	}
+} // namespace UE::PixelStreaming
\ No newline at end of file
diff --git a/Engine/Plugins/Media/PixelStreaming/Source/PixelStreaming/Private/TextureTripleBuffer.h b/Engine/Plugins/Media/PixelStreaming/Source/PixelStreaming/Private/TextureTripleBuffer.h
new file mode 100644
--- /dev/null
+++ b/Engine/Plugins/Media/PixelStreaming/Source/PixelStreaming/Private/TextureTripleBuffer.h
@@ -0,0 +1,73 @@
+// Copyright Epic Games, Inc. All Rights Reserved.
+
+#pragma once
+#include "PixelStreamingTextureSource.h"
+#include "Templates/SharedPointer.h"
+#include "Templates/RefCounting.h"
+#include "HAL/ThreadSafeBool.h"
+
+namespace UE::PixelStreaming
+{
+	struct FWriteBuffer
+	{
+		bool bAvailable = true;
+		uint64 PreWaitingOnCopy;
+		TSharedPtr<IPixelStreamingFrameCapturer> FrameCapturer;
+		TSharedPtr<FPixelStreamingTextureWrapper> CapturedTexture;
+	};
+
+	/*
+	* A triple buffering mechanism for textures being captured and encoded in Pixel Streaming.
+	* Texture copying happens on its own thread.
+	* Encoding happens on its own thread.
+	* Thus we have a multithreaded single producer/single consumer setup and the following requirements:
+	* 1) Reads and writes must be threadsafe
+	* 2) Neither read or write should block
+	* 3) The rate of capture should not impact the rate of encode.
+	* 
+	* We achieve these requirements by using a triple buffering strategy where reading a texture is always fast and non-blocking
+	* by returning the most recently written texture (which may be one we have read before).
+	* The result is we can read textures at a constant speed regardless of how slow writes are, which means we have decoupled render fps from encode fps.
+	*/
+	class FTextureTripleBuffer
+	{
+	public:
+		FTextureTripleBuffer(float InFrameScale, TUniquePtr<FPixelStreamingTextureSource> TextureGenerator);
+		virtual ~FTextureTripleBuffer();
+		virtual TSharedPtr<FPixelStreamingTextureWrapper> GetCurrent();
+		virtual rtc::scoped_refptr<webrtc::I420Buffer> ToWebRTCI420Buffer(TSharedPtr<FPixelStreamingTextureWrapper> Texture) { return TextureGenerator->ToWebRTCI420Buffer(Texture); }
+		virtual void SetEnabled(bool bInEnabled);
+		virtual bool IsEnabled() const { return *bEnabled; }
+		virtual bool IsAvailable() const { return bInitialized; }
+		virtual uint32 GetSourceWidth() const { return SourceWidth; }
+		virtual uint32 GetSourceHeight() const { return SourceHeight; }
+		virtual float GetFrameScaling() const { return FrameScale; }
+
+	protected:
+		virtual void OnNewTexture(FPixelStreamingTextureWrapper& NewFrame, uint32 FrameWidth, uint32 FrameHeight);
+		virtual TSharedPtr<FWriteBuffer> CreateWriteBuffer(uint32 Width, uint32 Height);
+		virtual void Initialize(uint32 Width, uint32 Height);
+
+	private:
+		const float FrameScale;
+		TUniquePtr<FPixelStreamingTextureSource> TextureGenerator;
+		uint32 SourceWidth = 0;
+		uint32 SourceHeight = 0;
+		FThreadSafeBool bInitialized = false;
+		TSharedRef<bool, ESPMode::ThreadSafe> bEnabled = MakeShared<bool, ESPMode::ThreadSafe>(true);
+		bool bWaitForTextureCopy = false;
+		/*
+		* Triple buffer setup with queued write buffers (since we have to wait for RHI copy).
+		* 1 Read buffer (read the captured texture)
+		* 1 Temp buffer (for swapping what is read and written)
+		* 2 Write buffers (2 write buffers because UE can render two frames before presenting sometimes)
+		*/
+		FCriticalSection CriticalSection;
+		bool bWriteParity = true;
+		TSharedPtr<FPixelStreamingTextureWrapper> TempBuffer;
+		TSharedPtr<FPixelStreamingTextureWrapper> ReadBuffer;
+		TSharedPtr<FWriteBuffer> WriteBuffer1;
+		TSharedPtr<FWriteBuffer> WriteBuffer2;
+		FThreadSafeBool bIsTempDirty;
+	};
+}
\ No newline at end of file
diff --git a/Engine/Plugins/Media/PixelStreaming/Source/PixelStreaming/Private/ToStringExtensions.h b/Engine/Plugins/Media/PixelStreaming/Source/PixelStreaming/Private/ToStringExtensions.h
--- a/Engine/Plugins/Media/PixelStreaming/Source/PixelStreaming/Private/ToStringExtensions.h
+++ b/Engine/Plugins/Media/PixelStreaming/Source/PixelStreaming/Private/ToStringExtensions.h
@@ -13,115 +13,115 @@
 #include "Templates/Atomic.h"
 #include "WebRTCIncludes.h"
 
-namespace UE
+namespace UE::PixelStreaming
 {
-	namespace PixelStreaming
+	inline FString ToString(const std::string& Str)
 	{
-		inline FString ToString(const std::string& Str)
-		{
-			return UTF8_TO_TCHAR(Str.c_str());
-		}
+		return UTF8_TO_TCHAR(Str.c_str());
+	}
+
+	inline std::string ToString(const FString& Str)
+	{
+		return TCHAR_TO_UTF8(*Str);
+	}
 
-		inline std::string ToString(const FString& Str)
+	inline FString ToString(const TSharedPtr<FJsonObject>& JsonObj, bool bPretty = true)
+	{
+		FString Res;
+		if (bPretty)
 		{
-			return TCHAR_TO_UTF8(*Str);
+			auto JsonWriter = TJsonWriterFactory<TCHAR, TPrettyJsonPrintPolicy<TCHAR>>::Create(&Res);
+			FJsonSerializer::Serialize(JsonObj.ToSharedRef(), JsonWriter);
 		}
-
-		inline FString ToString(const TSharedPtr<FJsonObject>& JsonObj, bool bPretty = true)
+		else
 		{
-			FString Res;
-			if (bPretty)
-			{
-				auto JsonWriter = TJsonWriterFactory<TCHAR, TPrettyJsonPrintPolicy<TCHAR>>::Create(&Res);
-				FJsonSerializer::Serialize(JsonObj.ToSharedRef(), JsonWriter);
-			}
-			else
-			{
-				auto JsonWriter = TJsonWriterFactory<TCHAR, TCondensedJsonPrintPolicy<TCHAR>>::Create(&Res);
-				FJsonSerializer::Serialize(JsonObj.ToSharedRef(), JsonWriter);
-			}
-			return Res;
+			auto JsonWriter = TJsonWriterFactory<TCHAR, TCondensedJsonPrintPolicy<TCHAR>>::Create(&Res);
+			FJsonSerializer::Serialize(JsonObj.ToSharedRef(), JsonWriter);
 		}
+		return Res;
+	}
 
-		inline const TCHAR* ToString(webrtc::PeerConnectionInterface::SignalingState Val)
-		{
-			TCHAR const* SignallingStatesStr[] = {
-				TEXT("Stable"),
-				TEXT("HaveLocalOffer"),
-				TEXT("HaveLocalPrAnswer"),
-				TEXT("HaveRemoteOffer"),
-				TEXT("HaveRemotePrAnswer"),
-				TEXT("Closed")
-			};
+	inline const TCHAR* ToString(webrtc::PeerConnectionInterface::SignalingState Val)
+	{
+		TCHAR const* SignallingStatesStr[] = {
+			TEXT("Stable"),
+			TEXT("HaveLocalOffer"),
+			TEXT("HaveLocalPrAnswer"),
+			TEXT("HaveRemoteOffer"),
+			TEXT("HaveRemotePrAnswer"),
+			TEXT("Closed")
+		};
 
-			return ensureMsgf(0 <= Val && Val <= webrtc::PeerConnectionInterface::kClosed, TEXT("Invalid `webrtc::PeerConnectionInterface::SignalingState` value: %d"), static_cast<uint32>(Val)) ? SignallingStatesStr[Val] : TEXT("Unknown");
-		}
+		return ensureMsgf(0 <= Val && Val <= webrtc::PeerConnectionInterface::kClosed, TEXT("Invalid `webrtc::PeerConnectionInterface::SignalingState` value: %d"), static_cast<uint32>(Val)) ? SignallingStatesStr[Val] : TEXT("Unknown");
+	}
 
-		inline const TCHAR* ToString(webrtc::PeerConnectionInterface::IceConnectionState Val)
-		{
-			TCHAR const* IceConnectionStatsStr[] = {
-				TEXT("IceConnectionNew"),
-				TEXT("IceConnectionChecking"),
-				TEXT("IceConnectionConnected"),
-				TEXT("IceConnectionCompleted"),
-				TEXT("IceConnectionFailed"),
-				TEXT("IceConnectionDisconnected"),
-				TEXT("IceConnectionClosed")
-			};
+	inline const TCHAR* ToString(webrtc::PeerConnectionInterface::IceConnectionState Val)
+	{
+		TCHAR const* IceConnectionStatsStr[] = {
+			TEXT("IceConnectionNew"),
+			TEXT("IceConnectionChecking"),
+			TEXT("IceConnectionConnected"),
+			TEXT("IceConnectionCompleted"),
+			TEXT("IceConnectionFailed"),
+			TEXT("IceConnectionDisconnected"),
+			TEXT("IceConnectionClosed")
+		};
 
-			return ensureMsgf(
-					   0 <= Val && Val < webrtc::PeerConnectionInterface::kIceConnectionMax,
-					   TEXT("Invalid `webrtc::PeerConnectionInterface::IceConnectionState` value: %d"),
-					   static_cast<uint32>(Val))
+		return ensureMsgf(
+					0 <= Val && Val < webrtc::PeerConnectionInterface::kIceConnectionMax,
+					TEXT("Invalid `webrtc::PeerConnectionInterface::IceConnectionState` value: %d"),
+					static_cast<uint32>(Val)
+				)
 				? IceConnectionStatsStr[Val]
 				: TEXT("Unknown");
 		}
 
-		inline FString ToString(const webrtc::SessionDescriptionInterface& Sdp)
-		{
-			std::string SdpAnsi;
-			verifyf(Sdp.ToString(&SdpAnsi), TEXT("Failed to serialize IceCandidate"));
-			return ToString(SdpAnsi);
-		}
+	inline FString ToString(const webrtc::SessionDescriptionInterface& Sdp)
+	{
+		std::string SdpAnsi;
+		verifyf(Sdp.ToString(&SdpAnsi), TEXT("Failed to serialize IceCandidate"));
+		return ToString(SdpAnsi);
+	}
 
-		inline FString ToString(const webrtc::IceCandidateInterface& IceCandidate)
-		{
-			std::string CandidateAnsi;
-			verifyf(IceCandidate.ToString(&CandidateAnsi), TEXT("Failed to serialize IceCandidate"));
-			return ToString(CandidateAnsi);
-		}
+	inline FString ToString(const webrtc::IceCandidateInterface& IceCandidate)
+	{
+		std::string CandidateAnsi;
+		verifyf(IceCandidate.ToString(&CandidateAnsi), TEXT("Failed to serialize IceCandidate"));
+		return ToString(CandidateAnsi);
+	}
 
-		inline const TCHAR* ToString(webrtc::PeerConnectionInterface::IceGatheringState Val)
-		{
-			TCHAR const* IceGatheringStatsStr[] = {
-				TEXT("IceGatheringNew"),
-				TEXT("IceGatheringGathering"),
-				TEXT("IceGatheringComplete")
-			};
+	inline const TCHAR* ToString(webrtc::PeerConnectionInterface::IceGatheringState Val)
+	{
+		TCHAR const* IceGatheringStatsStr[] = {
+			TEXT("IceGatheringNew"),
+			TEXT("IceGatheringGathering"),
+			TEXT("IceGatheringComplete")
+		};
 
-			return ensureMsgf(
-					   0 <= Val && Val <= webrtc::PeerConnectionInterface::kIceGatheringComplete,
-					   TEXT("Invalid `webrtc::PeerConnectionInterface::IceGatheringState` value: %d"),
-					   static_cast<uint32>(Val))
+		return ensureMsgf(
+					0 <= Val && Val <= webrtc::PeerConnectionInterface::kIceGatheringComplete,
+						TEXT("Invalid `webrtc::PeerConnectionInterface::IceGatheringState` value: %d"),
+						static_cast<uint32>(Val)
+				)
 				? IceGatheringStatsStr[Val]
 				: TEXT("Unknown");
-		}
+	}
 
-		inline const TCHAR* ToString(webrtc::VideoFrameType FrameType)
-		{
-			TCHAR const* FrameTypesStr[] = {
-				TEXT("EmptyFrame"),
-				TEXT("AudioFrameSpeech"),
-				TEXT("AudioFrameCN"),
-				TEXT("VideoFrameKey"),
-				TEXT("VideoFrameDelta")
-			};
-			int FrameTypeInt = (int)FrameType;
-			return ensureMsgf(
-					   0 <= FrameTypeInt && FrameTypeInt <= (int)webrtc::VideoFrameType::kVideoFrameDelta,
-					   TEXT("Invalid `webrtc::FrameType`: %d"), static_cast<uint32>(FrameType))
+	inline const TCHAR* ToString(webrtc::VideoFrameType FrameType)
+	{
+		TCHAR const* FrameTypesStr[] = {
+			TEXT("EmptyFrame"),
+			TEXT("AudioFrameSpeech"),
+			TEXT("AudioFrameCN"),
+			TEXT("VideoFrameKey"),
+			TEXT("VideoFrameDelta")
+		};
+		int FrameTypeInt = (int)FrameType;
+		return ensureMsgf(
+					0 <= FrameTypeInt && FrameTypeInt <= (int)webrtc::VideoFrameType::kVideoFrameDelta,
+					TEXT("Invalid `webrtc::FrameType`: %d"), static_cast<uint32>(FrameType)
+				)
 				? FrameTypesStr[FrameTypeInt]
 				: TEXT("Unknown");
-		}
-	} // namespace PixelStreaming
-} // namespace UE
+	}
+} // namespace UE::PixelStreaming
diff --git a/Engine/Plugins/Media/PixelStreaming/Source/PixelStreaming/Private/Utils.h b/Engine/Plugins/Media/PixelStreaming/Source/PixelStreaming/Private/Utils.h
--- a/Engine/Plugins/Media/PixelStreaming/Source/PixelStreaming/Private/Utils.h
+++ b/Engine/Plugins/Media/PixelStreaming/Source/PixelStreaming/Private/Utils.h
@@ -9,9 +9,47 @@
 #include "Modules/ModuleManager.h"
 #include "PixelStreamingPrivate.h"
 #include "WebRTCIncludes.h"
+#include "Async/Async.h"
+#include "RenderGraphBuilder.h"
+#include "RenderGraphUtils.h"
+#include "RHIGPUReadback.h"
+#include "PixelShaderUtils.h"
+#include "Runtime/Renderer/Private/ScreenPass.h"
 
 namespace UE::PixelStreaming
 {
+	template <typename T>
+	void DoOnGameThread(T&& Func)
+	{
+		if (IsInGameThread())
+		{
+			Func();
+		}
+		else
+		{
+			AsyncTask(ENamedThreads::GameThread, [Func]() { Func(); });
+		}
+	}
+
+	template <typename T>
+	void DoOnGameThreadAndWait(uint32 Timeout, T&& Func)
+	{
+		if (IsInGameThread())
+		{
+			Func();
+		}
+		else
+		{
+			FEvent* TaskEvent = FPlatformProcess::GetSynchEventFromPool();
+			AsyncTask(ENamedThreads::GameThread, [Func, TaskEvent]() {
+				Func();
+				TaskEvent->Trigger();
+			});
+			TaskEvent->Wait(Timeout);
+			FPlatformProcess::ReturnSynchEventToPool(TaskEvent);
+		}
+	}
+
 	inline webrtc::SdpVideoFormat CreateH264Format(webrtc::H264::Profile profile, webrtc::H264::Level level)
 	{
 		const absl::optional<std::string> ProfileString =
@@ -24,7 +62,82 @@ namespace UE::PixelStreaming
 				{ cricket::kH264FmtpPacketizationMode, "1" } });
 	}
 
-	inline void CopyTexture(FTexture2DRHIRef SourceTexture, FTexture2DRHIRef DestinationTexture, FGPUFenceRHIRef& CopyFence)
+	inline void MemCpyStride(void* Dest, const void* Src, size_t DestStride, size_t SrcStride, size_t Height)
+	{
+		char* DestPtr = static_cast<char*>(Dest);
+		const char* SrcPtr = static_cast<const char*>(Src);
+		size_t Row = Height;
+		while (Row--)
+		{
+			FMemory::Memcpy(DestPtr + DestStride * Row, SrcPtr + SrcStride * Row, DestStride);
+		}
+	}
+
+	inline FTextureRHIRef CreateRHITexture(uint32 Width, uint32 Height)
+	{
+		// Create empty texture
+		ETextureCreateFlags Flags;
+		ERHIAccess InitialState;
+
+		if (RHIGetInterfaceType() == ERHIInterfaceType::Vulkan)
+		{
+			Flags = TexCreate_RenderTargetable | TexCreate_External;
+			InitialState = ERHIAccess::Present;
+		}
+		else
+		{
+			Flags = TexCreate_RenderTargetable | TexCreate_Shared;
+			InitialState = ERHIAccess::CopyDest;
+		}
+
+		FRHITextureCreateDesc Desc = FRHITextureCreateDesc::Create2D(
+			TEXT("PixelStreamingBlankTexture"),
+			{ (int32)Width, (int32)Height },
+			EPixelFormat::PF_B8G8R8A8,
+			FClearValueBinding::None,
+			Flags,
+			1,
+			1,
+			0,
+			InitialState
+		);
+
+		return GDynamicRHI->RHICreateTexture(Desc);
+	}
+
+	inline FTextureRHIRef CreateCPUReadbackTexture(uint32 Width, uint32 Height)
+	{
+		// Create empty texture
+		ETextureCreateFlags Flags;
+		ERHIAccess InitialState;
+
+		if (RHIGetInterfaceType() == ERHIInterfaceType::Vulkan)
+		{
+			Flags = TexCreate_HideInVisualizeTexture | TexCreate_CPUReadback;
+			InitialState = ERHIAccess::Present;
+		}
+		else
+		{
+			Flags = TexCreate_HideInVisualizeTexture | TexCreate_CPUReadback;
+			InitialState = ERHIAccess::CopyDest;
+		}
+
+		FRHITextureCreateDesc Desc = FRHITextureCreateDesc::Create2D(
+			TEXT("PixelStreamingBlankTexture"),
+			{ (int32)Width, (int32)Height },
+			EPixelFormat::PF_B8G8R8A8,
+			FClearValueBinding::None,
+			Flags,
+			1,
+			1,
+			0,
+			InitialState
+		);
+
+		return GDynamicRHI->RHICreateTexture(Desc);
+	}
+
+	inline void CopyTextureToRHI(FTextureRHIRef SourceTexture, FTextureRHIRef DestinationTexture, FGPUFenceRHIRef& CopyFence)
 	{
 		FRHICommandListImmediate& RHICmdList = FRHICommandListExecutor::GetImmediateCommandList();
 
@@ -36,7 +149,7 @@ namespace UE::PixelStreaming
 		RHICmdList.BeginRenderPass(RPInfo, TEXT("CopyBackbuffer"));
 
 		{
-			RHICmdList.SetViewport(0, 0, 0.0f, DestinationTexture->GetSizeX(), DestinationTexture->GetSizeY(), 1.0f);
+			RHICmdList.SetViewport(0, 0, 0.0f, DestinationTexture->GetDesc().Extent.X, DestinationTexture->GetDesc().Extent.Y, 1.0f);
 
 			FGraphicsPipelineStateInitializer GraphicsPSOInit;
 			RHICmdList.ApplyCachedRenderTargets(GraphicsPSOInit);
@@ -57,7 +170,7 @@ namespace UE::PixelStreaming
 
 			SetGraphicsPipelineState(RHICmdList, GraphicsPSOInit, 0);
 
-			if (DestinationTexture->GetSizeX() != SourceTexture->GetSizeX() || DestinationTexture->GetSizeY() != SourceTexture->GetSizeY())
+			if (DestinationTexture->GetDesc().Extent.X != SourceTexture->GetDesc().Extent.X || DestinationTexture->GetDesc().Extent.Y != SourceTexture->GetDesc().Extent.Y)
 			{
 				PixelShader->SetParameters(RHICmdList, TStaticSamplerState<SF_Bilinear>::GetRHI(), SourceTexture);
 			}
@@ -66,12 +179,13 @@ namespace UE::PixelStreaming
 				PixelShader->SetParameters(RHICmdList, TStaticSamplerState<SF_Point>::GetRHI(), SourceTexture);
 			}
 
+			FIntPoint TargetBufferSize(DestinationTexture->GetDesc().Extent.X, DestinationTexture->GetDesc().Extent.Y);
 			RendererModule->DrawRectangle(RHICmdList, 0, 0, // Dest X, Y
-				DestinationTexture->GetSizeX(),				// Dest Width
-				DestinationTexture->GetSizeY(),				// Dest Height
+				DestinationTexture->GetDesc().Extent.X,		// Dest Width
+				DestinationTexture->GetDesc().Extent.Y,		// Dest Height
 				0, 0,										// Source U, V
 				1, 1,										// Source USize, VSize
-				DestinationTexture->GetSizeXY(),			// Target buffer size
+				TargetBufferSize,							// Target buffer size
 				FIntPoint(1, 1),							// Source texture size
 				VertexShader, EDRF_Default);
 		}
@@ -83,43 +197,94 @@ namespace UE::PixelStreaming
 		RHICmdList.ImmediateFlush(EImmediateFlushType::WaitForOutstandingTasksOnly);
 	}
 
-	inline FTexture2DRHIRef CreateTexture(uint32 Width, uint32 Height)
+	inline void CopyTextureToReadbackTexture(FTextureRHIRef SourceTexture, TSharedPtr<FRHIGPUTextureReadback> GPUTextureReadback, void* OutBuffer)
 	{
-		// Create empty texture
-		ETextureCreateFlags Flags;
-		ERHIAccess InitialState;
+		FRHICommandListImmediate& RHICmdList = FRHICommandListExecutor::GetImmediateCommandList();
+		FRDGBuilder GraphBuilder(RHICmdList);
 
+		FRDGTextureRef RDGSourceTexture = RegisterExternalTexture(GraphBuilder, SourceTexture, TEXT("SourceCopyTextureToReadbackTexture"));
+		FRDGTextureRef RDGStagingTexture = nullptr;
 
-		if (RHIGetInterfaceType() == ERHIInterfaceType::Vulkan)
+		if(RDGSourceTexture->Desc.Format != EPixelFormat::PF_B8G8R8A8)
 		{
-			Flags = TexCreate_RenderTargetable | TexCreate_External;
-			InitialState = ERHIAccess::Present;
+			// We need the pixel format to be BGRA8 so we first draw it to a staging texture
+			{
+				const FRDGTextureDesc Desc(FRDGTextureDesc::Create2D(RDGSourceTexture->Desc.Extent, PF_B8G8R8A8, FClearValueBinding::None, TexCreate_ShaderResource | TexCreate_RenderTargetable));
+				RDGStagingTexture = GraphBuilder.CreateTexture(Desc, TEXT("StagingCopyTextureToReadbackTexture"));
+			}
+
+			FGlobalShaderMap* ShaderMap = GetGlobalShaderMap(GMaxRHIFeatureLevel);
+
+			TShaderMapRef<FScreenPassVS> VertexShader(ShaderMap);
+
+			// Setup the pixel shader
+			TShaderMapRef<FCopyRectPS> PixelShader(ShaderMap);
+
+			FCopyRectPS::FParameters* PixelShaderParameters = GraphBuilder.AllocParameters<FCopyRectPS::FParameters>();
+			PixelShaderParameters->InputTexture = RDGSourceTexture;
+			PixelShaderParameters->InputSampler = TStaticSamplerState<SF_Bilinear, AM_Clamp, AM_Clamp, AM_Clamp>::GetRHI();
+			PixelShaderParameters->RenderTargets[0] = FRenderTargetBinding(RDGStagingTexture, ERenderTargetLoadAction::ELoad);
+
+			ClearUnusedGraphResources(PixelShader, PixelShaderParameters);
+
+			// We are not doing any clever blending stuff so we just use defaults here
+			FRHIBlendState* BlendState = FScreenPassPipelineState::FDefaultBlendState::GetRHI();
+			FRHIDepthStencilState* DepthStencilState = FScreenPassPipelineState::FDefaultDepthStencilState::GetRHI();
+
+			// Create the pipline state that will execute
+			const FScreenPassPipelineState PipelineState(VertexShader, PixelShader, BlendState, DepthStencilState);
+
+			// Add the pass the the graph builder
+			GraphBuilder.AddPass(
+				RDG_EVENT_NAME("PixelStreamingChangePixelFormat"),
+				PixelShaderParameters,
+				ERDGPassFlags::Raster,
+				[PipelineState, Extent = RDGSourceTexture->Desc.Extent, PixelShader, PixelShaderParameters](FRHICommandList& RHICmdList)
+				{
+					PipelineState.Validate();
+					
+					RHICmdList.SetViewport(0.0f, 0.0f, 0.0f, Extent.X, Extent.Y, 1.0f);
+
+					SetScreenPassPipelineState(RHICmdList, PipelineState);
+
+					SetShaderParameters(RHICmdList, PixelShader, PixelShader.GetPixelShader(), *PixelShaderParameters);
+
+					DrawRectangle(
+						RHICmdList,
+						0, 0, Extent.X, Extent.Y,
+						0, 0, Extent.X, Extent.Y,
+						Extent,
+						Extent,
+						PipelineState.VertexShader,
+						EDRF_UseTriangleOptimization);
+				});
 		}
 		else
 		{
-			Flags = TexCreate_RenderTargetable | TexCreate_Shared;
-			InitialState = ERHIAccess::CopyDest;
+			// Otherwise if the PixelFormat is already BGRA8 we can just use the SourceTexture as the staging texture
+			RDGStagingTexture = RDGSourceTexture;
 		}
 
-		FRHITextureCreateDesc Desc = FRHITextureCreateDesc::Create2D(
-			TEXT("PixelStreamingBlankTexture"),
-			{ (int32)Width, (int32)Height },
-			EPixelFormat::PF_B8G8R8A8,
-			FClearValueBinding::None,
-			Flags,
-			1,
-			1,
-			0,
-			InitialState
-		);
+		// Do the copy from staging RBGA8 texture into the readback
+		AddEnqueueCopyPass(GraphBuilder, GPUTextureReadback.Get(), RDGStagingTexture);
 
-		return GDynamicRHI->RHICreateTexture(Desc);
-	}
+		GraphBuilder.Execute();
 
-	inline void ReadTextureToCPU(FRHICommandListImmediate& RHICmdList, FTexture2DRHIRef& TextureRef, TArray<FColor>& OutPixels)
-	{
-		FIntRect Rect(0, 0, TextureRef->GetSizeX(), TextureRef->GetSizeY());
-		RHICmdList.ReadSurfaceData(TextureRef, Rect, OutPixels, FReadSurfaceDataFlags());
+		// Lock and copy out the content of the TextureReadback to the CPU
+		int32 OutRowPitchInPixels;
+		int32 BlockSize = GPixelFormats[RDGStagingTexture->Desc.Format].BlockBytes;
+		void* ResultsBuffer = GPUTextureReadback->Lock(OutRowPitchInPixels);
+		if (RDGSourceTexture->Desc.Extent.X == OutRowPitchInPixels)
+		{
+			// Source pixel width is the same as the stride of the result buffer (ie no padding), we can do a plain memcpy
+			FPlatformMemory::Memcpy(OutBuffer, ResultsBuffer, (RDGSourceTexture->Desc.Extent.X * RDGSourceTexture->Desc.Extent.Y * BlockSize));
+		} 
+		else
+		{
+			// Source pixel width differs from the stride of the result buffer (ie padding), do a memcpy that accounts for this
+			MemCpyStride(OutBuffer, ResultsBuffer, RDGSourceTexture->Desc.Extent.X * BlockSize, OutRowPitchInPixels * BlockSize, RDGSourceTexture->Desc.Extent.Y); 
+		}
+		GPUTextureReadback->Unlock();
 	}
 
 	inline size_t SerializeToBuffer(rtc::CopyOnWriteBuffer& Buffer, size_t Pos, const void* Data, size_t DataSize)
@@ -127,5 +292,4 @@ namespace UE::PixelStreaming
 		FMemory::Memcpy(&Buffer[Pos], reinterpret_cast<const uint8_t*>(Data), DataSize);
 		return Pos + DataSize;
 	}
-
-} // namespace UE::PixelStreaming
+} // namespace UE::PixelStreaming
\ No newline at end of file
diff --git a/Engine/Plugins/Media/PixelStreaming/Source/PixelStreaming/Private/VideoEncoderH264Wrapper.cpp b/Engine/Plugins/Media/PixelStreaming/Source/PixelStreaming/Private/VideoEncoderH264Wrapper.cpp
--- a/Engine/Plugins/Media/PixelStreaming/Source/PixelStreaming/Private/VideoEncoderH264Wrapper.cpp
+++ b/Engine/Plugins/Media/PixelStreaming/Source/PixelStreaming/Private/VideoEncoderH264Wrapper.cpp
@@ -4,33 +4,36 @@
 #include "Stats.h"
 #include "VideoEncoderFactory.h"
 #include "EncoderFactory.h"
-#include "PixelStreamingFrameBuffer.h"
+#include "FrameBuffer.h"
 
-UE::PixelStreaming::FVideoEncoderH264Wrapper::FVideoEncoderH264Wrapper(TUniquePtr<FEncoderFrameFactory> InFrameFactory, TUniquePtr<AVEncoder::FVideoEncoder> InEncoder)
-	: FrameFactory(MoveTemp(InFrameFactory))
-	, Encoder(MoveTemp(InEncoder)){
-		checkf(Encoder, TEXT("Encoder is nullptr."))
-	}
-
-	UE::PixelStreaming::FVideoEncoderH264Wrapper::~FVideoEncoderH264Wrapper()
-{
-	Encoder->ClearOnEncodedPacket();
-	Encoder->Shutdown();
-}
-
-void UE::PixelStreaming::FVideoEncoderH264Wrapper::Encode(const webrtc::VideoFrame& WebRTCFrame, bool bKeyframe)
+namespace UE::PixelStreaming
 {
-	FTexture2DRHIRef SourceTexture;
+	FVideoEncoderH264Wrapper::FVideoEncoderH264Wrapper(TUniquePtr<FEncoderFrameFactory> InFrameFactory, TUniquePtr<AVEncoder::FVideoEncoder> InEncoder)
+		: FrameFactory(MoveTemp(InFrameFactory))
+		, Encoder(MoveTemp(InEncoder)){
+			checkf(Encoder, TEXT("Encoder is nullptr."))
+		}
 
-	FPixelStreamingFrameBuffer* FrameBuffer = static_cast<FPixelStreamingFrameBuffer*>(WebRTCFrame.video_frame_buffer().get());
-	check(FrameBuffer->GetFrameBufferType() == UE::PixelStreaming::EFrameBufferType::Layer);
-	FLayerFrameBuffer* LayerFrameBuffer = static_cast<FLayerFrameBuffer*>(FrameBuffer);
-	SourceTexture = LayerFrameBuffer->GetFrame();
+		FVideoEncoderH264Wrapper::~FVideoEncoderH264Wrapper()
+	{
+		Encoder->ClearOnEncodedPacket();
+		Encoder->Shutdown();
+	}
 
-	if (SourceTexture)
+	void FVideoEncoderH264Wrapper::Encode(const webrtc::VideoFrame& WebRTCFrame, bool bKeyframe)
 	{
+		FPixelStreamingFrameBuffer* FrameBuffer = static_cast<FPixelStreamingFrameBuffer*>(WebRTCFrame.video_frame_buffer().get());
+		check(FrameBuffer->GetFrameBufferType() == EPixelStreamingFrameBufferType::Layer);
+		FLayerFrameBuffer* LayerFrameBuffer = static_cast<FLayerFrameBuffer*>(FrameBuffer);
+		TSharedPtr<FPixelStreamingTextureWrapper> SourceTexture = LayerFrameBuffer->GetFrame();
+
+		if (!SourceTexture.IsValid())
+		{
+			return;
+		}
+
 		AVEncoder::FVideoEncoder::FLayerConfig EncoderConfig = GetCurrentConfig();
-		TSharedPtr<AVEncoder::FVideoEncoderInputFrame> EncoderInputFrame = FrameFactory->GetFrameAndSetTexture(SourceTexture);
+		TSharedPtr<AVEncoder::FVideoEncoderInputFrame> EncoderInputFrame = FrameFactory->GetFrameAndSetTexture(SourceTexture->GetTexture<FTextureRHIRef>());
 		if (EncoderInputFrame)
 		{
 			EncoderInputFrame->SetTimestampUs(WebRTCFrame.timestamp_us());
@@ -44,72 +47,72 @@ void UE::PixelStreaming::FVideoEncoderH264Wrapper::Encode(const webrtc::VideoFra
 			Encoder->Encode(EncoderInputFrame, Options);
 		}
 	}
-}
 
-AVEncoder::FVideoEncoder::FLayerConfig UE::PixelStreaming::FVideoEncoderH264Wrapper::GetCurrentConfig()
-{
-	checkf(Encoder, TEXT("Cannot request layer config when encoder is nullptr."));
+	AVEncoder::FVideoEncoder::FLayerConfig FVideoEncoderH264Wrapper::GetCurrentConfig()
+	{
+		checkf(Encoder, TEXT("Cannot request layer config when encoder is nullptr."));
 
-	// Asume user wants config for layer zero.
-	return Encoder->GetLayerConfig(0);
-}
+		// Asume user wants config for layer zero.
+		return Encoder->GetLayerConfig(0);
+	}
 
-void UE::PixelStreaming::FVideoEncoderH264Wrapper::SetConfig(const AVEncoder::FVideoEncoder::FLayerConfig& NewConfig)
-{
-	checkf(Encoder, TEXT("Cannot set layer config when encoder is nullptr."));
-	AVEncoder::FVideoEncoder::FLayerConfig CurrentConfig = GetCurrentConfig();
-	// Assumer user wants to update layer zero.
-	if (NewConfig != CurrentConfig)
+	void FVideoEncoderH264Wrapper::SetConfig(const AVEncoder::FVideoEncoder::FLayerConfig& NewConfig)
 	{
-		if (Encoder)
+		checkf(Encoder, TEXT("Cannot set layer config when encoder is nullptr."));
+		AVEncoder::FVideoEncoder::FLayerConfig CurrentConfig = GetCurrentConfig();
+		// Assumer user wants to update layer zero.
+		if (NewConfig != CurrentConfig)
 		{
-			Encoder->UpdateLayerConfig(0, NewConfig);
+			if (Encoder)
+			{
+				Encoder->UpdateLayerConfig(0, NewConfig);
+			}
 		}
 	}
-}
 
-/* ------------------ Static functions below --------------------- */
+	/* ------------------ Static functions below --------------------- */
 
-void UE::PixelStreaming::FVideoEncoderH264Wrapper::OnEncodedPacket(UE::PixelStreaming::FVideoEncoderFactory* Factory, uint32 InLayerIndex, const TSharedPtr<AVEncoder::FVideoEncoderInputFrame> InFrame, const AVEncoder::FCodecPacket& InPacket)
-{
-	webrtc::EncodedImage Image;
+	void FVideoEncoderH264Wrapper::OnEncodedPacket(FVideoEncoderFactory* Factory, uint32 InLayerIndex, const TSharedPtr<AVEncoder::FVideoEncoderInputFrame> InFrame, const AVEncoder::FCodecPacket& InPacket)
+	{
+		webrtc::EncodedImage Image;
 
-	webrtc::RTPFragmentationHeader FragHeader;
-	//CreateH264FragmentHeader(InPacket.Data.Get(), InPacket.DataSize, FragHeader);
+		webrtc::RTPFragmentationHeader FragHeader;
+		//CreateH264FragmentHeader(InPacket.Data.Get(), InPacket.DataSize, FragHeader);
 
-	std::vector<webrtc::H264::NaluIndex> NALUIndices = webrtc::H264::FindNaluIndices(InPacket.Data.Get(), InPacket.DataSize);
-	FragHeader.VerifyAndAllocateFragmentationHeader(NALUIndices.size());
-	FragHeader.fragmentationVectorSize = static_cast<uint16_t>(NALUIndices.size());
-	for (int i = 0; i != NALUIndices.size(); ++i)
-	{
-		webrtc::H264::NaluIndex const& NALUIndex = NALUIndices[i];
-		FragHeader.fragmentationOffset[i] = NALUIndex.payload_start_offset;
-		FragHeader.fragmentationLength[i] = NALUIndex.payload_size;
-	}
+		std::vector<webrtc::H264::NaluIndex> NALUIndices = webrtc::H264::FindNaluIndices(InPacket.Data.Get(), InPacket.DataSize);
+		FragHeader.VerifyAndAllocateFragmentationHeader(NALUIndices.size());
+		FragHeader.fragmentationVectorSize = static_cast<uint16_t>(NALUIndices.size());
+		for (int i = 0; i != NALUIndices.size(); ++i)
+		{
+			webrtc::H264::NaluIndex const& NALUIndex = NALUIndices[i];
+			FragHeader.fragmentationOffset[i] = NALUIndex.payload_start_offset;
+			FragHeader.fragmentationLength[i] = NALUIndex.payload_size;
+		}
 
-	Image.timing_.packetization_finish_ms = FTimespan::FromSeconds(FPlatformTime::Seconds()).GetTotalMilliseconds();
-	Image.timing_.encode_start_ms = InPacket.Timings.StartTs.GetTotalMilliseconds();
-	Image.timing_.encode_finish_ms = InPacket.Timings.FinishTs.GetTotalMilliseconds();
-	Image.timing_.flags = webrtc::VideoSendTiming::kTriggeredByTimer;
-
-	Image.SetEncodedData(webrtc::EncodedImageBuffer::Create(const_cast<uint8_t*>(InPacket.Data.Get()), InPacket.DataSize));
-	Image._encodedWidth = InFrame->GetWidth();
-	Image._encodedHeight = InFrame->GetHeight();
-	Image._frameType = InPacket.IsKeyFrame ? webrtc::VideoFrameType::kVideoFrameKey : webrtc::VideoFrameType::kVideoFrameDelta;
-	Image.content_type_ = webrtc::VideoContentType::UNSPECIFIED;
-	Image.qp_ = InPacket.VideoQP;
-	Image.SetSpatialIndex(InLayerIndex);
-	Image._completeFrame = true;
-	Image.rotation_ = webrtc::VideoRotation::kVideoRotation_0;
-	Image.SetTimestamp(InFrame->GetTimestampRTP());
-	Image.capture_time_ms_ = InFrame->GetTimestampUs() / 1000.0;
-
-	webrtc::CodecSpecificInfo CodecInfo;
-	CodecInfo.codecType = webrtc::VideoCodecType::kVideoCodecH264;
-	CodecInfo.codecSpecific.H264.packetization_mode = webrtc::H264PacketizationMode::NonInterleaved;
-	CodecInfo.codecSpecific.H264.temporal_idx = webrtc::kNoTemporalIdx;
-	CodecInfo.codecSpecific.H264.idr_frame = InPacket.IsKeyFrame;
-	CodecInfo.codecSpecific.H264.base_layer_sync = false;
-
-	Factory->OnEncodedImage(Image, &CodecInfo, &FragHeader);
-}
+		Image.timing_.packetization_finish_ms = FTimespan::FromSeconds(FPlatformTime::Seconds()).GetTotalMilliseconds();
+		Image.timing_.encode_start_ms = InPacket.Timings.StartTs.GetTotalMilliseconds();
+		Image.timing_.encode_finish_ms = InPacket.Timings.FinishTs.GetTotalMilliseconds();
+		Image.timing_.flags = webrtc::VideoSendTiming::kTriggeredByTimer;
+
+		Image.SetEncodedData(webrtc::EncodedImageBuffer::Create(const_cast<uint8_t*>(InPacket.Data.Get()), InPacket.DataSize));
+		Image._encodedWidth = InFrame->GetWidth();
+		Image._encodedHeight = InFrame->GetHeight();
+		Image._frameType = InPacket.IsKeyFrame ? webrtc::VideoFrameType::kVideoFrameKey : webrtc::VideoFrameType::kVideoFrameDelta;
+		Image.content_type_ = webrtc::VideoContentType::UNSPECIFIED;
+		Image.qp_ = InPacket.VideoQP;
+		Image.SetSpatialIndex(InLayerIndex);
+		Image._completeFrame = true;
+		Image.rotation_ = webrtc::VideoRotation::kVideoRotation_0;
+		Image.SetTimestamp(InFrame->GetTimestampRTP());
+		Image.capture_time_ms_ = InFrame->GetTimestampUs() / 1000.0;
+
+		webrtc::CodecSpecificInfo CodecInfo;
+		CodecInfo.codecType = webrtc::VideoCodecType::kVideoCodecH264;
+		CodecInfo.codecSpecific.H264.packetization_mode = webrtc::H264PacketizationMode::NonInterleaved;
+		CodecInfo.codecSpecific.H264.temporal_idx = webrtc::kNoTemporalIdx;
+		CodecInfo.codecSpecific.H264.idr_frame = InPacket.IsKeyFrame;
+		CodecInfo.codecSpecific.H264.base_layer_sync = false;
+
+		Factory->OnEncodedImage(Image, &CodecInfo, &FragHeader);
+	}
+} // namespace UE::PixelStreaming
\ No newline at end of file
diff --git a/Engine/Plugins/Media/PixelStreaming/Source/PixelStreaming/Private/VideoEncoderH264Wrapper.h b/Engine/Plugins/Media/PixelStreaming/Source/PixelStreaming/Private/VideoEncoderH264Wrapper.h
--- a/Engine/Plugins/Media/PixelStreaming/Source/PixelStreaming/Private/VideoEncoderH264Wrapper.h
+++ b/Engine/Plugins/Media/PixelStreaming/Source/PixelStreaming/Private/VideoEncoderH264Wrapper.h
@@ -7,36 +7,35 @@
 namespace webrtc
 {
 	class VideoFrame;
-}
+} // namespace webrtc
 
-namespace UE {
-	namespace PixelStreaming {
-		class FVideoEncoderFactory;
-
-		class FVideoEncoderH264Wrapper
-		{
-		public:
-			FVideoEncoderH264Wrapper(TUniquePtr<FEncoderFrameFactory> FrameFactory, TUniquePtr<AVEncoder::FVideoEncoder> Encoder);
-			~FVideoEncoderH264Wrapper();
-
-			uint64 GetId() const { return Id; }
-
-			void SetForceNextKeyframe() { bForceNextKeyframe = true; }
-
-			void Encode(const webrtc::VideoFrame& WebRTCFrame, bool bKeyframe);
-
-			AVEncoder::FVideoEncoder::FLayerConfig GetCurrentConfig();
-			void SetConfig(const AVEncoder::FVideoEncoder::FLayerConfig& NewConfig);
-
-			static void OnEncodedPacket(UE::PixelStreaming::FVideoEncoderFactory* Factory, 
-				uint32 InLayerIndex, const TSharedPtr<AVEncoder::FVideoEncoderInputFrame> InFrame, 
-				const AVEncoder::FCodecPacket& InPacket);
-
-		private:
-			uint64 Id;
-			TUniquePtr<FEncoderFrameFactory> FrameFactory;
-			TUniquePtr<AVEncoder::FVideoEncoder> Encoder;
-			bool bForceNextKeyframe = false;
-		};
-	}
-}
+namespace UE::PixelStreaming 
+{
+	class FVideoEncoderFactory;
+	
+	class FVideoEncoderH264Wrapper
+	{
+	public:
+		FVideoEncoderH264Wrapper(TUniquePtr<FEncoderFrameFactory> FrameFactory, TUniquePtr<AVEncoder::FVideoEncoder> Encoder);
+		~FVideoEncoderH264Wrapper();
+
+		uint64 GetId() const { return Id; }
+
+		void SetForceNextKeyframe() { bForceNextKeyframe = true; }
+
+		void Encode(const webrtc::VideoFrame& WebRTCFrame, bool bKeyframe);
+
+		AVEncoder::FVideoEncoder::FLayerConfig GetCurrentConfig();
+		void SetConfig(const AVEncoder::FVideoEncoder::FLayerConfig& NewConfig);
+
+		static void OnEncodedPacket(FVideoEncoderFactory* Factory, 
+			uint32 InLayerIndex, const TSharedPtr<AVEncoder::FVideoEncoderInputFrame> InFrame, 
+			const AVEncoder::FCodecPacket& InPacket);
+
+	private:
+		uint64 Id;
+		TUniquePtr<FEncoderFrameFactory> FrameFactory;
+		TUniquePtr<AVEncoder::FVideoEncoder> Encoder;
+		bool bForceNextKeyframe = false;
+	};
+} // namespace UE::PixelStreaming
\ No newline at end of file
diff --git a/Engine/Plugins/Media/PixelStreaming/Source/PixelStreaming/Private/VideoEncoderRTC.cpp b/Engine/Plugins/Media/PixelStreaming/Source/PixelStreaming/Private/VideoEncoderRTC.cpp
--- a/Engine/Plugins/Media/PixelStreaming/Source/PixelStreaming/Private/VideoEncoderRTC.cpp
+++ b/Engine/Plugins/Media/PixelStreaming/Source/PixelStreaming/Private/VideoEncoderRTC.cpp
@@ -14,185 +14,188 @@
 #include "VideoEncoderH264Wrapper.h"
 #include "HAL/PlatformFileManager.h"
 
-UE::PixelStreaming::FVideoEncoderRTC::FVideoEncoderRTC(UE::PixelStreaming::FVideoEncoderFactory& InFactory)
-	: Factory(InFactory)
+namespace UE::PixelStreaming
 {
-}
-
-UE::PixelStreaming::FVideoEncoderRTC::~FVideoEncoderRTC()
-{
-	Factory.ReleaseVideoEncoder(this);
-}
-
-int UE::PixelStreaming::FVideoEncoderRTC::InitEncode(webrtc::VideoCodec const* InCodecSettings, VideoEncoder::Settings const& settings)
-{
-	checkf(AVEncoder::FVideoEncoderFactory::Get().IsSetup(), TEXT("FVideoEncoderFactory not setup"));
-	// Try and get the encoder. If it doesn't exist, create it.
-	FVideoEncoderH264Wrapper* Encoder = Factory.GetOrCreateHardwareEncoder(InCodecSettings->width, InCodecSettings->height, InCodecSettings->maxBitrate, InCodecSettings->startBitrate, InCodecSettings->maxFramerate);
-
-	if (Encoder != nullptr)
+	FVideoEncoderRTC::FVideoEncoderRTC(FVideoEncoderFactory& InFactory)
+		: Factory(InFactory)
 	{
-		HardwareEncoder = Encoder;
-		UpdateConfig();
-		WebRtcProposedTargetBitrate = InCodecSettings->startBitrate;
-		return WEBRTC_VIDEO_CODEC_OK;
 	}
-	else
+
+	FVideoEncoderRTC::~FVideoEncoderRTC()
 	{
-		return WEBRTC_VIDEO_CODEC_ERROR;
+		Factory.ReleaseVideoEncoder(this);
 	}
-}
-
-int32 UE::PixelStreaming::FVideoEncoderRTC::RegisterEncodeCompleteCallback(webrtc::EncodedImageCallback* callback)
-{
-	OnEncodedImageCallback = callback;
-	return WEBRTC_VIDEO_CODEC_OK;
-}
 
-int32 UE::PixelStreaming::FVideoEncoderRTC::Release()
-{
-	OnEncodedImageCallback = nullptr;
-	return WEBRTC_VIDEO_CODEC_OK;
-}
+	int FVideoEncoderRTC::InitEncode(webrtc::VideoCodec const* InCodecSettings, VideoEncoder::Settings const& settings)
+	{
+		checkf(AVEncoder::FVideoEncoderFactory::Get().IsSetup(), TEXT("FVideoEncoderFactory not setup"));
+		// Try and get the encoder. If it doesn't exist, create it.
+		FVideoEncoderH264Wrapper* Encoder = Factory.GetOrCreateHardwareEncoder(InCodecSettings->width, InCodecSettings->height, InCodecSettings->maxBitrate, InCodecSettings->startBitrate, InCodecSettings->maxFramerate);
 
-int32 UE::PixelStreaming::FVideoEncoderRTC::Encode(webrtc::VideoFrame const& frame, std::vector<webrtc::VideoFrameType> const* frame_types)
-{
-	FPixelStreamingFrameBuffer* FrameBuffer = static_cast<FPixelStreamingFrameBuffer*>(frame.video_frame_buffer().get());
+		if (Encoder != nullptr)
+		{
+			HardwareEncoder = Encoder;
+			UpdateConfig();
+			WebRtcProposedTargetBitrate = InCodecSettings->startBitrate;
+			return WEBRTC_VIDEO_CODEC_OK;
+		}
+		else
+		{
+			return WEBRTC_VIDEO_CODEC_ERROR;
+		}
+	}
 
-	// If initialize frame is passed, this is a dummy frame so WebRTC is happy frames are passing through
-	// This is used so that an encoder can be active as far as WebRTC is concerned by simply have frames transmitted by some other encoder on its behalf.
-	if (FrameBuffer->GetFrameBufferType() == UE::PixelStreaming::EFrameBufferType::Initialize)
+	int32 FVideoEncoderRTC::RegisterEncodeCompleteCallback(webrtc::EncodedImageCallback* callback)
 	{
+		OnEncodedImageCallback = callback;
 		return WEBRTC_VIDEO_CODEC_OK;
 	}
 
-	bool bKeyframe = false;
-	if ((frame_types && (*frame_types)[0] == webrtc::VideoFrameType::kVideoFrameKey))
+	int32 FVideoEncoderRTC::Release()
 	{
-		bKeyframe = true;
+		OnEncodedImageCallback = nullptr;
+		return WEBRTC_VIDEO_CODEC_OK;
 	}
 
-	UE::PixelStreaming::FVideoEncoderH264Wrapper* Encoder = Factory.GetHardwareEncoder();
-	if (Encoder == nullptr)
+	int32 FVideoEncoderRTC::Encode(webrtc::VideoFrame const& frame, std::vector<webrtc::VideoFrameType> const* frame_types)
 	{
-		return WEBRTC_VIDEO_CODEC_ERROR;
-	}
-
-	UpdateConfig();
+		FPixelStreamingFrameBuffer* FrameBuffer = static_cast<FPixelStreamingFrameBuffer*>(frame.video_frame_buffer().get());
 
-	Encoder->Encode(frame, bKeyframe);
+		// If initialize frame is passed, this is a dummy frame so WebRTC is happy frames are passing through
+		// This is used so that an encoder can be active as far as WebRTC is concerned by simply have frames transmitted by some other encoder on its behalf.
+		if (FrameBuffer->GetFrameBufferType() == EPixelStreamingFrameBufferType::Initialize)
+		{
+			return WEBRTC_VIDEO_CODEC_OK;
+		}
 
-	return WEBRTC_VIDEO_CODEC_OK;
-}
+		bool bKeyframe = false;
+		if ((frame_types && (*frame_types)[0] == webrtc::VideoFrameType::kVideoFrameKey))
+		{
+			bKeyframe = true;
+		}
 
-// Pass rate control parameters from WebRTC to our encoder
-// This is how WebRTC can control the bitrate/framerate of the encoder.
-void UE::PixelStreaming::FVideoEncoderRTC::SetRates(RateControlParameters const& parameters)
-{
-	PendingRateChange.Emplace(parameters);
-}
+		FVideoEncoderH264Wrapper* Encoder = Factory.GetHardwareEncoder();
+		if (Encoder == nullptr)
+		{
+			return WEBRTC_VIDEO_CODEC_ERROR;
+		}
 
-webrtc::VideoEncoder::EncoderInfo UE::PixelStreaming::FVideoEncoderRTC::GetEncoderInfo() const
-{
-	VideoEncoder::EncoderInfo info;
-	info.supports_native_handle = true;
-	info.is_hardware_accelerated = true;
-	info.has_internal_source = false;
-	info.supports_simulcast = false;
-	info.implementation_name = TCHAR_TO_UTF8(*FString::Printf(TEXT("PIXEL_STREAMING_HW_ENCODER_%s"), GDynamicRHI->GetName()));
+		UpdateConfig();
 
-	const int LowQP = UE::PixelStreaming::Settings::CVarPixelStreamingWebRTCLowQpThreshold.GetValueOnAnyThread();
-	const int HighQP = UE::PixelStreaming::Settings::CVarPixelStreamingWebRTCHighQpThreshold.GetValueOnAnyThread();
-	info.scaling_settings = VideoEncoder::ScalingSettings(LowQP, HighQP);
+		Encoder->Encode(frame, bKeyframe);
 
-	// if true means HW encoder must be perfect and drop frames itself etc
-	info.has_trusted_rate_controller = false;
+		return WEBRTC_VIDEO_CODEC_OK;
+	}
 
-	return info;
-}
+	// Pass rate control parameters from WebRTC to our encoder
+	// This is how WebRTC can control the bitrate/framerate of the encoder.
+	void FVideoEncoderRTC::SetRates(RateControlParameters const& parameters)
+	{
+		PendingRateChange.Emplace(parameters);
+	}
 
-void UE::PixelStreaming::FVideoEncoderRTC::UpdateConfig()
-{
-	if (UE::PixelStreaming::FVideoEncoderH264Wrapper* Encoder = Factory.GetHardwareEncoder())
+	webrtc::VideoEncoder::EncoderInfo FVideoEncoderRTC::GetEncoderInfo() const
 	{
-		AVEncoder::FVideoEncoder::FLayerConfig NewConfig = Encoder->GetCurrentConfig();
+		VideoEncoder::EncoderInfo info;
+		info.supports_native_handle = true;
+		info.is_hardware_accelerated = true;
+		info.has_internal_source = false;
+		info.supports_simulcast = false;
+		info.implementation_name = TCHAR_TO_UTF8(*FString::Printf(TEXT("PIXEL_STREAMING_HW_ENCODER_%s"), GDynamicRHI->GetName()));
+
+		const int LowQP = UE::PixelStreaming::Settings::CVarPixelStreamingWebRTCLowQpThreshold.GetValueOnAnyThread();
+		const int HighQP = UE::PixelStreaming::Settings::CVarPixelStreamingWebRTCHighQpThreshold.GetValueOnAnyThread();
+		info.scaling_settings = VideoEncoder::ScalingSettings(LowQP, HighQP);
+
+		// if true means HW encoder must be perfect and drop frames itself etc
+		info.has_trusted_rate_controller = false;
 
-		if (PendingRateChange.IsSet())
+		return info;
+	}
+
+	void FVideoEncoderRTC::UpdateConfig()
+	{
+		if (FVideoEncoderH264Wrapper* Encoder = Factory.GetHardwareEncoder())
 		{
-			const RateControlParameters& RateChangeParams = PendingRateChange.GetValue();
+			AVEncoder::FVideoEncoder::FLayerConfig NewConfig = Encoder->GetCurrentConfig();
 
-			NewConfig.MaxFramerate = RateChangeParams.framerate_fps;
+			if (PendingRateChange.IsSet())
+			{
+				const RateControlParameters& RateChangeParams = PendingRateChange.GetValue();
 
-			// We store what WebRTC wants as the bitrate, even if we are overriding it, so we can restore back to it when user stops using CVar.
-			WebRtcProposedTargetBitrate = RateChangeParams.bitrate.get_sum_kbps() * 1000;
+				NewConfig.MaxFramerate = RateChangeParams.framerate_fps;
 
-			// Clear the rate change request
-			PendingRateChange.Reset();
-		}
+				// We store what WebRTC wants as the bitrate, even if we are overriding it, so we can restore back to it when user stops using CVar.
+				WebRtcProposedTargetBitrate = RateChangeParams.bitrate.get_sum_kbps() * 1000;
+
+				// Clear the rate change request
+				PendingRateChange.Reset();
+			}
 
-		NewConfig = CreateEncoderConfigFromCVars(NewConfig);
+			NewConfig = CreateEncoderConfigFromCVars(NewConfig);
 
-		Encoder->SetConfig(NewConfig);
+			Encoder->SetConfig(NewConfig);
+		}
 	}
-}
 
-AVEncoder::FVideoEncoder::FLayerConfig UE::PixelStreaming::FVideoEncoderRTC::CreateEncoderConfigFromCVars(AVEncoder::FVideoEncoder::FLayerConfig InEncoderConfig) const
-{
-	// Change encoder settings through CVars
-	const int32 MaxBitrateCVar = UE::PixelStreaming::Settings::CVarPixelStreamingEncoderMaxBitrate.GetValueOnAnyThread();
-	const int32 TargetBitrateCVar = UE::PixelStreaming::Settings::CVarPixelStreamingEncoderTargetBitrate.GetValueOnAnyThread();
-	const int32 MinQPCVar = UE::PixelStreaming::Settings::CVarPixelStreamingEncoderMinQP.GetValueOnAnyThread();
-	const int32 MaxQPCVar = UE::PixelStreaming::Settings::CVarPixelStreamingEncoderMaxQP.GetValueOnAnyThread();
-	const AVEncoder::FVideoEncoder::RateControlMode RateControlCVar = UE::PixelStreaming::Settings::GetRateControlCVar();
-	const AVEncoder::FVideoEncoder::MultipassMode MultiPassCVar = UE::PixelStreaming::Settings::GetMultipassCVar();
-	const bool bFillerDataCVar = UE::PixelStreaming::Settings::CVarPixelStreamingEnableFillerData.GetValueOnAnyThread();
-	const AVEncoder::FVideoEncoder::H264Profile H264Profile = UE::PixelStreaming::Settings::GetH264Profile();
-
-	InEncoderConfig.MaxBitrate = MaxBitrateCVar > -1 ? MaxBitrateCVar : InEncoderConfig.MaxBitrate;
-	InEncoderConfig.TargetBitrate = TargetBitrateCVar > -1 ? TargetBitrateCVar : WebRtcProposedTargetBitrate;
-	InEncoderConfig.QPMin = MinQPCVar;
-	InEncoderConfig.QPMax = MaxQPCVar;
-	InEncoderConfig.RateControlMode = RateControlCVar;
-	InEncoderConfig.MultipassMode = MultiPassCVar;
-	InEncoderConfig.FillData = bFillerDataCVar;
-	InEncoderConfig.H264Profile = H264Profile;
-
-	return InEncoderConfig;
-}
-
-void UE::PixelStreaming::FVideoEncoderRTC::SendEncodedImage(webrtc::EncodedImage const& encoded_image, webrtc::CodecSpecificInfo const* codec_specific_info, webrtc::RTPFragmentationHeader const* fragmentation)
-{
-	if (FirstKeyframeCountdown > 0)
+	AVEncoder::FVideoEncoder::FLayerConfig FVideoEncoderRTC::CreateEncoderConfigFromCVars(AVEncoder::FVideoEncoder::FLayerConfig InEncoderConfig) const
 	{
-		// ideally we want to make the first frame of new peers a keyframe but we
-		// dont know when webrtc will decide to start sending out frames. this is
-		// the next best option. its a count beause when it was the very next frame
-		// it still didnt seem to work. delaying it a few frames seems to have worked.
-		--FirstKeyframeCountdown;
-		if (FirstKeyframeCountdown == 0)
-		{
-			Factory.ForceKeyFrame();
-		}
+		// Change encoder settings through CVars
+		const int32 MaxBitrateCVar = UE::PixelStreaming::Settings::CVarPixelStreamingEncoderMaxBitrate.GetValueOnAnyThread();
+		const int32 TargetBitrateCVar = UE::PixelStreaming::Settings::CVarPixelStreamingEncoderTargetBitrate.GetValueOnAnyThread();
+		const int32 MinQPCVar = UE::PixelStreaming::Settings::CVarPixelStreamingEncoderMinQP.GetValueOnAnyThread();
+		const int32 MaxQPCVar = UE::PixelStreaming::Settings::CVarPixelStreamingEncoderMaxQP.GetValueOnAnyThread();
+		const AVEncoder::FVideoEncoder::RateControlMode RateControlCVar = UE::PixelStreaming::Settings::GetRateControlCVar();
+		const AVEncoder::FVideoEncoder::MultipassMode MultiPassCVar = UE::PixelStreaming::Settings::GetMultipassCVar();
+		const bool bFillerDataCVar = UE::PixelStreaming::Settings::CVarPixelStreamingEnableFillerData.GetValueOnAnyThread();
+		const AVEncoder::FVideoEncoder::H264Profile H264Profile = UE::PixelStreaming::Settings::GetH264Profile();
+
+		InEncoderConfig.MaxBitrate = MaxBitrateCVar > -1 ? MaxBitrateCVar : InEncoderConfig.MaxBitrate;
+		InEncoderConfig.TargetBitrate = TargetBitrateCVar > -1 ? TargetBitrateCVar : WebRtcProposedTargetBitrate;
+		InEncoderConfig.QPMin = MinQPCVar;
+		InEncoderConfig.QPMax = MaxQPCVar;
+		InEncoderConfig.RateControlMode = RateControlCVar;
+		InEncoderConfig.MultipassMode = MultiPassCVar;
+		InEncoderConfig.FillData = bFillerDataCVar;
+		InEncoderConfig.H264Profile = H264Profile;
+
+		return InEncoderConfig;
 	}
 
-	// Dump H264 frames to file for debugging if CVar is turned on.
-	if (UE::PixelStreaming::Settings::CVarPixelStreamingDebugDumpFrame.GetValueOnAnyThread())
+	void FVideoEncoderRTC::SendEncodedImage(webrtc::EncodedImage const& encoded_image, webrtc::CodecSpecificInfo const* codec_specific_info, webrtc::RTPFragmentationHeader const* fragmentation)
 	{
-		static IFileHandle* FileHandle = nullptr;
-		if (!FileHandle)
+		if (FirstKeyframeCountdown > 0)
 		{
-			IPlatformFile& PlatformFile = FPlatformFileManager::Get().GetPlatformFile();
-			FString TempFilePath = FPaths::CreateTempFilename(*FPaths::ProjectSavedDir(), TEXT("encoded_frame"), TEXT(".h264"));
-			FileHandle = PlatformFile.OpenWrite(*TempFilePath);
-			check(FileHandle);
+			// ideally we want to make the first frame of new peers a keyframe but we
+			// dont know when webrtc will decide to start sending out frames. this is
+			// the next best option. its a count beause when it was the very next frame
+			// it still didnt seem to work. delaying it a few frames seems to have worked.
+			--FirstKeyframeCountdown;
+			if (FirstKeyframeCountdown == 0)
+			{
+				Factory.ForceKeyFrame();
+			}
 		}
 
-		FileHandle->Write(encoded_image.data(), encoded_image.size());
-		FileHandle->Flush();
-	}
+		// Dump H264 frames to file for debugging if CVar is turned on.
+		if (UE::PixelStreaming::Settings::CVarPixelStreamingDebugDumpFrame.GetValueOnAnyThread())
+		{
+			static IFileHandle* FileHandle = nullptr;
+			if (!FileHandle)
+			{
+				IPlatformFile& PlatformFile = FPlatformFileManager::Get().GetPlatformFile();
+				FString TempFilePath = FPaths::CreateTempFilename(*FPaths::ProjectSavedDir(), TEXT("encoded_frame"), TEXT(".h264"));
+				FileHandle = PlatformFile.OpenWrite(*TempFilePath);
+				check(FileHandle);
+			}
+
+			FileHandle->Write(encoded_image.data(), encoded_image.size());
+			FileHandle->Flush();
+		}
 
-	if (OnEncodedImageCallback)
-	{
-		OnEncodedImageCallback->OnEncodedImage(encoded_image, codec_specific_info, fragmentation);
+		if (OnEncodedImageCallback)
+		{
+			OnEncodedImageCallback->OnEncodedImage(encoded_image, codec_specific_info, fragmentation);
+		}
 	}
-}
+} // namespace UE::PixelStreaming
\ No newline at end of file
diff --git a/Engine/Plugins/Media/PixelStreaming/Source/PixelStreaming/Private/VideoEncoderRTC.h b/Engine/Plugins/Media/PixelStreaming/Source/PixelStreaming/Private/VideoEncoderRTC.h
--- a/Engine/Plugins/Media/PixelStreaming/Source/PixelStreaming/Private/VideoEncoderRTC.h
+++ b/Engine/Plugins/Media/PixelStreaming/Source/PixelStreaming/Private/VideoEncoderRTC.h
@@ -9,54 +9,53 @@
 #include "Misc/Optional.h"
 #include "VideoEncoderH264Wrapper.h"
 
-namespace UE {
-	namespace PixelStreaming {
-		class FVideoEncoderFactory;
-
-		// Implementation that is a WebRTC video encoder that allows us tie to our actually underlying non-WebRTC video encoder.
-		class FVideoEncoderRTC : public webrtc::VideoEncoder
-		{
-		public:
-			FVideoEncoderRTC(FVideoEncoderFactory& InFactory);
-			virtual ~FVideoEncoderRTC() override;
-
-			// WebRTC Interface
-			virtual int InitEncode(webrtc::VideoCodec const* codec_settings, webrtc::VideoEncoder::Settings const& settings) override;
-			virtual int32 RegisterEncodeCompleteCallback(webrtc::EncodedImageCallback* callback) override;
-			virtual int32 Release() override;
-			virtual int32 Encode(webrtc::VideoFrame const& frame, std::vector<webrtc::VideoFrameType> const* frame_types) override;
-			virtual void SetRates(RateControlParameters const& parameters) override;
-			virtual webrtc::VideoEncoder::EncoderInfo GetEncoderInfo() const override;
-
-			// Note: These funcs can also be overriden but are not pure virtual
-			// virtual void SetFecControllerOverride(FecControllerOverride* fec_controller_override) override;
-			// virtual void OnPacketLossRateUpdate(float packet_loss_rate) override;
-			// virtual void OnRttUpdate(int64_t rtt_ms) override;
-			// virtual void OnLossNotification(const LossNotification& loss_notification) override;
-			// End WebRTC Interface.
-
-			void SendEncodedImage(webrtc::EncodedImage const& encoded_image, webrtc::CodecSpecificInfo const* codec_specific_info, webrtc::RTPFragmentationHeader const* fragmentation);
-
-		private:
-			void UpdateConfig();
-			AVEncoder::FVideoEncoder::FLayerConfig CreateEncoderConfigFromCVars(AVEncoder::FVideoEncoder::FLayerConfig BaseEncoderConfig) const;
-
-			FVideoEncoderFactory& Factory;
-
-			UE::PixelStreaming::FVideoEncoderH264Wrapper* HardwareEncoder;
-
-			// We store this so we can restore back to it if the user decides to use then stop using the PixelStreaming.Encoder.TargetBitrate CVar.
-			int32 WebRtcProposedTargetBitrate = 5000000;
-
-			webrtc::EncodedImageCallback* OnEncodedImageCallback = nullptr;
-
-			// WebRTC may request a bitrate/framerate change using SetRates(), we only respect this if this encoder is actually encoding
-			// so we use this optional object to store a rate change and act upon it when this encoder does its next call to Encode().
-			TOptional<RateControlParameters> PendingRateChange;
-
-			// Used to send an initial keyframe
-			// see notes in SendEncodedImage implementation
-			int FirstKeyframeCountdown = 2;
-		};
-	}
-}
+namespace UE::PixelStreaming 
+{
+	class FVideoEncoderFactory;
+
+	// Implementation that is a WebRTC video encoder that allows us tie to our actually underlying non-WebRTC video encoder.
+	class FVideoEncoderRTC : public webrtc::VideoEncoder
+	{
+	public:
+		FVideoEncoderRTC(FVideoEncoderFactory& InFactory);
+		virtual ~FVideoEncoderRTC() override;
+
+		// WebRTC Interface
+		virtual int InitEncode(webrtc::VideoCodec const* codec_settings, webrtc::VideoEncoder::Settings const& settings) override;
+		virtual int32 RegisterEncodeCompleteCallback(webrtc::EncodedImageCallback* callback) override;
+		virtual int32 Release() override;
+		virtual int32 Encode(webrtc::VideoFrame const& frame, std::vector<webrtc::VideoFrameType> const* frame_types) override;
+		virtual void SetRates(RateControlParameters const& parameters) override;
+		virtual webrtc::VideoEncoder::EncoderInfo GetEncoderInfo() const override;
+
+		// Note: These funcs can also be overriden but are not pure virtual
+		// virtual void SetFecControllerOverride(FecControllerOverride* fec_controller_override) override;
+		// virtual void OnPacketLossRateUpdate(float packet_loss_rate) override;
+		// virtual void OnRttUpdate(int64_t rtt_ms) override;
+		// virtual void OnLossNotification(const LossNotification& loss_notification) override;
+		// End WebRTC Interface.
+
+		void SendEncodedImage(webrtc::EncodedImage const& encoded_image, webrtc::CodecSpecificInfo const* codec_specific_info, webrtc::RTPFragmentationHeader const* fragmentation);
+
+	private:
+		void UpdateConfig();
+		AVEncoder::FVideoEncoder::FLayerConfig CreateEncoderConfigFromCVars(AVEncoder::FVideoEncoder::FLayerConfig BaseEncoderConfig) const;
+
+		FVideoEncoderFactory& Factory;
+
+		FVideoEncoderH264Wrapper* HardwareEncoder;
+
+		// We store this so we can restore back to it if the user decides to use then stop using the PixelStreaming.Encoder.TargetBitrate CVar.
+		int32 WebRtcProposedTargetBitrate = 5000000;
+
+		webrtc::EncodedImageCallback* OnEncodedImageCallback = nullptr;
+
+		// WebRTC may request a bitrate/framerate change using SetRates(), we only respect this if this encoder is actually encoding
+		// so we use this optional object to store a rate change and act upon it when this encoder does its next call to Encode().
+		TOptional<RateControlParameters> PendingRateChange;
+
+		// Used to send an initial keyframe
+		// see notes in SendEncodedImage implementation
+		int FirstKeyframeCountdown = 2;
+	};
+} // namespace UE::PixelStreaming
\ No newline at end of file
