This commit is a large collection (from various authors) of Pixel Streaming 5.1 work

diff --git a/Engine/Plugins/Media/PixelStreaming/Source/PixelStreaming/Private/VideoSource.cpp b/Engine/Plugins/Media/PixelStreaming/Source/PixelStreaming/Private/VideoSource.cpp
--- a/Engine/Plugins/Media/PixelStreaming/Source/PixelStreaming/Private/VideoSource.cpp
+++ b/Engine/Plugins/Media/PixelStreaming/Source/PixelStreaming/Private/VideoSource.cpp
@@ -6,36 +6,22 @@
 #include "Stats.h"
 #include "PixelStreamingStatNames.h"
 #include "PlayerSessions.h"
-#include "IPixelStreamingTextureSource.h"
+#include "PixelStreamingModule.h"
+#include "FrameBuffer.h"
+#include "TextureSourceComputeI420.h"
+#include "TextureSourceCPUI420.h"
 
 namespace UE::PixelStreaming
 {
-	FVideoSourceBase::FVideoSourceBase(FPixelStreamingPlayerId InPlayerId)
-		: PlayerId(InPlayerId)
+	FVideoSourceBase::FVideoSourceBase()
 	{
 	}
 
-	void FVideoSourceBase::AddRef() const
-	{
-		RefCount.IncRef();
-	}
-
-	rtc::RefCountReleaseStatus FVideoSourceBase::Release() const
-	{
-		const rtc::RefCountReleaseStatus Status = RefCount.DecRef();
-		if (Status == rtc::RefCountReleaseStatus::kDroppedLastRef)
-		{
-			FFixedFPSPump::Get()->UnregisterVideoSource(PlayerId);
-			delete this;
-		}
-		return Status;
-	}
-
 	void FVideoSourceBase::Initialize()
 	{
 		CurrentState = webrtc::MediaSourceInterface::SourceState::kInitializing;
 
-		FFixedFPSPump::Get()->RegisterVideoSource(PlayerId, this);
+		IPixelStreamingModule::Get().RegisterPumpable(this);
 	}
 
 	FVideoSourceBase::~FVideoSourceBase()
@@ -79,18 +65,13 @@ namespace UE::PixelStreaming
 * ------------------ FPlayerVideSource ------------------
 */
 
-	FVideoSourceP2P::FVideoSourceP2P(FPixelStreamingPlayerId InPlayerId, FPlayerSessions* InSessions)
-		: FVideoSourceBase(InPlayerId)
-		, Sessions(InSessions)
+	FVideoSourceP2P::FVideoSourceP2P(FName SourceType, TFunction<bool()> InIsQualityControllerFunc)
+		: IsQualityControllerFunc(InIsQualityControllerFunc)
 	{
-		if (Settings::IsCodecVPX())
-		{
-			TextureSource = TSharedPtr<IPixelStreamingTextureSource, ESPMode::ThreadSafe>(new FBackBufferToCPUTextureSource(Settings::CVarPixelStreamingFrameScale.GetValueOnAnyThread() <= 0 ? 1.0 : Settings::CVarPixelStreamingFrameScale.GetValueOnAnyThread()));
-		}
-		else
-		{
-			TextureSource = TSharedPtr<IPixelStreamingTextureSource, ESPMode::ThreadSafe>(new FBackBufferTextureSource(Settings::CVarPixelStreamingFrameScale.GetValueOnAnyThread() <= 0 ? 1.0 : Settings::CVarPixelStreamingFrameScale.GetValueOnAnyThread()));
-		}
+		TUniquePtr<FPixelStreamingTextureSource> TextureGenerator = IPixelStreamingModule::Get().GetTextureSourceFactory().CreateTextureSource(SourceType);
+
+		float FrameScale = Settings::CVarPixelStreamingFrameScale.GetValueOnAnyThread() <= 0 ? 1.0 : Settings::CVarPixelStreamingFrameScale.GetValueOnAnyThread();
+		TextureSource = TSharedPtr<FTextureTripleBuffer, ESPMode::ThreadSafe>(new FTextureTripleBuffer(FrameScale, MoveTemp(TextureGenerator)));
 
 		// We store the codec during construction as querying it everytime seem overly wasteful
 		// especially considering we don't actually support switching codec mid-stream.
@@ -104,7 +85,8 @@ namespace UE::PixelStreaming
 
 	webrtc::VideoFrame FVideoSourceP2P::CreateFrameH264(int32 FrameId)
 	{
-		bool bQualityController = Sessions->IsQualityController(PlayerId);
+		bool bQualityController = IsQualityControllerFunc();
+
 		TextureSource->SetEnabled(bQualityController);
 
 		const int64 TimestampUs = rtc::TimeMicros();
@@ -170,8 +152,8 @@ namespace UE::PixelStreaming
 * ------------------ FVideoSourceSFU ------------------
 */
 
-	FVideoSourceSFU::FVideoSourceSFU()
-		: FVideoSourceBase(SFU_PLAYER_ID)
+	FVideoSourceSFU::FVideoSourceSFU(FName SourceType)
+		: FVideoSourceBase()
 	{
 		// Make a copy of simulcast settings and sort them based on scaling.
 		using FLayer = Settings::FSimulcastParameters::FLayer;
@@ -183,15 +165,31 @@ namespace UE::PixelStreaming
 		}
 		SortedLayers.Sort([](const FLayer& LayerA, const FLayer& LayerB) { return LayerA.Scaling > LayerB.Scaling; });
 
-		for (FLayer* SimulcastLayer : SortedLayers)
+		if(Settings::IsCodecVPX())
+		{
+			// VPX Encoding. We only keep a single texture source (the largest one) and do CPU scaling in the Encode loop
+			const float Scale = SortedLayers[SortedLayers.Num() - 1]->Scaling;
+			AddLayerTexture(SourceType, Scale);
+		}
+		else
 		{
-			const float Scale = 1.0f / SimulcastLayer->Scaling;
-			TSharedPtr<FBackBufferTextureSource> TextureSource = MakeShared<FBackBufferTextureSource>(Scale);
-			TextureSource->SetEnabled(true);
-			LayerTextures.Add(TextureSource);
+			// H264 Encoding. We store a texture source for each simulcast layer
+			for (FLayer* SimulcastLayer : SortedLayers)
+			{
+				const float Scale = 1.0f / SimulcastLayer->Scaling;
+				AddLayerTexture(SourceType, Scale);
+			}
 		}
 	}
 
+	void FVideoSourceSFU::AddLayerTexture(FName SourceType, float Scale) 
+	{
+		TUniquePtr<FPixelStreamingTextureSource> TextureGenerator = IPixelStreamingModule::Get().GetTextureSourceFactory().CreateTextureSource(SourceType);
+		TSharedPtr<FTextureTripleBuffer> TextureSource = MakeShared<FTextureTripleBuffer, ESPMode::ThreadSafe>(Scale, MoveTemp(TextureGenerator));
+		TextureSource->SetEnabled(true);
+		LayerTextures.Add(TextureSource);
+	}
+
 	bool FVideoSourceSFU::IsReadyForPump() const
 	{
 		// Check all texture sources are ready.
@@ -208,7 +206,16 @@ namespace UE::PixelStreaming
 
 	webrtc::VideoFrame FVideoSourceSFU::CreateFrame(int32 FrameId)
 	{
-		rtc::scoped_refptr<webrtc::VideoFrameBuffer> FrameBuffer = new rtc::RefCountedObject<FSimulcastFrameBuffer>(LayerTextures);
+		rtc::scoped_refptr<webrtc::VideoFrameBuffer> FrameBuffer;
+		if(Settings::IsCodecVPX())
+		{
+			// VPX Encoding. We only ever have a single layer
+			FrameBuffer = new rtc::RefCountedObject<FFrameBufferI420>(LayerTextures[0]);
+		}
+		else
+		{
+			FrameBuffer = new rtc::RefCountedObject<FSimulcastFrameBuffer>(LayerTextures);
+		}
 		const int64 TimestampUs = rtc::TimeMicros();
 
 		webrtc::VideoFrame Frame = webrtc::VideoFrame::Builder()
diff --git a/Engine/Plugins/Media/PixelStreaming/Source/PixelStreaming/Private/VideoSource.h b/Engine/Plugins/Media/PixelStreaming/Source/PixelStreaming/Private/VideoSource.h
--- a/Engine/Plugins/Media/PixelStreaming/Source/PixelStreaming/Private/VideoSource.h
+++ b/Engine/Plugins/Media/PixelStreaming/Source/PixelStreaming/Private/VideoSource.h
@@ -2,11 +2,11 @@
 
 #pragma once
 
-#include "FixedFPSPump.h"
-#include "PixelStreamingTextureSource.h"
+#include "PixelStreamingPumpable.h"
+#include "AdaptedVideoTrackSource.h"
 #include "Settings.h"
-
-class IPixelStreamingTextureSource;
+#include "TextureTripleBuffer.h"
+#include "PixelStreamingPlayerId.h"
 
 namespace UE::PixelStreaming
 {
@@ -14,44 +14,36 @@ namespace UE::PixelStreaming
 	class FPlayerSessions;
 
 	/*
-		* Base class for pumped Pixel Streaming video sources. Video sources are WebRTC video sources that populate WebRTC tracks
-		* and pass WebRTC video frames to `OnFrame`, which eventually gets passed to a WebRTC video encoder, encoded, and transmitted.
-		*/
-	class FVideoSourceBase : public rtc::AdaptedVideoTrackSource, public IPumpedVideoSource
+	* Video sources are WebRTC video sources that populate WebRTC tracks
+	* and pass WebRTC video frames to `OnFrame`, which eventually gets passed to a WebRTC video encoder, encoded, and transmitted.
+	*/
+	class FVideoSourceBase : public AdaptedVideoTrackSource, public FPixelStreamingPumpable
 	{
 	public:
-		FVideoSourceBase(FPixelStreamingPlayerId InPlayerId);
+		FVideoSourceBase();
 		virtual ~FVideoSourceBase();
-		FPixelStreamingPlayerId GetPlayerId() const { return PlayerId; }
 		virtual void Initialize();
 
-		/* Begin rtc::AdaptedVideoTrackSource overrides */
+		/* Begin UE::PixelStreaming::AdaptedVideoTrackSource overrides */
 		virtual webrtc::MediaSourceInterface::SourceState state() const override { return CurrentState; }
 		virtual bool remote() const override { return false; }
 		virtual bool is_screencast() const override { return false; }
 		virtual absl::optional<bool> needs_denoising() const override { return false; }
-		/* End rtc::AdaptedVideoTrackSource overrides */
+		void AddRef() const override { FPixelStreamingPumpable::AddRef(); }
+		virtual rtc::RefCountReleaseStatus Release() const override { return FPixelStreamingPumpable::Release(); }
+		/* End UE::PixelStreaming::AdaptedVideoTrackSource overrides */
 
-		/* Begin IPumpedVideoSource */
+		/* Begin FPixelStreamingPumpable */
 		virtual void OnPump(int32 FrameId) override;
 		virtual bool IsReadyForPump() const = 0;
-		virtual void AddRef() const override;
-		virtual rtc::RefCountReleaseStatus Release() const override;
-		virtual bool HasOneRef() const { return RefCount.HasOneRef(); }
-		/* End IPumpedVideoSource */
-
-		/* Begin rtc::RefCountInterface */
+		/* End FPixelStreamingPumpable */
 
 	protected:
 		virtual bool AdaptCaptureFrame(const int64 TimestampUs, FIntPoint Resolution);
 		virtual webrtc::VideoFrame CreateFrame(int32 FrameId) = 0;
 
-	private:
-		mutable webrtc::webrtc_impl::RefCounter RefCount{ 0 };
-
 	protected:
 		webrtc::MediaSourceInterface::SourceState CurrentState;
-		FPixelStreamingPlayerId PlayerId;
 	};
 
 	/*
@@ -60,12 +52,12 @@ namespace UE::PixelStreaming
 	class FVideoSourceP2P : public FVideoSourceBase
 	{
 	public:
-		FVideoSourceP2P(FPixelStreamingPlayerId InPlayerId, FPlayerSessions* InSessions);
+		FVideoSourceP2P(FName SourceType, TFunction<bool()> InIsQualityControllerFunc);
 
 	protected:
-		FPlayerSessions* Sessions;
-		TSharedPtr<IPixelStreamingTextureSource> TextureSource;
+		TSharedPtr<FTextureTripleBuffer> TextureSource;
 		Settings::ECodec Codec;
+		TFunction<bool()> IsQualityControllerFunc;
 
 	protected:
 		/* Begin FVideoSourceBase */
@@ -83,16 +75,18 @@ namespace UE::PixelStreaming
 	class FVideoSourceSFU : public FVideoSourceBase
 	{
 	public:
-		FVideoSourceSFU();
+		FVideoSourceSFU(FName SourceType);
 
 	protected:
-		TArray<TSharedPtr<IPixelStreamingTextureSource>> LayerTextures;
+		TArray<TSharedPtr<FTextureTripleBuffer>> LayerTextures;
 
 	protected:
 		/* Begin FVideoSourceBase */
 		virtual webrtc::VideoFrame CreateFrame(int32 FrameId) override;
 		virtual bool IsReadyForPump() const override;
 		/* End FVideoSourceBase */
+
+		void AddLayerTexture(FName SourceType, float Scale);
 	};
 
 } // namespace UE::PixelStreaming
diff --git a/Engine/Plugins/Media/PixelStreaming/Source/PixelStreaming/Public/IPixelStreamingModule.h b/Engine/Plugins/Media/PixelStreaming/Source/PixelStreaming/Public/IPixelStreamingModule.h
--- a/Engine/Plugins/Media/PixelStreaming/Source/PixelStreaming/Public/IPixelStreamingModule.h
+++ b/Engine/Plugins/Media/PixelStreaming/Source/PixelStreaming/Public/IPixelStreamingModule.h
@@ -10,94 +10,94 @@
 #include "Templates/RefCounting.h"
 #include "PixelStreamingPlayerId.h"
 #include "IPixelStreamingAudioSink.h"
-#include "IPixelStreamingPumpedVideoSource.h"
-
+#include "PixelStreamingPumpable.h"
+#include "IPixelStreamingTextureSourceFactory.h"
 #include "IInputDevice.h"
+#include "api/media_stream_interface.h"
+#include "api/scoped_refptr.h"
+#include "api/video_codecs/video_encoder_factory.h"
+#include "api/video/video_frame.h"
+#include "api/video/video_sink_interface.h"
 
 class UTexture2D;
 class UPixelStreamingInput;
 class FRHIGPUFence;
 
-namespace webrtc
-{
-	class VideoEncoderFactory;
-}
-
 /**
-* The public interface to this module
-*/
+ * The public interface of the Pixel Streaming module.
+ */
 class PIXELSTREAMING_API IPixelStreamingModule : public IInputDeviceModule
 {
 public:
 	/**
-	* Singleton-like access to this module's interface.
-	* Beware of calling this during the shutdown phase, though.  Your module might have been unloaded already.
-	*
-	* @return Returns singleton instance, loading the module on demand if needed
-	*/
+	 * Singleton-like access to this module's interface.
+	 * Beware of calling this during the shutdown phase, though.  Your module might have been unloaded already.
+	 *
+	 * @return Returns singleton instance, loading the module on demand if needed
+	 */
 	static inline IPixelStreamingModule& Get()
 	{
 		return FModuleManager::LoadModuleChecked<IPixelStreamingModule>("PixelStreaming");
 	}
 
 	/**
-	* Checks to see if this module is loaded.  
-	*
-	* @return True if the module is loaded.
-	*/
+	 * Checks to see if this module is loaded.
+	 *
+	 * @return True if the module is loaded.
+	 */
 	static inline bool IsAvailable()
 	{
 		return FModuleManager::Get().IsModuleLoaded("PixelStreaming");
 	}
 
 	/*
-	* Event fired when internal streamer is initialized and the methods on this module are ready for use.
-	*/
+	 * Event fired when internal streamer is initialized and the methods on this module are ready for use.
+	 */
 	DECLARE_EVENT_OneParam(IPixelStreamingModule, FReadyEvent, IPixelStreamingModule&);
 
-	/*
-	* A getter for the OnReady event. Intent is for users to call IPixelStreamingModule::Get().OnReady().AddXXX.
-	* @return The bindable OnReady event.
-	*/
+	/**
+	 * A getter for the OnReady event. Intent is for users to call IPixelStreamingModule::Get().OnReady().AddXXX.
+	 * @return The bindable OnReady event.
+	 */
 	virtual FReadyEvent& OnReady() = 0;
 
-	/*
-	* Is the PixelStreaming module actually ready to use? Is the streamer created.
-	* @return True if Pixel Streaming module methods are ready for use.
-	*/
+	/**
+	 * Is the PixelStreaming module actually ready to use? Is the streamer created.
+	 * @return True if Pixel Streaming module methods are ready for use.
+	 */
 	virtual bool IsReady() = 0;
 
-	/*
-	* Connect to the specified signalling server and begin a streaming session
-	* @param SignallingServerUrl The url of the signalling server to use. eg. ws://localhost:8888
-	*/
+	/**
+	 * Connect to the specified signalling server and begin a streaming session
+	 * @param SignallingServerUrl The url of the signalling server to use. eg. ws://localhost:8888
+	 */
 	virtual bool StartStreaming(const FString& SignallingServerUrl) = 0;
 
 	/*
-	* Ends the current streaming session
-	*/
+	 * Ends the current streaming session
+	 */
 	virtual void StopStreaming() = 0;
 
 	/*
-	* Event fired when the streamer has connected to a signalling server and is ready for peers.
-	*/
+	 * Event fired when the streamer has connected to a signalling server and is ready for peers.
+	 */
 	DECLARE_EVENT_OneParam(IPixelStreamingModule, FStreamingStartedEvent, IPixelStreamingModule&);
 
 	/*
-	* A getter for the OnStreamingStarted event. Intent is for users to call IPixelStreamingModule::Get().OnStreamingStarted().AddXXX.
-	* @return The bindable OnStreamingStarted event.
-	*/
+	 * A getter for the OnStreamingStarted event. Intent is for users to call IPixelStreamingModule::Get().OnStreamingStarted().AddXXX.
+	 * @return The bindable OnStreamingStarted event.
+	 */
 	virtual FStreamingStartedEvent& OnStreamingStarted() = 0;
 
 	/*
-	* Event fired when the streamer has disconnected from a signalling server and has stopped streaming.
-	*/
+	 * Event fired when the streamer has disconnected from a signalling server and has stopped streaming.
+	 */
 	DECLARE_EVENT_OneParam(IPixelStreamingModule, FStreamingStoppedEvent, IPixelStreamingModule&);
 
 	/*
-	* A getter for the OnStreamingStopped event. Intent is for users to call IPixelStreamingModule::Get().OnStreamingStopped().AddXXX.
-	* @return The bindable OnStreamingStopped event.
-	*/
+	 * A getter for the OnStreamingStopped event. Intent is for users to call IPixelStreamingModule::Get().OnStreamingStopped().AddXXX.
+	 * @return The bindable OnStreamingStopped event.
+	 */
 	virtual FStreamingStoppedEvent& OnStreamingStopped() = 0;
 
 	/**
@@ -168,35 +168,70 @@ public:
 	 */
 	virtual void AddInputComponent(UPixelStreamingInput* InInputComponent) = 0;
 
-	/*
+	/**
 	 * Tell the input device that a pixel streaming input component is no longer
 	 * relevant.
 	 * @param InInputComponent - The pixel streaming input component which is no longer relevant.
 	 */
 	virtual void RemoveInputComponent(UPixelStreamingInput* InInputComponent) = 0;
 
-	/*
+	/**
 	 * Get the input components currently attached to Pixel Streaming.
+	 * @return An array of input components.
 	 */
 	virtual const TArray<UPixelStreamingInput*> GetInputComponents() = 0;
 
-	/*
+	/**
 	 * Create a webrtc::VideoEncoderFactory pointer.
+	 * @return A WebRTC video encoder factory with its encoders populated by Pixel Streaming.
 	 */
 	virtual webrtc::VideoEncoderFactory* CreateVideoEncoderFactory() = 0;
 
-	/*
-	 * Register a VideoSource with the module
+	/**
+	 * Register a "Pumpable" with the Pixel Streaming module and it will be pumped at the WebRTC max FPS (usually 60).
+	 * @param Pumpable - A thing that will be called at a fixed interval (typically 60fps).
 	 */
-	virtual void RegisterVideoSource(FPixelStreamingPlayerId PlayerId, IPumpedVideoSource* VideoSource) = 0;
+	virtual void RegisterPumpable(rtc::scoped_refptr<FPixelStreamingPumpable> Pumpable) = 0;
 
-	/*
-	 * Unregister a VideoSource with the module
+	/**
+	 * Unregister a "Pumpable" from the Pixel Streaming module.
+	 * @param Pumpable - The pumpable you no longer wish to pump.
 	 */
-	virtual void UnregisterVideoSource(FPixelStreamingPlayerId PlayerId) = 0;
+	virtual void UnregisterPumpable(rtc::scoped_refptr<FPixelStreamingPumpable> Pumpable) = 0;
 
-	/*
-	 * Create add a task to the GPU fence poller.
+	/**
+	 * Create a task that is continuously polled until it is done and can be short circuited.
+	 * This is useful for in Pixel Streaming for checking fences during texture copies.
+	 * @param Task - The task to execute one time, this will happen on the "Poller" thread.
+	 * @param IsTaskFinished - A check that will be called as often as possible until the task is finished.
+	 * @param bKeepRunning - A flag to short circuit the completion checking (useful if we no longer care because it took too long).
+	 */
+	virtual void AddPollerTask(TFunction<void()> Task, TFunction<bool()> IsTaskFinished, TSharedRef<bool, ESPMode::ThreadSafe> bKeepRunning) = 0;
+
+	/**
+	 * Create a WebRTC video source that is has its lifetime managed outside of Pixel Streaming.
+	 * @return A ref-counted WebRTC video source, when WebRTC is done with it the ref count will decrement by one, if that is the only remaining ref the source will be destroyed.
+	 */
+	virtual rtc::scoped_refptr<webrtc::VideoTrackSourceInterface> CreateExternalVideoSource(FName SourceType = FName(TEXT("Backbuffer"))) = 0;
+
+	/**
+	 * The Pixel Streaming module can have n active texture source types that dictate
+	 * the type of video tracks that are made when a new peer joins.
+	 * @param SourceTypes - An array of registered source types to make tracks from, e.g. "Backbuffer".
+	 */
+	virtual void SetActiveTextureSourceTypes(const TArray<FName>& SourceTypes) = 0;
+
+	/**
+	 * The n active texture source types that dictate
+	 * how many video tracks of each source type are made when a new peer joins.
+	 * @return The array of active registered source types that tracks are made from, e.g. "Backbuffer".
+	 */
+	virtual const TArray<FName>& GetActiveTextureSourceTypes() const = 0;
+
+	/**
+	 * Get the texture source factory which you can use to register new texture source types and create
+	 * texture sources
+	 * @return The texture source factory in use by the Pixel Streaming module
 	 */
-	virtual void AddGPUFencePollerTask(TRefCountPtr<FRHIGPUFence> Fence, TSharedRef<bool, ESPMode::ThreadSafe> bIsEnabled, TFunction<void()> Task) = 0;
+	virtual IPixelStreamingTextureSourceFactory& GetTextureSourceFactory() = 0;
 };
diff --git a/Engine/Plugins/Media/PixelStreaming/Source/PixelStreaming/Public/IPixelStreamingSessions.h b/Engine/Plugins/Media/PixelStreaming/Source/PixelStreaming/Public/IPixelStreamingSessions.h
--- a/Engine/Plugins/Media/PixelStreaming/Source/PixelStreaming/Public/IPixelStreamingSessions.h
+++ b/Engine/Plugins/Media/PixelStreaming/Source/PixelStreaming/Public/IPixelStreamingSessions.h
@@ -4,11 +4,10 @@
 
 #include "IPixelStreamingAudioSink.h"
 #include "PixelStreamingPlayerId.h"
-#include "ProtocolDefs.h"
 
 namespace UE::PixelStreaming
 {
-	class IPixelStreamingSessions
+	class PIXELSTREAMING_API IPixelStreamingSessions
 	{
 	public:
 		virtual int GetNumPlayers() const = 0;
diff --git a/Engine/Plugins/Media/PixelStreaming/Source/PixelStreaming/Public/IPixelStreamingTextureSourceFactory.h b/Engine/Plugins/Media/PixelStreaming/Source/PixelStreaming/Public/IPixelStreamingTextureSourceFactory.h
new file mode 100644
--- /dev/null
+++ b/Engine/Plugins/Media/PixelStreaming/Source/PixelStreaming/Public/IPixelStreamingTextureSourceFactory.h
@@ -0,0 +1,38 @@
+// Copyright Epic Games, Inc. All Rights Reserved.
+
+#pragma once
+
+#include "Templates/UniquePtr.h"
+#include "Templates/Function.h"
+#include "PixelStreamingTextureSource.h"
+
+/*
+ * A factory where "TextureSourceCreators" are registered as lambdas with unique FNames.
+ * The premise is that outside implementers can register texture source creators using the Pixel Streaming module.
+ */
+class PIXELSTREAMING_API IPixelStreamingTextureSourceFactory
+{
+public:
+	virtual ~IPixelStreamingTextureSourceFactory(){};
+
+	/**
+	 * Create a TextureSource based on a known named source, e.g. Backbuffer.
+	 * These sources can be used to capture raw frames/textures and use them as desired - typically wrapping them as WebRTC video sources.
+	 * @param ScoureType - The type of a registered texture source.
+	 * @return The texture source as a TUniquePtr - this means you own it now.
+	 */
+	virtual TUniquePtr<FPixelStreamingTextureSource> CreateTextureSource(FName SourceType) = 0;
+
+	/**
+	 * Register a texture source type that can be made active or created using the other functions on this module.
+	 * @param SourceType - The unique name to identify this type of texture source.
+	 * @param CreatorFunc - A lambda that will create the texture source of this type.
+	 */
+	virtual void RegisterTextureSourceType(FName SourceType, TFunction<TUniquePtr<FPixelStreamingTextureSource>()> CreatorFunc) = 0;
+
+	/**
+	 * Unregister this texture source type from being created.
+	 * @param SourceType - The type to unregister and remove.
+	 */
+	virtual void UnregisterTextureSourceType(FName SourceType) = 0;
+};
diff --git a/Engine/Plugins/Media/PixelStreaming/Source/PixelStreaming/Public/PixelStreamingAudioComponent.h b/Engine/Plugins/Media/PixelStreaming/Source/PixelStreaming/Public/PixelStreamingAudioComponent.h
--- a/Engine/Plugins/Media/PixelStreaming/Source/PixelStreaming/Public/PixelStreamingAudioComponent.h
+++ b/Engine/Plugins/Media/PixelStreaming/Source/PixelStreaming/Public/PixelStreamingAudioComponent.h
@@ -11,10 +11,10 @@
 #include "PixelStreamingAudioComponent.generated.h"
 
 /*
-* An `ISoundGenerator` implementation to pump some audio from WebRTC into this synth component
-*/
+ * An `ISoundGenerator` implementation to pump some audio from WebRTC into this synth component
+ */
 
-class FWebRTCSoundGenerator : public ISoundGenerator
+class PIXELSTREAMING_API FWebRTCSoundGenerator : public ISoundGenerator
 {
 public:
 	FWebRTCSoundGenerator();
@@ -78,11 +78,11 @@ protected:
 	//~ End UActorComponent interface
 
 public:
-	/** 
-	*   The Pixel Streaming player/peer whose audio we wish to listen to.
-	*   If this is left blank this component will listen to the first non-listened to peer that connects after this component is ready.
-	*   Note: that when the listened to peer disconnects this component is reset to blank and will once again listen to the next non-listened to peer that connects.
-	*/
+	/**
+	 *   The Pixel Streaming player/peer whose audio we wish to listen to.
+	 *   If this is left blank this component will listen to the first non-listened to peer that connects after this component is ready.
+	 *   Note: that when the listened to peer disconnects this component is reset to blank and will once again listen to the next non-listened to peer that connects.
+	 */
 	UPROPERTY(EditAnywhere, BlueprintReadWrite, Category = "Pixel Streaming Audio Component")
 	FString PlayerToHear;
 
diff --git a/Engine/Plugins/Media/PixelStreaming/Source/PixelStreaming/Public/PixelStreamingFrameBuffer.h b/Engine/Plugins/Media/PixelStreaming/Source/PixelStreaming/Public/PixelStreamingFrameBuffer.h
--- a/Engine/Plugins/Media/PixelStreaming/Source/PixelStreaming/Public/PixelStreamingFrameBuffer.h
+++ b/Engine/Plugins/Media/PixelStreaming/Source/PixelStreaming/Public/PixelStreamingFrameBuffer.h
@@ -4,26 +4,16 @@
 
 #include "Templates/RefCounting.h"
 #include "Templates/SharedPointer.h"
-
 #include "api/video/video_frame_buffer.h"
 #include "api/video/i420_buffer.h"
 #include "api/scoped_refptr.h"
 
-class FRHITexture;
-class IPixelStreamingTextureSource;
-
-namespace UE
+enum EPixelStreamingFrameBufferType
 {
-	namespace PixelStreaming
-	{
-		enum EFrameBufferType
-		{
-			Initialize,
-			Simulcast,
-			Layer
-		};
-	} // namespace PixelStreaming
-} // namespace UE
+	Initialize,
+	Simulcast,
+	Layer
+};
 
 /*
  * ----------------- FPixelStreamingFrameBuffer -----------------
@@ -34,7 +24,7 @@ class PIXELSTREAMING_API FPixelStreamingFrameBuffer : public webrtc::VideoFrameB
 public:
 	virtual ~FPixelStreamingFrameBuffer() {}
 
-	virtual UE::PixelStreaming::EFrameBufferType GetFrameBufferType() const = 0;
+	virtual EPixelStreamingFrameBufferType GetFrameBufferType() const = 0;
 
 	// Begin webrtc::VideoFrameBuffer interface
 	virtual webrtc::VideoFrameBuffer::Type type() const override
@@ -58,87 +48,3 @@ public:
 	}
 	// End webrtc::VideoFrameBuffer interface
 };
-
-/*
- * ----------------- FInitializeFrameBuffer -----------------
- * We use this frame to force webrtc to create encoders that will siphon off other streams
- * but never get real frames directly pumped to them.
- */
-class PIXELSTREAMING_API FInitializeFrameBuffer : public FPixelStreamingFrameBuffer
-{
-public:
-	FInitializeFrameBuffer(TSharedPtr<IPixelStreamingTextureSource> InTextureSource);
-	virtual ~FInitializeFrameBuffer();
-
-	virtual UE::PixelStreaming::EFrameBufferType GetFrameBufferType() const { return UE::PixelStreaming::EFrameBufferType::Initialize; }
-
-	virtual int width() const override;
-	virtual int height() const override;
-
-private:
-	TSharedPtr<IPixelStreamingTextureSource> TextureSource;
-};
-
-/*
- * ----------------- FSimulcastFrameBuffer -----------------
- * Holds a number of textures (called layers) and each is assigned an index.
- */
-class PIXELSTREAMING_API FSimulcastFrameBuffer : public FPixelStreamingFrameBuffer
-{
-public:
-	FSimulcastFrameBuffer(TArray<TSharedPtr<IPixelStreamingTextureSource>>& InTextureSources);
-	virtual ~FSimulcastFrameBuffer();
-
-	virtual UE::PixelStreaming::EFrameBufferType GetFrameBufferType() const { return UE::PixelStreaming::EFrameBufferType::Simulcast; }
-	TSharedPtr<IPixelStreamingTextureSource> GetLayerFrameSource(int LayerIndex) const;
-	int GetNumLayers() const;
-
-	virtual int width() const override;
-	virtual int height() const override;
-
-private:
-	TArray<TSharedPtr<IPixelStreamingTextureSource>>& TextureSources;
-};
-
-/*
- * ----------------- FLayerFrameBuffer -----------------
- * Holds a genuine single texture for encoding.
- */
-class PIXELSTREAMING_API FLayerFrameBuffer : public FPixelStreamingFrameBuffer
-{
-public:
-	FLayerFrameBuffer(TSharedPtr<IPixelStreamingTextureSource> InTextureSource);
-	virtual ~FLayerFrameBuffer();
-
-	virtual UE::PixelStreaming::EFrameBufferType GetFrameBufferType() const { return UE::PixelStreaming::EFrameBufferType::Layer; }
-	TRefCountPtr<FRHITexture> GetFrame() const;
-
-	virtual int width() const override;
-	virtual int height() const override;
-
-private:
-	TSharedPtr<IPixelStreamingTextureSource> TextureSource;
-};
-
-/*
- * ----------------- FFrameBufferI420 -----------------
- * Holds a texture source that is capable of being converted to I420 for encoding.
- */
-class PIXELSTREAMING_API FFrameBufferI420 : public FPixelStreamingFrameBuffer
-{
-public:
-	FFrameBufferI420(TSharedPtr<IPixelStreamingTextureSource> InTextureSource);
-	virtual ~FFrameBufferI420();
-
-	virtual UE::PixelStreaming::EFrameBufferType GetFrameBufferType() const { return UE::PixelStreaming::EFrameBufferType::Layer; }
-	TRefCountPtr<FRHITexture> GetFrame() const;
-
-	virtual int width() const override;
-	virtual int height() const override;
-	virtual rtc::scoped_refptr<webrtc::I420BufferInterface> ToI420() override;
-	virtual const webrtc::I420BufferInterface* GetI420() const override;
-
-private:
-	TSharedPtr<IPixelStreamingTextureSource> TextureSource;
-	rtc::scoped_refptr<webrtc::I420Buffer> Buffer;
-};
diff --git a/Engine/Plugins/Media/PixelStreaming/Source/PixelStreaming/Public/PixelStreamingProtocolDefs.h b/Engine/Plugins/Media/PixelStreaming/Source/PixelStreaming/Public/PixelStreamingProtocolDefs.h
--- a/Engine/Plugins/Media/PixelStreaming/Source/PixelStreaming/Public/PixelStreamingProtocolDefs.h
+++ b/Engine/Plugins/Media/PixelStreaming/Source/PixelStreaming/Public/PixelStreamingProtocolDefs.h
@@ -2,10 +2,6 @@
 
 #pragma once
 
-#include "HAL/Platform.h"
-#include "api/data_channel_interface.h"
-#include "WebRTCIncludes.h"
-
 namespace UE::PixelStreaming::Protocol
 {
 	enum EPlayerFlags
@@ -15,13 +11,13 @@ namespace UE::PixelStreaming::Protocol
 		PSPFlag_IsSFU = 0x02,
 	};
 
-	enum class EToStreamerMsg : uint8
+	enum class EToStreamerMsg : uint8_t
 	{
 		/**********************************************************************/
 
 		/*
-				* Control Messages. Range = 0..49.
-				*/
+		 * Control Messages. Range = 0..49.
+		 */
 
 		IFrameRequest = 0,
 		RequestQualityControl = 1,
@@ -36,8 +32,8 @@ namespace UE::PixelStreaming::Protocol
 		/**********************************************************************/
 
 		/*
-				* Input Messages. Range = 50..89.
-				*/
+		 * Input Messages. Range = 50..89.
+		 */
 
 		// Generic Input Messages. Range = 50..59.
 		UIInteraction = 50,
@@ -69,15 +65,15 @@ namespace UE::PixelStreaming::Protocol
 		/**********************************************************************/
 
 		/*
-				* Ensure Count is the final entry.
-				*/
+		 * Ensure Count is the final entry.
+		 */
 		Count
 
 		/**********************************************************************/
 	};
 
 	//! Messages that can be sent to the webrtc players
-	enum class EToPlayerMsg : uint8
+	enum class EToPlayerMsg : uint8_t
 	{
 		QualityControlOwnership = 0,
 		Response = 1,
@@ -93,27 +89,4 @@ namespace UE::PixelStreaming::Protocol
 		TestEcho = 11
 	};
 
-	template <typename T>
-	static const T& ParseBuffer(const uint8*& Data, uint32& Size)
-	{
-		checkf(sizeof(T) <= Size, TEXT("%d - %d"), sizeof(T), Size);
-		const T& Value = *reinterpret_cast<const T*>(Data);
-		Data += sizeof(T);
-		Size -= sizeof(T);
-		return Value;
-	}
-
-	static FString ParseString(const webrtc::DataBuffer& Buffer, const size_t Offset = 0)
-	{
-		FString Res;
-		if (Buffer.data.size() > Offset)
-		{
-			size_t StringLength = (Buffer.data.size() - Offset) / sizeof(TCHAR);
-			Res.GetCharArray().SetNumUninitialized(StringLength + 1);
-			FMemory::Memcpy(Res.GetCharArray().GetData(), Buffer.data.data() + Offset, StringLength * sizeof(TCHAR));
-			Res.GetCharArray()[StringLength] = 0;
-		}
-		return Res;
-	}
-
 }; // namespace UE::PixelStreaming::Protocol
diff --git a/Engine/Plugins/Media/PixelStreaming/Source/PixelStreaming/Public/PixelStreamingPumpable.h b/Engine/Plugins/Media/PixelStreaming/Source/PixelStreaming/Public/PixelStreamingPumpable.h
new file mode 100644
--- /dev/null
+++ b/Engine/Plugins/Media/PixelStreaming/Source/PixelStreaming/Public/PixelStreamingPumpable.h
@@ -0,0 +1,32 @@
+// Copyright Epic Games, Inc. All Rights Reserved.
+#pragma once
+
+#include "rtc_base/ref_counter.h"
+#include "rtc_base/ref_count.h"
+#include "HAL/Platform.h"
+
+/*
+* Base class for objects that need to be pumped at a fixed interval.
+*/
+class PIXELSTREAMING_API FPixelStreamingPumpable
+{
+public:
+	FPixelStreamingPumpable() = default;
+	virtual ~FPixelStreamingPumpable() = default;
+	virtual void OnPump(int32 FrameId) = 0;
+	virtual bool IsReadyForPump() const = 0;
+	virtual void AddRef() const { RefCount.IncRef(); }
+	virtual bool HasOneRef() const { return RefCount.HasOneRef(); }
+	virtual rtc::RefCountReleaseStatus Release() const
+	{
+		const rtc::RefCountReleaseStatus Status = RefCount.DecRef();
+		if (Status == rtc::RefCountReleaseStatus::kDroppedLastRef)
+		{
+			delete this;
+		}
+		return Status;
+	}
+
+protected:
+	mutable webrtc::webrtc_impl::RefCounter RefCount{ 0 };
+};
diff --git a/Engine/Plugins/Media/PixelStreaming/Source/PixelStreaming/Public/PixelStreamingTextureSource.h b/Engine/Plugins/Media/PixelStreaming/Source/PixelStreaming/Public/PixelStreamingTextureSource.h
--- a/Engine/Plugins/Media/PixelStreaming/Source/PixelStreaming/Public/PixelStreamingTextureSource.h
+++ b/Engine/Plugins/Media/PixelStreaming/Source/PixelStreaming/Public/PixelStreamingTextureSource.h
@@ -1,169 +1,64 @@
 // Copyright Epic Games, Inc. All Rights Reserved.
 #pragma once
 
+#include "api/video/video_frame_buffer.h"
+#include "api/video/i420_buffer.h"
+#include "api/scoped_refptr.h"
 #include "Templates/SharedPointer.h"
-#include "Templates/RefCounting.h"
-#include "Async/Async.h"
-#include "RHI.h"
-#include "HAL/ThreadSafeBool.h"
-#include "Rendering/SlateRenderer.h"
-#include "Framework/Application/SlateApplication.h"
-#include "IPixelStreamingTextureSource.h"
-#include "IPixelStreamingModule.h"
+#include "PixelStreamingTextureWrapper.h"
+#include "Delegates/Delegate.h"
 
-/*
- * Interface to wrap a Texture for use with the TextureSource
- * By extending this class you can add additional operations to a class derived from BackBufferTextureSource
- */
+//---------------------------------IPixelStreamingFrameCapturer--------------------------------------------------------
 
-class IPixelStreamingBackBufferTextureWrapper : public FRefCountBase
+class PIXELSTREAMING_API IPixelStreamingFrameCapturer
 {
 public:
-	explicit IPixelStreamingBackBufferTextureWrapper(FTexture2DRHIRef InTexture)
-	{
-		Texture = InTexture;
-	}
-
-	FTexture2DRHIRef& GetTexture()
-	{
-		return Texture;
-	}
-
-private:
-	FTexture2DRHIRef Texture;
+	virtual ~IPixelStreamingFrameCapturer() = default;
+	virtual void CaptureTexture(FPixelStreamingTextureWrapper& TextureToCopy, TSharedPtr<FPixelStreamingTextureWrapper> DestinationTexture) = 0;
+	virtual void OnCaptureFinished(TSharedPtr<FPixelStreamingTextureWrapper> CapturedTexture) = 0;
+	virtual bool IsCaptureFinished() = 0;
 };
 
-/*
- * Base class for TextureSources that get their textures from the UE backbuffer.
- * Textures are copied from the backbuffer using a triplebuffering mechanism so that texture read access is always thread safe while writes are occurring.
- * If no texture has been written since the last read then the same texture will be read again.
- * This class also has the additional functionality of scaling textures from the backbuffer.
- * Note: This is a template class that uses CRTP - see: https://en.wikipedia.org/wiki/Curiously_recurring_template_pattern
- * For derived types of this class they must contain the following static methods:
- *
- * static void CopyTexture(const FTexture2DRHIRef& SourceTexture, TRefCountPtr<IPixelStreamingBackBufferTextureWrapper> DestTexture, FGPUFenceRHIRef& CopyFence)
- * static TRefCountPtr<IPixelStreamingBackBufferTextureWrapper> CreateTexture(int Width, int Height)
- * static FTexture2DRHIRef ToTextureRef(TRefCountPtr<IPixelStreamingBackBufferTextureWrapper> Texture)
- * static FString& GetName()
- */
-
-class PIXELSTREAMING_API IBackBufferTextureSource : public IPixelStreamingTextureSource, public TSharedFromThis<IBackBufferTextureSource, ESPMode::ThreadSafe>
-{
-public:
-	IBackBufferTextureSource(float InFrameScale);
-	IBackBufferTextureSource()
-		: IBackBufferTextureSource(1.0) {}
-
-	virtual ~IBackBufferTextureSource();
-
-	TRefCountPtr<IPixelStreamingBackBufferTextureWrapper> GetCurrent();
-
-	// abstract methods for overriding
-	virtual void CopyTexture(const FTexture2DRHIRef& SourceTexture, TRefCountPtr<IPixelStreamingBackBufferTextureWrapper> DestTexture, FGPUFenceRHIRef& CopyFence) = 0;
-	virtual TRefCountPtr<IPixelStreamingBackBufferTextureWrapper> CreateTexture(int Width, int Height) = 0;
-	virtual FString GetName() const = 0;
-
-	/* Begin IPixelStreamingTextureSource interface */
-	virtual void SetEnabled(bool bInEnabled) override;
-	virtual FTexture2DRHIRef GetTexture() override { return GetCurrent()->GetTexture(); }
-
-	virtual bool IsEnabled() const override { return *bEnabled; }
-	virtual bool IsAvailable() const override { return bInitialized; }
-	virtual int GetSourceWidth() const override { return SourceWidth; }
-	virtual int GetSourceHeight() const override { return SourceHeight; }
-	/* End IPixelStreamingTextureSource interface */
-
-	float GetFrameScaling() const { return FrameScale; }
-
-private:
-	void OnBackBufferReady_RenderThread(SWindow& SlateWindow, const FTexture2DRHIRef& FrameBuffer);
-	void Initialize(int Width, int Height);
-
-	struct FCaptureFrame
-	{
-		TRefCountPtr<IPixelStreamingBackBufferTextureWrapper> TextureWrapper;
-		FGPUFenceRHIRef Fence;
-		bool bAvailable = true;
-		uint64 PreWaitingOnCopy;
-	};
-
-private:
-	const float FrameScale;
-	int SourceWidth = 0;
-	int SourceHeight = 0;
-	FThreadSafeBool bInitialized = false;
-	TSharedRef<bool, ESPMode::ThreadSafe> bEnabled = MakeShared<bool, ESPMode::ThreadSafe>(true);
-	FSlateRenderer::FOnBackBufferReadyToPresent* OnBackbuffer = nullptr;
-	FDelegateHandle BackbufferDelegateHandle;
-	FCriticalSection DestructorCriticalSection;
-
-	/*
-	 * Triple buffer setup with queued write buffers (since we have to wait for RHI copy).
-	 * 1 Read buffer (read the captured texture)
-	 * 1 Temp buffer (for swapping what is read and written)
-	 * 2 Write buffers (2 write buffers because UE can render two frames before presenting sometimes)
-	 */
-	FCriticalSection CriticalSection;
-	bool bWriteParity = true;
-	FCaptureFrame WriteBuffers[2];
-	TRefCountPtr<IPixelStreamingBackBufferTextureWrapper> TempBuffer;
-	TRefCountPtr<IPixelStreamingBackBufferTextureWrapper> ReadBuffer;
-	FThreadSafeBool bIsTempDirty;
-};
+//--------------------------------FPixelStreamingTextureSourceBase---------------------------------------------------------
 
 /*
- * Basic IPixelStreamingBackBufferTextureWrapper for use with FBackBufferTextureSource
- */
-
-class PIXELSTREAMING_API FPixelStreamingRHIBackBufferTexture : public IPixelStreamingBackBufferTextureWrapper
+* Base class for all texture sources in Pixel Streaming.
+* These texture sources are used to populate video sources for video tracks.
+* At minimum a texture source must do the following:
+* 1) Create destination textures to capture to.
+* 2) Create a capturer object that does that capturing.
+*/
+class PIXELSTREAMING_API FPixelStreamingTextureSource
 {
 public:
-	FPixelStreamingRHIBackBufferTexture(FTexture2DRHIRef InTexture)
-		: IPixelStreamingBackBufferTextureWrapper(InTexture) {}
-};
-
-class PIXELSTREAMING_API FBackBufferTextureSource : public IBackBufferTextureSource
-{
-public:
-	FBackBufferTextureSource()
-		: IBackBufferTextureSource() {}
-	FBackBufferTextureSource(float InScale)
-		: IBackBufferTextureSource(InScale) {}
-
-	virtual void CopyTexture(const FTexture2DRHIRef& SourceTexture, TRefCountPtr<IPixelStreamingBackBufferTextureWrapper> DestTexture, FGPUFenceRHIRef& CopyFence) override;
-	virtual TRefCountPtr<IPixelStreamingBackBufferTextureWrapper> CreateTexture(int Width, int Height) override;
-
-	virtual FString GetName() const override { return TEXT("FBackBufferTextureSource"); }
-};
-
-/*
- * Extention for IPixelStreamingBackBufferTextureWrapper which also stores a CPU accessible memory address of the texture
- */
-class PIXELSTREAMING_API FPixelStreamingCPUReadableBackbufferTexture : public IPixelStreamingBackBufferTextureWrapper
-{
-public:
-	FPixelStreamingCPUReadableBackbufferTexture(FTexture2DRHIRef InTexture)
-		: IPixelStreamingBackBufferTextureWrapper(InTexture) {}
-
-	TArray<FColor>& GetRawPixels()
-	{
-		return RawPixels;
-	}
-
-private:
-	TArray<FColor> RawPixels;
+	FPixelStreamingTextureSource() = default;
+	virtual ~FPixelStreamingTextureSource() = default;
+
+	/**
+	* Create a blank staging texture that is used when the source texture is captured.
+	* @param Width
+	* @param Height
+	* @return The blank "texture" - which does not have to be a UE texture.
+	*/
+	virtual TSharedPtr<FPixelStreamingTextureWrapper> CreateBlankStagingTexture(uint32 Width, uint32 Height) = 0;
+
+	/**
+	 * Creates the object that does the actual capturing of frames to textures.
+	 * @return The capturer.
+	*/
+	virtual TSharedPtr<IPixelStreamingFrameCapturer> CreateFrameCapturer() = 0;
+
+	/**
+	 *  Convert a "texture" to a WebRTC YUV I420 buffer.
+	 * @param Texture - The "texture" to convert.
+	 * @return The WebRTC I420 buffer.
+	 */
+	virtual rtc::scoped_refptr<webrtc::I420Buffer> ToWebRTCI420Buffer(TSharedPtr<FPixelStreamingTextureWrapper> Texture) = 0;
+
+	/* 
+	* Parameters for the delegate are as follows: FOnNewTexture(FPixelStreamingTextureWrapper& NewFrame, uint32 FrameWidth, uint32 FrameHeight) 
+	* The intent is for you to use whatever mechanism you want to broadcast new textures to through this delegate as they become available.
+	*/
+	DECLARE_MULTICAST_DELEGATE_ThreeParams(FOnNewTexture, FPixelStreamingTextureWrapper&, uint32, uint32);
+	FOnNewTexture OnNewTexture;
 };
-
-class PIXELSTREAMING_API FBackBufferToCPUTextureSource : public IBackBufferTextureSource
-{
-public:
-	FBackBufferToCPUTextureSource()
-		: IBackBufferTextureSource() {}
-	FBackBufferToCPUTextureSource(float InScale)
-		: IBackBufferTextureSource(InScale) {}
-
-	virtual void CopyTexture(const FTexture2DRHIRef& SourceTexture, TRefCountPtr<IPixelStreamingBackBufferTextureWrapper> DestTexture, FGPUFenceRHIRef& CopyFence) override;
-	virtual TRefCountPtr<IPixelStreamingBackBufferTextureWrapper> CreateTexture(int Width, int Height) override;
-
-	virtual FString GetName() const override { return TEXT("FBackBufferToCPUTextureSource"); }
-};
\ No newline at end of file
diff --git a/Engine/Plugins/Media/PixelStreaming/Source/PixelStreaming/Public/PixelStreamingTextureWrapper.h b/Engine/Plugins/Media/PixelStreaming/Source/PixelStreaming/Public/PixelStreamingTextureWrapper.h
new file mode 100644
--- /dev/null
+++ b/Engine/Plugins/Media/PixelStreaming/Source/PixelStreaming/Public/PixelStreamingTextureWrapper.h
@@ -0,0 +1,76 @@
+// Copyright Epic Games, Inc. All Rights Reserved.
+#pragma once
+
+#include "Templates/SharedPointer.h"
+
+//--------------------------------FPixelStreamingTextureWrapper---------------------------------------------------------
+
+/*
+* Wraps a "texture" for use with the TextureSource.
+* In this context a texture can be any type you want as long as it is useful to the eventual TextureSource that uses it. 
+*/
+class PIXELSTREAMING_API FPixelStreamingTextureWrapper
+{
+private:
+	/* A simple id incrementer, every call increments by 1 */
+	static int NextTypeId()
+	{
+		static int Id(0);
+		return Id++;
+	}
+
+	/* TypeIdOf<T> will get the unique id for specified type - this is because static values inside template are specific to each type. */
+	template <typename T>
+	static int TypeIdOf()
+	{
+		static int Id(NextTypeId());
+		return Id;
+	}
+
+	/* Base "any Texture" type - all it does it have a type id as a member. */
+	struct FPixelStreamingAnyTextureBase
+	{
+		int TypeId;
+		FPixelStreamingAnyTextureBase(const int InTypeId)
+			: TypeId(InTypeId) {}
+		virtual ~FPixelStreamingAnyTextureBase() {}
+	};
+
+	/* A wrapper around any type that could represent a texture. */
+	template <typename T>
+	struct FPixelStreamingAnyTexture : FPixelStreamingAnyTextureBase
+	{
+		T Value;
+
+		FPixelStreamingAnyTexture(const T& InValue)
+			: FPixelStreamingAnyTextureBase(TypeIdOf<T>()), Value(InValue) {}
+	};
+
+public:
+	/**
+	* Constructor where we wrap any "texture" type in this class.
+	* @param InTexture - The "texture" to wrap.
+	*/
+	template <typename T>
+	FPixelStreamingTextureWrapper(const T& InTexture)
+		: WrappedTexture(new FPixelStreamingAnyTexture<T>(InTexture))
+	{
+	}
+
+	/**
+	* Gets the underlying wrapped "texture" and attempts to cast it to type specified in <T>.
+	* Note: This method with fire an verify assertion if the type requested by this method does not exactly match the underlying type.
+	* Even the difference between Type& and Type will cause this assertion to fire.
+	* @return The "texture" type unwrapped and cast to <T>.
+	*/
+	template <typename T>
+	T GetTexture()
+	{
+		verifyf(TypeIdOf<T>() == WrappedTexture->TypeId, TEXT("Type ids mismatch, this indicates you are trying to get a texture with a different type than it was created with. %d =/= %d"), TypeIdOf<T>(), WrappedTexture->TypeId);
+		return StaticCastSharedPtr<const FPixelStreamingAnyTexture<T>>(WrappedTexture)->Value;
+	}
+
+private:
+	/* The actual underlying member that wraps the "texture". */
+	TSharedPtr<FPixelStreamingAnyTextureBase> WrappedTexture;
+};
\ No newline at end of file
diff --git a/Engine/Plugins/Media/PixelStreaming/Source/PixelStreamingShaders/PixelStreamingShaders.build.cs b/Engine/Plugins/Media/PixelStreaming/Source/PixelStreamingShaders/PixelStreamingShaders.build.cs
new file mode 100644
--- /dev/null
+++ b/Engine/Plugins/Media/PixelStreaming/Source/PixelStreamingShaders/PixelStreamingShaders.build.cs
@@ -0,0 +1,33 @@
+// Copyright Epic Games, Inc. All Rights Reserved.
+
+using System;
+using System.IO;
+using System.Collections.Generic;
+using EpicGames.Core;
+
+namespace UnrealBuildTool.Rules
+{
+	public class PixelStreamingShaders : ModuleRules
+	{
+		public PixelStreamingShaders(ReadOnlyTargetRules Target) : base(Target)
+		{
+			// This is so for game projects using our public headers don't have to include extra modules they might not know about.
+			PublicDependencyModuleNames.AddRange(new string[] {});
+
+			// NOTE: General rule is not to access the private folder of another module
+			PrivateIncludePaths.AddRange(new string[] {});
+
+			PrivateDependencyModuleNames.AddRange(new string[]
+			{
+				"Core",
+				"Engine",
+				"Projects",
+				"RenderCore",
+				"RHI",
+			});
+
+			// required for casting UE4 BackBuffer to Vulkan Texture2D for NvEnc
+			PrivateDependencyModuleNames.AddRange(new string[] {});
+		}
+	}
+}
diff --git a/Engine/Plugins/Media/PixelStreaming/Source/PixelStreamingShaders/Private/PixelStreamingShadersModule.cpp b/Engine/Plugins/Media/PixelStreaming/Source/PixelStreamingShaders/Private/PixelStreamingShadersModule.cpp
new file mode 100644
--- /dev/null
+++ b/Engine/Plugins/Media/PixelStreaming/Source/PixelStreamingShaders/Private/PixelStreamingShadersModule.cpp
@@ -0,0 +1,22 @@
+#include "CoreMinimal.h"
+#include "Containers/UnrealString.h"
+#include "Misc/Paths.h"
+#include "Modules/ModuleManager.h"
+#include "Interfaces/IPluginManager.h"
+#include "ShaderCore.h"
+
+class FPixelStreamingShadersModule : public IModuleInterface
+{
+public:
+	virtual ~FPixelStreamingShadersModule() = default;
+
+	virtual void StartupModule() override;
+};
+
+void FPixelStreamingShadersModule::StartupModule()
+{
+	FString ShaderDirectory = FPaths::Combine(IPluginManager::Get().FindPlugin(TEXT("PixelStreaming"))->GetBaseDir(), TEXT("Shaders"));
+	AddShaderSourceDirectoryMapping(TEXT("/Plugin/PixelStreaming"), ShaderDirectory);
+}
+
+IMPLEMENT_MODULE(FPixelStreamingShadersModule, PixelStreamingShaders)
