Fixing copyright boilerplate.

diff --git a/Engine/Plugins/Media/PixelStreaming/Source/PixelStreaming/Private/AdaptedVideoTrackSource.cpp b/Engine/Plugins/Media/PixelStreaming/Source/PixelStreaming/Private/AdaptedVideoTrackSource.cpp
--- a/Engine/Plugins/Media/PixelStreaming/Source/PixelStreaming/Private/AdaptedVideoTrackSource.cpp
+++ b/Engine/Plugins/Media/PixelStreaming/Source/PixelStreaming/Private/AdaptedVideoTrackSource.cpp
@@ -1,3 +1,4 @@
+// Copyright Epic Games, Inc. All Rights Reserved.
 
 #include "AdaptedVideoTrackSource.h"
 
diff --git a/Engine/Plugins/Media/PixelStreaming/Source/PixelStreamingShaders/Private/PixelStreamingShadersModule.cpp b/Engine/Plugins/Media/PixelStreaming/Source/PixelStreamingShaders/Private/PixelStreamingShadersModule.cpp
--- a/Engine/Plugins/Media/PixelStreaming/Source/PixelStreamingShaders/Private/PixelStreamingShadersModule.cpp
+++ b/Engine/Plugins/Media/PixelStreaming/Source/PixelStreamingShaders/Private/PixelStreamingShadersModule.cpp
@@ -1,3 +1,5 @@
+// Copyright Epic Games, Inc. All Rights Reserved.
+
 #include "CoreMinimal.h"
 #include "Containers/UnrealString.h"
 #include "Misc/Paths.h"
