Add PixelStreaming tests to CIS/Horde and remove test which currently doesn't meet standards for CIS/Horde

diff --git a/Engine/Plugins/Media/PixelStreaming/Source/PixelStreaming/Private/Tests/ClientRTCTests.cpp b/Engine/Plugins/Media/PixelStreaming/Source/PixelStreaming/Private/Tests/ClientRTCTests.cpp
deleted file mode 100644
diff --git a/Engine/Plugins/Media/PixelStreaming/Source/PixelStreaming/Private/Tests/SignallingServerConnectionTests.cpp b/Engine/Plugins/Media/PixelStreaming/Source/PixelStreaming/Private/Tests/SignallingServerConnectionTests.cpp
--- a/Engine/Plugins/Media/PixelStreaming/Source/PixelStreaming/Private/Tests/SignallingServerConnectionTests.cpp
+++ b/Engine/Plugins/Media/PixelStreaming/Source/PixelStreaming/Private/Tests/SignallingServerConnectionTests.cpp
@@ -141,7 +141,7 @@ namespace
 		cricket::Candidate MockCandidate;
 	};
 
-	BEGIN_DEFINE_SPEC(FSignallingServerConnectionSpec, "PixelStreaming.SignallingServerConnection", EAutomationTestFlags::ProductFilter | EAutomationTestFlags::ApplicationContextMask)
+	BEGIN_DEFINE_SPEC(FSignallingServerConnectionSpec, "System.Plugins.PixelStreaming.SignallingServerConnection",  EAutomationTestFlags::EditorContext | EAutomationTestFlags::ClientContext | EAutomationTestFlags::ProductFilter)
 	TUniquePtr<FSignallingServerConnection> SignallingServerConnection; // the object being tested
 	FTestSSConnectionObserver Observer;
 	FSignallingServerConnection::FWebSocketFactory WebSocketFactory;
