Fix for ASan error heap-use-after-free in Pixel Streaming on shutdown

diff --git a/Engine/Plugins/Media/PixelStreaming/Source/PixelStreaming/Private/PlayerSessions.cpp b/Engine/Plugins/Media/PixelStreaming/Source/PixelStreaming/Private/PlayerSessions.cpp
--- a/Engine/Plugins/Media/PixelStreaming/Source/PixelStreaming/Private/PlayerSessions.cpp
+++ b/Engine/Plugins/Media/PixelStreaming/Source/PixelStreaming/Private/PlayerSessions.cpp
@@ -193,7 +193,7 @@ namespace UE::PixelStreaming
 		return false;
 	}
 
-	void FPlayerSessions::DeleteAllPlayerSessions()
+	void FPlayerSessions::DeleteAllPlayerSessions(bool bEnginePreExit)
 	{
 		const FPixelStreamingPlayerId OldQualityController = QualityControllingPlayer;
 		TSet<FPixelStreamingPlayerId> OldPlayerIds;
@@ -211,12 +211,13 @@ namespace UE::PixelStreaming
 		}
 
 		// Notify all delegates of all the closed players
+		bool bEngineShuttingDown = (bEnginePreExit || IsEngineExitRequested());
 		UPixelStreamingDelegates* Delegates = UPixelStreamingDelegates::GetPixelStreamingDelegates();
-		if (Delegates && FModuleManager::Get().IsModuleLoaded("PixelStreaming"))
+		if (Delegates && FModuleManager::Get().IsModuleLoaded("PixelStreaming") && !bEngineShuttingDown)
 		{
 			for (auto&& PlayerId : OldPlayerIds)
 			{
-				bool bWasQualityController = OldQualityController == PlayerId;
+				bool bWasQualityController = (OldQualityController == PlayerId);
 
 				Delegates->OnClosedConnection.Broadcast(PlayerId, bWasQualityController);
 				Delegates->OnClosedConnectionNative.Broadcast(PlayerId, bWasQualityController);
diff --git a/Engine/Plugins/Media/PixelStreaming/Source/PixelStreaming/Private/PlayerSessions.h b/Engine/Plugins/Media/PixelStreaming/Source/PixelStreaming/Private/PlayerSessions.h
--- a/Engine/Plugins/Media/PixelStreaming/Source/PixelStreaming/Private/PlayerSessions.h
+++ b/Engine/Plugins/Media/PixelStreaming/Source/PixelStreaming/Private/PlayerSessions.h
@@ -28,7 +28,7 @@ namespace UE::PixelStreaming
 			FSignallingServerConnection* SignallingServerConnection);
 
 		int DeletePlayerSession(FPixelStreamingPlayerId PlayerId);
-		void DeleteAllPlayerSessions();
+		void DeleteAllPlayerSessions(bool bEnginePreExit = false);
 
 		int GetNumPlayers() const;
 
diff --git a/Engine/Plugins/Media/PixelStreaming/Source/PixelStreaming/Private/Streamer.cpp b/Engine/Plugins/Media/PixelStreaming/Source/PixelStreaming/Private/Streamer.cpp
--- a/Engine/Plugins/Media/PixelStreaming/Source/PixelStreaming/Private/Streamer.cpp
+++ b/Engine/Plugins/Media/PixelStreaming/Source/PixelStreaming/Private/Streamer.cpp
@@ -38,6 +38,10 @@ namespace UE::PixelStreaming
 		FSignallingServerConnection::FWebSocketFactory WebSocketFactory = [](const FString& Url) { return FWebSocketsModule::Get().CreateWebSocket(Url, TEXT("")); };
 		SignallingServerConnection = MakeUnique<FSignallingServerConnection>(WebSocketFactory, *this, InStreamerId);
 		StartWebRtcSignallingThread();
+		FCoreDelegates::OnPreExit.AddLambda([&]()
+		{ 
+			PlayerSessions.DeleteAllPlayerSessions(true);
+		});
 	}
 
 	FStreamer::~FStreamer()
