Fixing initial state issues after changing the FRHITextureCreateDesc creation.

diff --git a/Engine/Plugins/Media/PixelStreaming/Source/PixelStreaming/Private/TextureSourceComputeI420.cpp b/Engine/Plugins/Media/PixelStreaming/Source/PixelStreaming/Private/TextureSourceComputeI420.cpp
--- a/Engine/Plugins/Media/PixelStreaming/Source/PixelStreaming/Private/TextureSourceComputeI420.cpp
+++ b/Engine/Plugins/Media/PixelStreaming/Source/PixelStreaming/Private/TextureSourceComputeI420.cpp
@@ -80,12 +80,14 @@ namespace
 		const FRHITextureCreateDesc StagingDescY =
 			FRHITextureCreateDesc::Create2D(TEXT("YUV Output CPU Texture"), PlaneYDimensions, PF_R8)
 			.SetClearValue(FClearValueBinding::None)
-			.SetFlags(ETextureCreateFlags::CPUReadback);
+			.SetFlags(ETextureCreateFlags::CPUReadback)
+			.DetermineInititialState();
 
 		const FRHITextureCreateDesc StagingDescUV =
 			FRHITextureCreateDesc::Create2D(TEXT("YUV Output CPU Texture"), PlaneUVDimensions, PF_R8)
 			.SetClearValue(FClearValueBinding::None)
-			.SetFlags(ETextureCreateFlags::CPUReadback);
+			.SetFlags(ETextureCreateFlags::CPUReadback)
+			.DetermineInititialState();
 
 		StagingTextureY = GDynamicRHI->RHICreateTexture(StagingDescY);
 		StagingTextureU = GDynamicRHI->RHICreateTexture(StagingDescUV);
@@ -128,7 +130,8 @@ namespace
 			const FRHITextureCreateDesc Desc =
 				FRHITextureCreateDesc::Create2D(TEXT("PixelStreamingBlankTexture"), SourceTexture->GetDesc().Extent, SourceTexture->GetDesc().Format)
 				.SetClearValue(FClearValueBinding::None)
-				.SetFlags(ETextureCreateFlags::ShaderResource);
+				.SetFlags(ETextureCreateFlags::ShaderResource)
+				.DetermineInititialState();
 
 			TempFBCopy = GDynamicRHI->RHICreateTexture(Desc);
 		}
diff --git a/Engine/Source/Runtime/OpenGLDrv/Private/Android/AndroidOpenGL.cpp b/Engine/Source/Runtime/OpenGLDrv/Private/Android/AndroidOpenGL.cpp
--- a/Engine/Source/Runtime/OpenGLDrv/Private/Android/AndroidOpenGL.cpp
+++ b/Engine/Source/Runtime/OpenGLDrv/Private/Android/AndroidOpenGL.cpp
@@ -453,7 +453,8 @@ FOpenGLTexture* PlatformCreateBuiltinBackBuffer(FOpenGLDynamicRHI* OpenGLRHI, ui
 	{
 		const FOpenGLTextureCreateDesc CreateDesc =
 			FRHITextureCreateDesc::Create2D(TEXT("PlatformCreateBuiltinBackBuffer"), SizeX, SizeY, PF_B8G8R8A8)
-			.SetFlags(ETextureCreateFlags::RenderTargetable | ETextureCreateFlags::Presentable | ETextureCreateFlags::ResolveTargetable);
+			.SetFlags(ETextureCreateFlags::RenderTargetable | ETextureCreateFlags::Presentable | ETextureCreateFlags::ResolveTargetable)
+			.DetermineInititialState();
 
 		return new FOpenGLTexture(CreateDesc);
 	}
diff --git a/Engine/Source/Runtime/OpenGLDrv/Private/OpenGLTexture.cpp b/Engine/Source/Runtime/OpenGLDrv/Private/OpenGLTexture.cpp
--- a/Engine/Source/Runtime/OpenGLDrv/Private/OpenGLTexture.cpp
+++ b/Engine/Source/Runtime/OpenGLDrv/Private/OpenGLTexture.cpp
@@ -2232,7 +2232,8 @@ FTexture2DRHIRef FOpenGLDynamicRHI::RHICreateTexture2DFromResource(EPixelFormat
 		.SetClearValue(ClearValueBinding)
 		.SetFlags(TexCreateFlags)
 		.SetNumMips(NumMips)
-		.SetNumSamples(NumSamples);
+		.SetNumSamples(NumSamples)
+		.DetermineInititialState();
 
 	return new FOpenGLTexture(Desc, Resource);
 }
@@ -2244,7 +2245,8 @@ FTexture2DRHIRef FOpenGLDynamicRHI::RHICreateTexture2DArrayFromResource(EPixelFo
 		.SetClearValue(ClearValueBinding)
 		.SetFlags(TexCreateFlags)
 		.SetNumMips(NumMips)
-		.SetNumSamples(NumSamples);
+		.SetNumSamples(NumSamples)
+		.DetermineInititialState();
 
 	return new FOpenGLTexture(Desc, Resource);
 }
@@ -2259,7 +2261,8 @@ FTextureCubeRHIRef FOpenGLDynamicRHI::RHICreateTextureCubeFromResource(EPixelFor
 		.SetClearValue(ClearValueBinding)
 		.SetFlags(TexCreateFlags)
 		.SetNumMips(NumMips)
-		.SetNumSamples(NumSamples);
+		.SetNumSamples(NumSamples)
+		.DetermineInititialState();
 
 	return new FOpenGLTexture(Desc, Resource);
 }
diff --git a/Engine/Source/Runtime/OpenGLDrv/Private/OpenGLViewport.cpp b/Engine/Source/Runtime/OpenGLDrv/Private/OpenGLViewport.cpp
--- a/Engine/Source/Runtime/OpenGLDrv/Private/OpenGLViewport.cpp
+++ b/Engine/Source/Runtime/OpenGLDrv/Private/OpenGLViewport.cpp
@@ -318,7 +318,8 @@ void FOpenGLViewport::Resize(uint32 InSizeX,uint32 InSizeY,bool bInIsFullscreen)
 			const FRHITextureCreateDesc Desc =
 				FRHITextureCreateDesc::Create2D(TEXT("FOpenGLViewport"), InSizeX, InSizeY, PixelFormat)
 				.SetClearValue(FClearValueBinding::Transparent)
-				.SetFlags(ETextureCreateFlags::RenderTargetable | ETextureCreateFlags::ResolveTargetable);
+				.SetFlags(ETextureCreateFlags::RenderTargetable | ETextureCreateFlags::ResolveTargetable)
+				.DetermineInititialState();
 
 			BackBuffer = new FOpenGLTexture(Desc);
 		}
diff --git a/Engine/Source/Runtime/OpenGLDrv/Private/Windows/OpenGLWindows.cpp b/Engine/Source/Runtime/OpenGLDrv/Private/Windows/OpenGLWindows.cpp
--- a/Engine/Source/Runtime/OpenGLDrv/Private/Windows/OpenGLWindows.cpp
+++ b/Engine/Source/Runtime/OpenGLDrv/Private/Windows/OpenGLWindows.cpp
@@ -1006,7 +1006,8 @@ FOpenGLTexture* PlatformCreateBuiltinBackBuffer(FOpenGLDynamicRHI* OpenGLRHI, ui
 		const FRHITextureCreateDesc Desc =
 			FRHITextureCreateDesc::Create2D(TEXT("PlatformCreateBuiltinBackBuffer"), SizeX, SizeY, PF_B8G8R8A8)
 			.SetClearValue(FClearValueBinding::Transparent)
-			.SetFlags(ETextureCreateFlags::RenderTargetable | ETextureCreateFlags::Presentable | ETextureCreateFlags::ResolveTargetable);
+			.SetFlags(ETextureCreateFlags::RenderTargetable | ETextureCreateFlags::Presentable | ETextureCreateFlags::ResolveTargetable)
+			.DetermineInititialState();
 
 		return new FOpenGLTexture(Desc);
 	}
diff --git a/Engine/Source/Runtime/RHI/Public/RHIResources.h b/Engine/Source/Runtime/RHI/Public/RHIResources.h
--- a/Engine/Source/Runtime/RHI/Public/RHIResources.h
+++ b/Engine/Source/Runtime/RHI/Public/RHIResources.h
@@ -1651,6 +1651,7 @@ struct FRHITextureCreateDesc : public FRHITextureDesc
 	FRHITextureCreateDesc& SetDebugName(const TCHAR* InDebugName)              { DebugName = InDebugName;                  return *this; }
 	FRHITextureCreateDesc& SetGPUMask(FRHIGPUMask InGPUMask)                   { GPUMask = InGPUMask;                      return *this; }
 	FRHITextureCreateDesc& SetBulkData(FResourceBulkDataInterface* InBulkData) { BulkData = InBulkData;                    return *this; }
+	FRHITextureCreateDesc& DetermineInititialState() { if (InitialState == ERHIAccess::Unknown) InitialState = RHIGetDefaultResourceState(Flags, BulkData != nullptr); return *this; }
 
 	/* The RHI access state that the resource will be created in. */
 	ERHIAccess InitialState = ERHIAccess::Unknown;
diff --git a/Engine/Source/Runtime/VulkanRHI/Private/VulkanRHI.cpp b/Engine/Source/Runtime/VulkanRHI/Private/VulkanRHI.cpp
--- a/Engine/Source/Runtime/VulkanRHI/Private/VulkanRHI.cpp
+++ b/Engine/Source/Runtime/VulkanRHI/Private/VulkanRHI.cpp
@@ -1432,7 +1432,8 @@ FTexture2DRHIRef FVulkanDynamicRHI::RHICreateTexture2DFromResource(EPixelFormat
 		.SetClearValue(EnumHasAnyFlags(Flags, ETextureCreateFlags::DepthStencilTargetable) ? FClearValueBinding::DepthZero : FClearValueBinding::Transparent)
 		.SetFlags(Flags)
 		.SetNumMips(NumMips)
-		.SetNumSamples(NumSamples);
+		.SetNumSamples(NumSamples)
+		.DetermineInititialState();
 
 	return new FVulkanTexture(*Device, Desc, Resource, false);
 }
@@ -1444,7 +1445,8 @@ FTexture2DArrayRHIRef FVulkanDynamicRHI::RHICreateTexture2DArrayFromResource(EPi
 		.SetClearValue(EnumHasAnyFlags(Flags, ETextureCreateFlags::DepthStencilTargetable) ? FClearValueBinding::DepthZero : FClearValueBinding::Transparent)
 		.SetFlags(Flags)
 		.SetNumMips(NumMips)
-		.SetNumSamples(NumSamples);
+		.SetNumSamples(NumSamples)
+		.DetermineInititialState();
 
 	return new FVulkanTexture(*Device, Desc, Resource, false);
 }
@@ -1458,7 +1460,8 @@ FTextureCubeRHIRef FVulkanDynamicRHI::RHICreateTextureCubeFromResource(EPixelFor
 		.SetFormat(Format)
 		.SetClearValue(EnumHasAnyFlags(Flags, ETextureCreateFlags::DepthStencilTargetable) ? FClearValueBinding::DepthZero : FClearValueBinding::Transparent)
 		.SetFlags(Flags)
-		.SetNumMips(NumMips);
+		.SetNumMips(NumMips)
+		.DetermineInititialState();
 
 	return new FVulkanTexture(*Device, Desc, Resource, false);
 }
diff --git a/Engine/Source/Runtime/VulkanRHI/Private/VulkanSwapChain.cpp b/Engine/Source/Runtime/VulkanRHI/Private/VulkanSwapChain.cpp
--- a/Engine/Source/Runtime/VulkanRHI/Private/VulkanSwapChain.cpp
+++ b/Engine/Source/Runtime/VulkanRHI/Private/VulkanSwapChain.cpp
@@ -918,7 +918,8 @@ void FVulkanSwapChain::CreateQCOMDepthStencil(const FVulkanTexture& InSurface) c
 		.SetClearValue(FClearValueBinding::None)
 		.SetFlags(UEFlags)
 		.SetNumMips(Desc.NumMips)
-		.SetNumSamples(Desc.NumSamples);
+		.SetNumSamples(Desc.NumSamples)
+		.DetermineInititialState();
 
 	QCOMDepthStencilSurface = new FVulkanTexture(Device, CreateDesc, nullptr);
 
diff --git a/Engine/Source/Runtime/VulkanRHI/Private/VulkanTexture.cpp b/Engine/Source/Runtime/VulkanRHI/Private/VulkanTexture.cpp
--- a/Engine/Source/Runtime/VulkanRHI/Private/VulkanTexture.cpp
+++ b/Engine/Source/Runtime/VulkanRHI/Private/VulkanTexture.cpp
@@ -1063,7 +1063,8 @@ FTexture2DRHIRef FVulkanDynamicRHI::AsyncReallocateTexture2D_RenderThread(FRHICo
 		.SetClearValue(OldDesc.ClearValue)
 		.SetFlags(OldDesc.Flags)
 		.SetNumMips(NewMipCount)
-		.SetNumSamples(OldDesc.NumSamples);
+		.SetNumSamples(OldDesc.NumSamples)
+		.DetermineInititialState();
 
 	FVulkanTexture* NewTexture = new FVulkanTexture(*Device, Desc, nullptr);
 	FVulkanCommandListContext& Context = (FVulkanCommandListContext&)RHICmdList.GetContext().GetLowestLevelContext();
@@ -1083,7 +1084,8 @@ FTexture2DRHIRef FVulkanDynamicRHI::RHIAsyncReallocateTexture2D(FRHITexture2D* O
 		.SetClearValue(OldDesc.ClearValue)
 		.SetFlags(OldDesc.Flags)
 		.SetNumMips(NewMipCount)
-		.SetNumSamples(OldDesc.NumSamples);
+		.SetNumSamples(OldDesc.NumSamples)
+		.DetermineInititialState();
 
 	FVulkanTexture* NewTexture = new FVulkanTexture(*Device, Desc, nullptr);
 
diff --git a/Engine/Source/Runtime/VulkanRHI/Private/VulkanViewport.cpp b/Engine/Source/Runtime/VulkanRHI/Private/VulkanViewport.cpp
--- a/Engine/Source/Runtime/VulkanRHI/Private/VulkanViewport.cpp
+++ b/Engine/Source/Runtime/VulkanRHI/Private/VulkanViewport.cpp
@@ -15,7 +15,7 @@
 
 
 FVulkanBackBuffer::FVulkanBackBuffer(FVulkanDevice& Device, FVulkanViewport* InViewport, EPixelFormat Format, uint32 SizeX, uint32 SizeY, ETextureCreateFlags UEFlags)
-	: FVulkanTexture(Device, FRHITextureCreateDesc::Create2D(TEXT("FVulkanBackBuffer"), SizeX, SizeY, Format).SetFlags(UEFlags), VK_NULL_HANDLE, false)
+	: FVulkanTexture(Device, FRHITextureCreateDesc::Create2D(TEXT("FVulkanBackBuffer"), SizeX, SizeY, Format).SetFlags(UEFlags).DetermineInititialState(), VK_NULL_HANDLE, false)
 	, Viewport(InViewport)
 {
 }
diff --git a/Engine/Source/Runtime/Windows/D3D11RHI/Private/D3D11Texture.cpp b/Engine/Source/Runtime/Windows/D3D11RHI/Private/D3D11Texture.cpp
--- a/Engine/Source/Runtime/Windows/D3D11RHI/Private/D3D11Texture.cpp
+++ b/Engine/Source/Runtime/Windows/D3D11RHI/Private/D3D11Texture.cpp
@@ -1001,7 +1001,8 @@ FTextureRHIRef FD3D11DynamicRHI::RHIAsyncCreateTexture2D(uint32 SizeX, uint32 Si
 		FRHITextureCreateDesc::Create2D(TEXT("RHIAsyncCreateTexture2D"), SizeX, SizeY, (EPixelFormat)Format)
 		.SetClearValue(FClearValueBinding::None)
 		.SetFlags(Flags)
-		.SetNumMips(NumMips);
+		.SetNumMips(NumMips)
+		.DetermineInititialState();
 
 	FTextureRHIRef Texture = CreateD3D11Texture2D(Desc, SubresourceData);
 
@@ -1953,7 +1954,7 @@ FD3D11Texture* FD3D11DynamicRHI::CreateTextureFromResource(bool bTextureArray, b
 		.SetFlags(TexCreateFlags)
 		.SetNumMips(TextureDesc.MipLevels)
 		.SetNumSamples(TextureDesc.SampleDesc.Count)
-		.SetInitialState(RHIGetDefaultResourceState(TexCreateFlags, false));
+		.DetermineInititialState();
 
 	FD3D11Texture* Texture2D = new FD3D11Texture(
 		RHITextureDesc,
diff --git a/Engine/Source/Runtime/Windows/D3D11RHI/Private/D3D11Viewport.cpp b/Engine/Source/Runtime/Windows/D3D11RHI/Private/D3D11Viewport.cpp
--- a/Engine/Source/Runtime/Windows/D3D11RHI/Private/D3D11Viewport.cpp
+++ b/Engine/Source/Runtime/Windows/D3D11RHI/Private/D3D11Viewport.cpp
@@ -183,7 +183,7 @@ FD3D11Texture* FD3D11Viewport::GetSwapChainSurface(FD3D11DynamicRHI* D3DRHI, EPi
 	const FRHITextureCreateDesc CreateDesc =
 		FRHITextureCreateDesc::Create2D(TEXT("FD3D11Viewport::GetSwapChainSurface"), TextureDesc.Width, TextureDesc.Height, PixelFormat)
 		.SetFlags(ETextureCreateFlags::RenderTargetable)
-		.SetInitialState(RHIGetDefaultResourceState(ETextureCreateFlags::RenderTargetable, false));
+		.DetermineInititialState();
 
 	FD3D11Texture* NewTexture = new FD3D11Texture(
 		CreateDesc,
