Early Pixel Streaming 5.1 features, bugfixes, and refactors that benefit main. (From various authors)
*************
Features
*************
- Added `stat pixelstreaminggraphs` that show realtime graphs of Pixel Streaming stats.
- Added -PixelStreamingInputController=Host (or Any) to optionally restrict stream inputs to a single host peer. The default is -PixelStreamingInputController=Any
- Add PixelStreamingServers utility to launch PixelStreaming servers from C++. (Added accompanying test too).
- Added tests for FPixelStreamingSourceFrame, FPixelStreamingFrameMetadata, FPixelStreamingFrameAdapterProcess and FPixelStreamingFrameAdapter
- Added the ability to override the default InputDevice.
- Added the ability to have multiple streamers.
*************
Bugfixes
*************
- Fix so that NVENC now only needs to scale frames as needed, there were some extraneous scaling happening before.
- Fix on screen hud stats on startup that were not dismissable.
- Fix PixelStreaming module state cleanup that was causing packaging crashes.
- Fix bug where SFU peers were unable to receive data channels messages from the streamer.
- Fix failing "connect" test where quality & input controller messages were being sent too early (before datachannel was established) for P2P peers.
- Fix PixelStreamingAudioComponent finding trying to find audio sink on a data-only peer and crashing.
- Fix PIXELSTREAMING_API exports on some files.
- Fix PixelStreaming.build.cs that was too harsh about excluding any file with 'node' in the name.
- Fix PixelStreaming.build.cs to exclude packaging the platform_scripts/cmd/coturn dir.
- Fix Vulkan fast texture copy codepath to not crash in ssome edgecases.
*************
Refactors
*************
- Frames now flow through a much more rigid set of operations to get streamed.
- Video Sources have a Video Adapter and are fed by a Video Input.
- Video Adapters are fed by the input and contains a number of `Processes`.
- Processes take the input frame and adapt it to a format that the encoder can use.
- Frame Buffers now contains a reference to the adapter, frame data is grabbed from here as needed by the encoder.
- Adapters and their processes are also public with the intention that users can implement their own adapt step for other encoders.
- Streamer is also now modular so we can have multiple streamers with their own inputs that stream to their own location.
- Extracted out the copy texture into a rendering util and tweaked how it handles transitions.
- Modify PixelStreamingModule to use new functions for creating a streamer.
- Added optional codepaths to use WebRTC M96 or M84.
- Refactor plugin to enable using our API for the new Pixel Streaming Player plugin.
- Created a new PixelStreamingPeerConnection class to wrap the webrtc peer connection with specific pixelstreaming behaviour.
- Made the signalling server connection public and renamed its observer methods so its obvious when they're implemented.
- Some small changes to frame buffers so make handling of cached data slightly safer.
- Made use of lambdas more explicit where possible.
- Remove unnecessary warning message from PixelStreamingAudioComponent.

diff --git a/Engine/Plugins/Media/PixelStreaming/PixelStreaming.uplugin b/Engine/Plugins/Media/PixelStreaming/PixelStreaming.uplugin
--- a/Engine/Plugins/Media/PixelStreaming/PixelStreaming.uplugin
+++ b/Engine/Plugins/Media/PixelStreaming/PixelStreaming.uplugin
@@ -28,7 +28,7 @@
             "LoadingPhase": "PostConfigInit",
             "PlatformAllowList": [ "Win64", "Linux" ],
             "TargetDenyList": [ "Server" ]
-		},
+		}
 	],
 	"Plugins": [
 		{
diff --git a/Engine/Plugins/Media/PixelStreaming/Source/PixelStreaming/Private/AdaptedVideoTrackSource.cpp b/Engine/Plugins/Media/PixelStreaming/Source/PixelStreaming/Private/AdaptedVideoTrackSource.cpp
--- a/Engine/Plugins/Media/PixelStreaming/Source/PixelStreaming/Private/AdaptedVideoTrackSource.cpp
+++ b/Engine/Plugins/Media/PixelStreaming/Source/PixelStreaming/Private/AdaptedVideoTrackSource.cpp
@@ -4,16 +4,20 @@
 
 namespace UE::PixelStreaming
 {
-	AdaptedVideoTrackSource::AdaptedVideoTrackSource() = default;
+	FAdaptedVideoTrackSource::FAdaptedVideoTrackSource() = default;
 
-	AdaptedVideoTrackSource::AdaptedVideoTrackSource(int required_alignment)
+	FAdaptedVideoTrackSource::FAdaptedVideoTrackSource(int required_alignment)
 		: video_adapter_(required_alignment) {}
 
-	AdaptedVideoTrackSource::~AdaptedVideoTrackSource() = default;
+	FAdaptedVideoTrackSource::~FAdaptedVideoTrackSource() = default;
 
-	bool AdaptedVideoTrackSource::GetStats(Stats* stats)
+	bool FAdaptedVideoTrackSource::GetStats(Stats* stats)
 	{
+#if WEBRTC_VERSION == 84
 		rtc::CritScope lock(&stats_crit_);
+#elif WEBRTC_VERSION == 96
+		webrtc::MutexLock lock(&stats_mutex_);
+#endif
 
 		if (!stats_)
 		{
@@ -24,7 +28,7 @@ namespace UE::PixelStreaming
 		return true;
 	}
 
-	void AdaptedVideoTrackSource::OnFrame(const webrtc::VideoFrame& frame)
+	void FAdaptedVideoTrackSource::OnFrame(const webrtc::VideoFrame& frame)
 	{
 		rtc::scoped_refptr<webrtc::VideoFrameBuffer> buffer(
 			frame.video_frame_buffer());
@@ -52,7 +56,7 @@ namespace UE::PixelStreaming
 		}
 	}
 
-	void AdaptedVideoTrackSource::AddOrUpdateSink(
+	void FAdaptedVideoTrackSource::AddOrUpdateSink(
 		rtc::VideoSinkInterface<webrtc::VideoFrame>* sink,
 		const rtc::VideoSinkWants& wants)
 	{
@@ -60,25 +64,25 @@ namespace UE::PixelStreaming
 		OnSinkWantsChanged(broadcaster_.wants());
 	}
 
-	void AdaptedVideoTrackSource::RemoveSink(
+	void FAdaptedVideoTrackSource::RemoveSink(
 		rtc::VideoSinkInterface<webrtc::VideoFrame>* sink)
 	{
 		broadcaster_.RemoveSink(sink);
 		OnSinkWantsChanged(broadcaster_.wants());
 	}
 
-	bool AdaptedVideoTrackSource::apply_rotation()
+	bool FAdaptedVideoTrackSource::apply_rotation()
 	{
 		return broadcaster_.wants().rotation_applied;
 	}
 
-	void AdaptedVideoTrackSource::OnSinkWantsChanged(
+	void FAdaptedVideoTrackSource::OnSinkWantsChanged(
 		const rtc::VideoSinkWants& wants)
 	{
 		video_adapter_.OnSinkWants(wants);
 	}
 
-	bool AdaptedVideoTrackSource::AdaptFrame(int width,
+	bool FAdaptedVideoTrackSource::AdaptFrame(int width,
 		int height,
 		int64_t time_us,
 		int* out_width,
@@ -89,7 +93,11 @@ namespace UE::PixelStreaming
 		int* crop_y)
 	{
 		{
+#if WEBRTC_VERSION == 84
 			rtc::CritScope lock(&stats_crit_);
+#elif WEBRTC_VERSION == 96
+			webrtc::MutexLock lock(&stats_mutex_);
+#endif
 			stats_ = Stats{ width, height };
 		}
 
diff --git a/Engine/Plugins/Media/PixelStreaming/Source/PixelStreaming/Private/AdaptedVideoTrackSource.h b/Engine/Plugins/Media/PixelStreaming/Source/PixelStreaming/Private/AdaptedVideoTrackSource.h
--- a/Engine/Plugins/Media/PixelStreaming/Source/PixelStreaming/Private/AdaptedVideoTrackSource.h
+++ b/Engine/Plugins/Media/PixelStreaming/Source/PixelStreaming/Private/AdaptedVideoTrackSource.h
@@ -13,15 +13,12 @@
 #pragma once
 
 #include <stdint.h>
-
 #include "WebRTCIncludes.h"
 
-
 /*
-* This code is nearly identical to the version found inside the WebRTC source file:
-* media/base/adapted_video_track_source.h
-*/
-
+ * This code is nearly identical to the version found inside the WebRTC source file:
+ * media/base/adapted_video_track_source.h
+ */
 
 namespace UE::PixelStreaming
 {
@@ -29,22 +26,32 @@ namespace UE::PixelStreaming
 	// capture sources. Sinks must be added and removed on one and only
 	// one thread, while AdaptFrame and OnFrame may be called on any
 	// thread.
-	class AdaptedVideoTrackSource
-		: public webrtc::Notifier<webrtc::VideoTrackSourceInterface>
+	class FAdaptedVideoTrackSource : public webrtc::Notifier<webrtc::VideoTrackSourceInterface>
 	{
 	public:
-		AdaptedVideoTrackSource();
-		~AdaptedVideoTrackSource() override;
-				
+		FAdaptedVideoTrackSource();
+		~FAdaptedVideoTrackSource() override;
+
 		// Implements rtc::VideoSourceInterface.
 		void AddOrUpdateSink(rtc::VideoSinkInterface<webrtc::VideoFrame>* sink,
 			const rtc::VideoSinkWants& wants) override;
 		void RemoveSink(rtc::VideoSinkInterface<webrtc::VideoFrame>* sink) override;
 
+		virtual void AddRef() const override { RefCount.IncRef(); }
+		virtual rtc::RefCountReleaseStatus Release() const override
+		{
+			const rtc::RefCountReleaseStatus Status = RefCount.DecRef();
+			if (Status == rtc::RefCountReleaseStatus::kDroppedLastRef)
+			{
+				delete this;
+			}
+			return Status;
+		}
+
 	protected:
 		// Allows derived classes to initialize |video_adapter_| with a custom
 		// alignment.
-		explicit AdaptedVideoTrackSource(int required_alignment);
+		explicit FAdaptedVideoTrackSource(int required_alignment);
 		// Checks the apply_rotation() flag. If the frame needs rotation, and it is a
 		// plain memory frame, it is rotated. Subclasses producing native frames must
 		// handle apply_rotation() themselves.
@@ -73,6 +80,8 @@ namespace UE::PixelStreaming
 		cricket::VideoAdapter* video_adapter() { return &video_adapter_; }
 
 	private:
+		mutable webrtc::webrtc_impl::RefCounter RefCount{ 0 };
+
 		// Part of VideoTrackSourceInterface.
 		bool GetStats(Stats* stats) override;
 
@@ -88,8 +97,13 @@ namespace UE::PixelStreaming
 
 		cricket::VideoAdapter video_adapter_;
 
+#if WEBRTC_VERSION == 84
 		rtc::CriticalSection stats_crit_;
 		absl::optional<Stats> stats_ RTC_GUARDED_BY(stats_crit_);
+#elif WEBRTC_VERSION == 96
+		webrtc::Mutex stats_mutex_;
+		absl::optional<Stats> stats_ RTC_GUARDED_BY(stats_mutex_);
+#endif
 
 		rtc::VideoBroadcaster broadcaster_;
 	};
diff --git a/Engine/Plugins/Media/PixelStreaming/Source/PixelStreaming/Private/AudioSubmixCapturer.h b/Engine/Plugins/Media/PixelStreaming/Source/PixelStreaming/Private/AudioSubmixCapturer.h
new file mode 100644
--- /dev/null
+++ b/Engine/Plugins/Media/PixelStreaming/Source/PixelStreaming/Private/AudioSubmixCapturer.h
@@ -0,0 +1,51 @@
+// Copyright Epic Games, Inc. All Rights Reserved.
+
+#pragma once
+
+#include "IPixelStreamingAudioSubmixCapturer.h"
+#include "AudioDevice.h"
+
+namespace UE::PixelStreaming
+{
+	// Captures audio from UE and passes it along to WebRTC.
+	class FAudioSubmixCapturer : public IPixelStreamingAudioSubmixCapturer, public ISubmixBufferListener
+	{
+	public:
+		// This magic number is the max volume used in webrtc fake audio device
+		// module.
+		static const uint32_t MaxVolumeLevel = 14392;
+
+		FAudioSubmixCapturer();
+		virtual ~FAudioSubmixCapturer() = default;
+
+		virtual bool Init() override;
+		virtual bool IsInitialised() const override;
+		virtual bool IsCapturing() const override;
+		virtual void Uninitialise() override;
+		virtual bool StartCapturing() override;
+		virtual bool EndCapturing() override;
+		virtual uint32_t GetVolume() const override;
+		virtual void SetVolume(uint32_t NewVolume) override;
+		virtual void RegisterAudioCallback(webrtc::AudioTransport* AudioCb) override;
+
+		// ISubmixBufferListener interface
+		void OnNewSubmixBuffer(const USoundSubmix* OwningSubmix, float* AudioData,
+			int32 NumSamples, int32 NumChannels,
+			const int32 SampleRate, double AudioClock) override;
+
+	private:
+		int32 GetSamplesPerDurationSecs(float InSeconds) const;
+
+	private:
+		FThreadSafeBool bInitialised;
+		FThreadSafeBool bCapturing;
+		uint32 TargetSampleRate;
+		uint8 TargetNumChannels;
+		bool bReportedSampleRateMismatch;
+		webrtc::AudioTransport* AudioCallback;
+		uint32_t VolumeLevel;
+		TArray<int16_t> RecordingBuffer;
+		FCriticalSection CriticalSection; // One thread captures audio from UE, the
+										  // other controls the state of the capturer.
+	};
+} // namespace UE::PixelStreaming
diff --git a/Engine/Plugins/Media/PixelStreaming/Source/PixelStreaming/Private/DataChannelObserver.cpp b/Engine/Plugins/Media/PixelStreaming/Source/PixelStreaming/Private/DataChannelObserver.cpp
--- a/Engine/Plugins/Media/PixelStreaming/Source/PixelStreaming/Private/DataChannelObserver.cpp
+++ b/Engine/Plugins/Media/PixelStreaming/Source/PixelStreaming/Private/DataChannelObserver.cpp
@@ -13,15 +13,32 @@
 
 namespace UE::PixelStreaming
 {
-	FDataChannelObserver::FDataChannelObserver(FPlayerSessions* InPlayerSessions, FPixelStreamingPlayerId InPlayerId)
+	FDataChannelObserver::FDataChannelObserver(FPlayerSessions* InPlayerSessions, FPixelStreamingPlayerId InPlayerId, EDataChannelDirection InDirection, TSharedPtr<IPixelStreamingInputDevice> InInputDevice)
 		: PlayerSessions(InPlayerSessions)
 		, PlayerId(InPlayerId)
-		, InputDevice(static_cast<FInputDevice&>(IPixelStreamingModule::Get().GetInputDevice()))
+		, Direction(InDirection)
+		, InputDevice(InInputDevice)
 	{
 	}
 
+	FDataChannelObserver::~FDataChannelObserver()
+	{
+		Unregister();
+	}
+
+	bool FDataChannelObserver::IsDataChannelOpen() const
+	{
+		return DataChannel && DataChannel->state() == webrtc::DataChannelInterface::DataState::kOpen;
+	}
+
 	void FDataChannelObserver::Register(rtc::scoped_refptr<webrtc::DataChannelInterface> InDataChannel)
 	{
+		checkf(InDataChannel, TEXT("Datachannel we are registering must be non null."));
+		if (!InDataChannel)
+		{
+			return;
+		}
+
 		DataChannel = InDataChannel;
 		DataChannel->RegisterObserver(this);
 	}
@@ -42,16 +59,21 @@ namespace UE::PixelStreaming
 			return;
 		}
 
-		switch (DataChannel->state())
+		const webrtc::DataChannelInterface::DataState State = DataChannel->state();
+		FString StateStr = FString(webrtc::DataChannelInterface::DataStateString(State));
+		FString LabelStr = FString(DataChannel->label().c_str());
+		UE_LOG(LogPixelStreaming, Log, TEXT("Player %s data channel (label: %s) - State=%s."), *PlayerId, *LabelStr, *StateStr);
+
+		switch (State)
 		{
 			case webrtc::DataChannelInterface::DataState::kOpen:
 			{
-				if (UPixelStreamingDelegates* Delegates = UPixelStreamingDelegates::GetPixelStreamingDelegates())
-				{
-					Delegates->OnDataChannelOpenNative.Broadcast(PlayerId, DataChannel.get());
-				}
+				OnDataChannelOpen();
 				break;
 			}
+			case webrtc::DataChannelInterface::DataState::kConnecting:
+			case webrtc::DataChannelInterface::DataState::kClosing:
+				break;
 			case webrtc::DataChannelInterface::DataState::kClosed:
 			{
 				if (UPixelStreamingDelegates* Delegates = UPixelStreamingDelegates::GetPixelStreamingDelegates())
@@ -66,8 +88,34 @@ namespace UE::PixelStreaming
 		}
 	}
 
+	void FDataChannelObserver::OnDataChannelOpen()
+	{
+		// Send all initial settings to the peer
+		SendInitialSettings();
+
+		// Input controller/quality controller
+		SendPeerControllerMessages();
+
+		// Let an delegate subscribers know data channel is open
+		if (UPixelStreamingDelegates* Delegates = UPixelStreamingDelegates::GetPixelStreamingDelegates())
+		{
+			Delegates->OnDataChannelOpenNative.Broadcast(PlayerId, DataChannel.get());
+		}
+	}
+
+	void FDataChannelObserver::SetDirection(EDataChannelDirection InDirection)
+	{
+		Direction = InDirection;
+	}
+
 	void FDataChannelObserver::OnMessage(const webrtc::DataBuffer& Buffer)
 	{
+		// If we can't recieve then early exit here.
+		if (Direction == EDataChannelDirection::SendOnly)
+		{
+			return;
+		}
+
 		Protocol::EToStreamerMsg MsgType = static_cast<Protocol::EToStreamerMsg>(Buffer.data.data()[0]);
 
 		if (MsgType == Protocol::EToStreamerMsg::RequestQualityControl)
@@ -89,18 +137,54 @@ namespace UE::PixelStreaming
 			const size_t StringLen = (Buffer.data.size() - 1) / sizeof(TCHAR);
 			const TCHAR* StringPtr = reinterpret_cast<const TCHAR*>(Buffer.data.data() + 1);
 			const FString EchoMessage(StringLen, StringPtr);
-			PlayerSessions->ForSession(PlayerId, [&](TSharedPtr<IPlayerSession> Session) {
+			PlayerSessions->ForSession(PlayerId, [&EchoMessage](TSharedPtr<IPlayerSession> Session) {
 				Session->SendMessage(Protocol::EToPlayerMsg::TestEcho, EchoMessage);
 			});
 		}
 		else if (!IsEngineExitRequested())
 		{
-			InputDevice.OnMessage(Buffer);
+			if (UE::PixelStreaming::Settings::GetInputControllerMode() == UE::PixelStreaming::Settings::EInputControllerMode::Host)
+			{
+				// If we are in "Host" mode and the current peer is not the host, then discard this input.
+				if (!PlayerSessions->IsInputController(PlayerId))
+				{
+					return;
+				}
+			}
+
+			if (InputDevice)
+			{
+				InputDevice->OnMessage(Buffer);
+			}
+		}
+	}
+
+	void FDataChannelObserver::SendPeerControllerMessages() const
+	{
+		// Only send if we can actually send.
+		if (Direction == EDataChannelDirection::RecvOnly)
+		{
+			return;
 		}
+
+		FPixelStreamingPlayerId ThisPlayerId = PlayerId;
+		FPlayerSessions* ThisPlayerSessions = PlayerSessions;
+
+		// Send quality and input control status messages to peer
+		PlayerSessions->ForSession(PlayerId, [ThisPlayerSessions, ThisPlayerId](TSharedPtr<IPlayerSession> Session) {
+			Session->SendInputControlStatus(ThisPlayerSessions->IsInputController(ThisPlayerId));
+			Session->SendQualityControlStatus(ThisPlayerSessions->IsQualityController(ThisPlayerId));
+		});
 	}
 
 	void FDataChannelObserver::SendLatencyReport() const
 	{
+		// Only send if we can actually send.
+		if (Direction == EDataChannelDirection::RecvOnly)
+		{
+			return;
+		}
+
 		if (Settings::CVarPixelStreamingDisableLatencyTester.GetValueOnAnyThread())
 		{
 			return;
@@ -141,7 +225,7 @@ namespace UE::PixelStreaming
 					TransmissionTimeMs);
 			}
 
-			PlayerSessions->ForSession(PlayerId, [&](TSharedPtr<IPlayerSession> Session) {
+			PlayerSessions->ForSession(PlayerId, [&ReportToTransmitJSON](TSharedPtr<IPlayerSession> Session) {
 				Session->SendMessage(Protocol::EToPlayerMsg::LatencyTest, ReportToTransmitJSON);
 			});
 		});
@@ -165,6 +249,12 @@ namespace UE::PixelStreaming
 
 	void FDataChannelObserver::SendInitialSettings() const
 	{
+		// Only send if we can actually send.
+		if (Direction == EDataChannelDirection::RecvOnly)
+		{
+			return;
+		}
+
 		const FString PixelStreamingPayload = FString::Printf(TEXT("{ \"AllowPixelStreamingCommands\": %s, \"DisableLatencyTest\": %s }"),
 			Settings::CVarPixelStreamingAllowConsoleCommands.GetValueOnAnyThread() ? TEXT("true") : TEXT("false"),
 			Settings::CVarPixelStreamingDisableLatencyTester.GetValueOnAnyThread() ? TEXT("true") : TEXT("false"));
@@ -188,7 +278,7 @@ namespace UE::PixelStreaming
 
 		const FString FullPayload = FString::Printf(TEXT("{ \"PixelStreaming\": %s, \"Encoder\": %s, \"WebRTC\": %s }"), *PixelStreamingPayload, *EncoderPayload, *WebRTCPayload);
 
-		PlayerSessions->ForSession(PlayerId, [&](TSharedPtr<IPlayerSession> Session) {
+		PlayerSessions->ForSession(PlayerId, [&FullPayload](TSharedPtr<IPlayerSession> Session) {
 			if (!Session->SendMessage(Protocol::EToPlayerMsg::InitialSettings, FullPayload))
 			{
 				UE_LOG(LogPixelStreaming, Log, TEXT("Failed to send initial Pixel Streaming settings."));
diff --git a/Engine/Plugins/Media/PixelStreaming/Source/PixelStreaming/Private/DataChannelObserver.h b/Engine/Plugins/Media/PixelStreaming/Source/PixelStreaming/Private/DataChannelObserver.h
--- a/Engine/Plugins/Media/PixelStreaming/Source/PixelStreaming/Private/DataChannelObserver.h
+++ b/Engine/Plugins/Media/PixelStreaming/Source/PixelStreaming/Private/DataChannelObserver.h
@@ -6,18 +6,26 @@
 #include "PixelStreamingPlayerId.h"
 #include "Delegates/Delegate.h"
 #include "Delegates/DelegateCombinations.h"
+#include "IPixelStreamingInputDevice.h"
 
 namespace UE::PixelStreaming
 {
 	class FPlayerSessions;
-	class FInputDevice;
+
+	enum EDataChannelDirection
+	{
+		Bidirectional,
+		SendOnly,
+		RecvOnly
+	};
 
 	// Observer on a particular player/peer's data channel.
 	class FDataChannelObserver : public webrtc::DataChannelObserver
 	{
 
 	public:
-		FDataChannelObserver(FPlayerSessions* InSessions, FPixelStreamingPlayerId InPlayerId);
+		FDataChannelObserver(FPlayerSessions* InSessions, FPixelStreamingPlayerId InPlayerId, EDataChannelDirection InDirection, TSharedPtr<IPixelStreamingInputDevice> InInputDevice);
+		virtual ~FDataChannelObserver();
 
 		// Begin webrtc::DataChannelObserver
 
@@ -32,18 +40,22 @@ namespace UE::PixelStreaming
 
 		// End webrtc::DataChannelObserver
 
+		bool IsDataChannelOpen() const;
 		void SendInitialSettings() const;
-
 		void Register(rtc::scoped_refptr<webrtc::DataChannelInterface> InDataChannel);
 		void Unregister();
 
 	private:
+		void SendPeerControllerMessages() const;
 		void SendLatencyReport() const;
+		void OnDataChannelOpen();
+		void SetDirection(EDataChannelDirection InDirection);
 
 	private:
 		rtc::scoped_refptr<webrtc::DataChannelInterface> DataChannel;
 		FPlayerSessions* PlayerSessions;
 		FPixelStreamingPlayerId PlayerId;
-		FInputDevice& InputDevice;
+		EDataChannelDirection Direction = EDataChannelDirection::Bidirectional;
+		TSharedPtr<IPixelStreamingInputDevice> InputDevice;
 	};
 } // namespace UE::PixelStreaming
diff --git a/Engine/Plugins/Media/PixelStreaming/Source/PixelStreaming/Private/EncoderFrameFactory.cpp b/Engine/Plugins/Media/PixelStreaming/Source/PixelStreaming/Private/EncoderFrameFactory.cpp
--- a/Engine/Plugins/Media/PixelStreaming/Source/PixelStreaming/Private/EncoderFrameFactory.cpp
+++ b/Engine/Plugins/Media/PixelStreaming/Source/PixelStreaming/Private/EncoderFrameFactory.cpp
@@ -30,567 +30,567 @@ namespace UE::PixelStreaming
 
 	FEncoderFrameFactory::~FEncoderFrameFactory()
 	{
-	FlushFrames();
+		FlushFrames();
 	}
 
 	void FEncoderFrameFactory::FlushFrames()
 	{
-	TextureToFrameMapping.Empty();
+		TextureToFrameMapping.Empty();
 	}
 
 	void FEncoderFrameFactory::RemoveStaleTextures()
 	{
-	// Remove any textures whose only reference is the one held by this class
+		// Remove any textures whose only reference is the one held by this class
 
-	TMap<FTextureRHIRef, TSharedPtr<AVEncoder::FVideoEncoderInputFrame>>::TIterator Iter = TextureToFrameMapping.CreateIterator();
-	for(; Iter; ++Iter)
-	{
-		FTextureRHIRef& Tex = Iter.Key();
+		TMap<FTextureRHIRef, TSharedPtr<AVEncoder::FVideoEncoderInputFrame>>::TIterator Iter = TextureToFrameMapping.CreateIterator();
+		for (; Iter; ++Iter)
+		{
+			FTextureRHIRef& Tex = Iter.Key();
 
-		if(Tex.GetRefCount() == 1) {
-			Iter.RemoveCurrent();
+			if (Tex.GetRefCount() == 1)
+			{
+				Iter.RemoveCurrent();
+			}
 		}
 	}
-	}
 
 	TSharedPtr<AVEncoder::FVideoEncoderInputFrame> FEncoderFrameFactory::GetOrCreateFrame(const FTextureRHIRef InTexture)
 	{
-	check(EncoderInput.IsValid());
+		check(EncoderInput.IsValid());
 
-	RemoveStaleTextures();
+		RemoveStaleTextures();
 
-	TSharedPtr<AVEncoder::FVideoEncoderInputFrame> OutFrame;	
+		TSharedPtr<AVEncoder::FVideoEncoderInputFrame> OutFrame;
 
-	if (TextureToFrameMapping.Contains(InTexture))
-	{
-		OutFrame = *(TextureToFrameMapping.Find(InTexture));
-	}
-	else
-	{
-		// A frame needs to be created
-		// NOTE: This create and move is necessary as this constructor is the only way we can create a new TSharedPtr with a custom deleter.
-		TSharedPtr<AVEncoder::FVideoEncoderInputFrame> NewFrame(
-			EncoderInput->CreateBuffer([this](const AVEncoder::FVideoEncoderInputFrame* ReleasedFrame) { /* OnReleased */ }), 
-			[this](AVEncoder::FVideoEncoderInputFrame* InFrame){ EncoderInput->DestroyBuffer(InFrame); }
-		);
-		OutFrame = MoveTemp(NewFrame);
-
-		SetTexture(OutFrame, InTexture);
-		TextureToFrameMapping.Add(InTexture, OutFrame);
-	}
+		if (TextureToFrameMapping.Contains(InTexture))
+		{
+			OutFrame = *(TextureToFrameMapping.Find(InTexture));
+		}
+		else
+		{
+			// A frame needs to be created
+			// NOTE: This create and move is necessary as this constructor is the only way we can create a new TSharedPtr with a custom deleter.
+			TSharedPtr<AVEncoder::FVideoEncoderInputFrame> NewFrame(
+				EncoderInput->CreateBuffer([&EncoderInput = EncoderInput](const AVEncoder::FVideoEncoderInputFrame* ReleasedFrame) { /* OnReleased */ }),
+				[&EncoderInput = EncoderInput](AVEncoder::FVideoEncoderInputFrame* InFrame) { EncoderInput->DestroyBuffer(InFrame); });
+			OutFrame = MoveTemp(NewFrame);
+
+			SetTexture(OutFrame, InTexture);
+			TextureToFrameMapping.Add(InTexture, OutFrame);
+		}
 
-	OutFrame->SetWidth(InTexture->GetDesc().Extent.X);
-	OutFrame->SetHeight(InTexture->GetDesc().Extent.Y);
-	OutFrame->SetFrameID(++FrameId);
-	return OutFrame;
+		OutFrame->SetWidth(InTexture->GetDesc().Extent.X);
+		OutFrame->SetHeight(InTexture->GetDesc().Extent.Y);
+		OutFrame->SetFrameID(++FrameId);
+		return OutFrame;
 	}
 
 	TSharedPtr<AVEncoder::FVideoEncoderInputFrame> FEncoderFrameFactory::GetFrameAndSetTexture(FTextureRHIRef InTexture)
 	{
-	check(EncoderInput.IsValid());
+		check(EncoderInput.IsValid());
 
-	TSharedPtr<AVEncoder::FVideoEncoderInputFrame> Frame = GetOrCreateFrame(InTexture);
+		TSharedPtr<AVEncoder::FVideoEncoderInputFrame> Frame = GetOrCreateFrame(InTexture);
 
-	return Frame;
+		return Frame;
 	}
 
 	TSharedPtr<AVEncoder::FVideoEncoderInput> FEncoderFrameFactory::GetOrCreateVideoEncoderInput()
 	{
-	if (!EncoderInput.IsValid())
-	{
-		EncoderInput = CreateVideoEncoderInput();
-	}
+		if (!EncoderInput.IsValid())
+		{
+			EncoderInput = CreateVideoEncoderInput();
+		}
 
-	return EncoderInput;
+		return EncoderInput;
 	}
 
 	TSharedPtr<AVEncoder::FVideoEncoderInput> FEncoderFrameFactory::CreateVideoEncoderInput() const
 	{
-	if (!GDynamicRHI)
-	{
-		UE_LOG(LogPixelStreaming, Error, TEXT("GDynamicRHI not valid for some reason."));
-		return nullptr;
-	}
+		if (!GDynamicRHI)
+		{
+			UE_LOG(LogPixelStreaming, Error, TEXT("GDynamicRHI not valid for some reason."));
+			return nullptr;
+		}
 
-	const ERHIInterfaceType RHIType = RHIGetInterfaceType();
+		const ERHIInterfaceType RHIType = RHIGetInterfaceType();
 
-	// Consider if we want to support runtime resolution changing?
-	bool bIsResizable = false;
+		// Consider if we want to support runtime resolution changing?
+		bool bIsResizable = false;
 
-	if (RHIType == ERHIInterfaceType::Vulkan)
-	{
-		if (IsRHIDeviceAMD())
+		if (RHIType == ERHIInterfaceType::Vulkan)
 		{
-			IVulkanDynamicRHI* DynamicRHI = GetDynamicRHI<IVulkanDynamicRHI>();
-			AVEncoder::FVulkanDataStruct VulkanData = { DynamicRHI->RHIGetVkInstance(), DynamicRHI->RHIGetVkPhysicalDevice(), DynamicRHI->RHIGetVkDevice() };
+			if (IsRHIDeviceAMD())
+			{
+				IVulkanDynamicRHI* DynamicRHI = GetDynamicRHI<IVulkanDynamicRHI>();
+				AVEncoder::FVulkanDataStruct VulkanData = { DynamicRHI->RHIGetVkInstance(), DynamicRHI->RHIGetVkPhysicalDevice(), DynamicRHI->RHIGetVkDevice() };
 
-			return AVEncoder::FVideoEncoderInput::CreateForVulkan(&VulkanData, bIsResizable);
+				return AVEncoder::FVideoEncoderInput::CreateForVulkan(&VulkanData, bIsResizable);
+			}
+			else if (IsRHIDeviceNVIDIA())
+			{
+				if (FModuleManager::Get().IsModuleLoaded("CUDA"))
+				{
+					return AVEncoder::FVideoEncoderInput::CreateForCUDA(FModuleManager::GetModuleChecked<FCUDAModule>("CUDA").GetCudaContext(), bIsResizable);
+				}
+				else
+				{
+					UE_LOG(LogPixelStreaming, Error, TEXT("CUDA module is not loaded!"));
+					return nullptr;
+				}
+			}
 		}
-		else if (IsRHIDeviceNVIDIA())
+#if PLATFORM_WINDOWS
+		else if (RHIType == ERHIInterfaceType::D3D11)
 		{
-			if (FModuleManager::Get().IsModuleLoaded("CUDA"))
+			if (IsRHIDeviceAMD())
 			{
-				return AVEncoder::FVideoEncoderInput::CreateForCUDA(FModuleManager::GetModuleChecked<FCUDAModule>("CUDA").GetCudaContext(), bIsResizable);
+				return AVEncoder::FVideoEncoderInput::CreateForD3D11(GDynamicRHI->RHIGetNativeDevice(), bIsResizable, true);
 			}
-			else
+			else if (IsRHIDeviceNVIDIA())
 			{
-				UE_LOG(LogPixelStreaming, Error, TEXT("CUDA module is not loaded!"));
-				return nullptr;
+				return AVEncoder::FVideoEncoderInput::CreateForD3D11(GDynamicRHI->RHIGetNativeDevice(), bIsResizable, false);
 			}
 		}
-	}
-	#if PLATFORM_WINDOWS
-	else if (RHIType == ERHIInterfaceType::D3D11)
-	{
-		if (IsRHIDeviceAMD())
-		{
-			return AVEncoder::FVideoEncoderInput::CreateForD3D11(GDynamicRHI->RHIGetNativeDevice(), bIsResizable, true);
-		}
-		else if (IsRHIDeviceNVIDIA())
+		else if (RHIType == ERHIInterfaceType::D3D12)
 		{
-			return AVEncoder::FVideoEncoderInput::CreateForD3D11(GDynamicRHI->RHIGetNativeDevice(), bIsResizable, false);
-		}
-	}
-	else if (RHIType == ERHIInterfaceType::D3D12)
-	{
-		if (IsRHIDeviceAMD())
-		{
-			return AVEncoder::FVideoEncoderInput::CreateForD3D12(GDynamicRHI->RHIGetNativeDevice(), bIsResizable, false);
-		}
-		else if (IsRHIDeviceNVIDIA())
-		{
-			return AVEncoder::FVideoEncoderInput::CreateForCUDA(FModuleManager::GetModuleChecked<FCUDAModule>("CUDA").GetCudaContext(), bIsResizable);
+			if (IsRHIDeviceAMD())
+			{
+				return AVEncoder::FVideoEncoderInput::CreateForD3D12(GDynamicRHI->RHIGetNativeDevice(), bIsResizable, false);
+			}
+			else if (IsRHIDeviceNVIDIA())
+			{
+				return AVEncoder::FVideoEncoderInput::CreateForCUDA(FModuleManager::GetModuleChecked<FCUDAModule>("CUDA").GetCudaContext(), bIsResizable);
+			}
 		}
-	}
-	#endif
+#endif
 
-	UE_LOG(LogPixelStreaming, Error, TEXT("Current RHI %s is not supported in Pixel Streaming"), GDynamicRHI->GetName());
-	return nullptr;
+		UE_LOG(LogPixelStreaming, Error, TEXT("Current RHI %s is not supported in Pixel Streaming"), GDynamicRHI->GetName());
+		return nullptr;
 	}
 
 	void FEncoderFrameFactory::SetTexture(TSharedPtr<AVEncoder::FVideoEncoderInputFrame> InputFrame, const FTextureRHIRef& Texture)
 	{
-	const ERHIInterfaceType RHIType = RHIGetInterfaceType();
+		const ERHIInterfaceType RHIType = RHIGetInterfaceType();
 
-	// VULKAN
-	if (RHIType == ERHIInterfaceType::Vulkan)
-	{
-		if (IsRHIDeviceAMD())
-		{
-			const VkImage VulkanImage = GetIVulkanDynamicRHI()->RHIGetVkImage(Texture.GetReference());
-			InputFrame->SetTexture(VulkanImage, [](VkImage NativeTexture) { /* Do something with released texture if needed */ });
-		}
-		else if (IsRHIDeviceNVIDIA())
-		{
-			SetTextureCUDAVulkan(InputFrame, Texture);
-		}
-		else
+		// VULKAN
+		if (RHIType == ERHIInterfaceType::Vulkan)
 		{
-			UE_LOG(LogPixelStreaming, Error, TEXT("Pixel Streaming only supports AMD and NVIDIA devices, this device is neither of those."));
+			if (IsRHIDeviceAMD())
+			{
+				const VkImage VulkanImage = GetIVulkanDynamicRHI()->RHIGetVkImage(Texture.GetReference());
+				InputFrame->SetTexture(VulkanImage, [](VkImage NativeTexture) { /* Do something with released texture if needed */ });
+			}
+			else if (IsRHIDeviceNVIDIA())
+			{
+				SetTextureCUDAVulkan(InputFrame, Texture);
+			}
+			else
+			{
+				UE_LOG(LogPixelStreaming, Error, TEXT("Pixel Streaming only supports AMD and NVIDIA devices, this device is neither of those."));
+			}
 		}
-	}
-	#if PLATFORM_WINDOWS
-	// TODO: Fix CUDA DX11 (Using CUDA as a bridge between DX11 and NVENC currently produces garbled results)
-	// DX11
-	else if (RHIType == ERHIInterfaceType::D3D11)
-	{
-		InputFrame->SetTexture((ID3D11Texture2D*)Texture->GetNativeResource(), [](ID3D11Texture2D* NativeTexture) { /* Do something with released texture if needed */ });
-	}
-	// DX12
-	else if (RHIType == ERHIInterfaceType::D3D12)
-	{
-		if (IsRHIDeviceAMD())
+#if PLATFORM_WINDOWS
+		// TODO: Fix CUDA DX11 (Using CUDA as a bridge between DX11 and NVENC currently produces garbled results)
+		// DX11
+		else if (RHIType == ERHIInterfaceType::D3D11)
 		{
-			InputFrame->SetTexture((ID3D12Resource*)Texture->GetNativeResource(), [](ID3D12Resource* NativeTexture) { /* Do something with released texture if needed */ });
+			InputFrame->SetTexture((ID3D11Texture2D*)Texture->GetNativeResource(), [](ID3D11Texture2D* NativeTexture) { /* Do something with released texture if needed */ });
 		}
-		else if (IsRHIDeviceNVIDIA())
+		// DX12
+		else if (RHIType == ERHIInterfaceType::D3D12)
 		{
-			SetTextureCUDAD3D12(InputFrame, Texture);
+			if (IsRHIDeviceAMD())
+			{
+				InputFrame->SetTexture((ID3D12Resource*)Texture->GetNativeResource(), [](ID3D12Resource* NativeTexture) { /* Do something with released texture if needed */ });
+			}
+			else if (IsRHIDeviceNVIDIA())
+			{
+				SetTextureCUDAD3D12(InputFrame, Texture);
+			}
+			else
+			{
+				UE_LOG(LogPixelStreaming, Error, TEXT("Pixel Streaming only supports AMD and NVIDIA devices, this device is neither of those."));
+			}
 		}
+#endif // PLATFORM_WINDOWS
 		else
 		{
-			UE_LOG(LogPixelStreaming, Error, TEXT("Pixel Streaming only supports AMD and NVIDIA devices, this device is neither of those."));
+			UE_LOG(LogPixelStreaming, Error, TEXT("Pixel Streaming does not support this RHI - %s"), GDynamicRHI->GetName());
 		}
 	}
-	#endif // PLATFORM_WINDOWS
-	else
-	{
-		UE_LOG(LogPixelStreaming, Error, TEXT("Pixel Streaming does not support this RHI - %s"), GDynamicRHI->GetName());
-	}
-	}
 
 	void FEncoderFrameFactory::SetTextureCUDAVulkan(TSharedPtr<AVEncoder::FVideoEncoderInputFrame> InputFrame, const FTextureRHIRef& Texture)
 	{
-	IVulkanDynamicRHI* VulkanRHI = GetIVulkanDynamicRHI();
+		IVulkanDynamicRHI* VulkanRHI = GetIVulkanDynamicRHI();
 
-	const FVulkanRHIAllocationInfo TextureAllocationInfo = VulkanRHI->RHIGetAllocationInfo(Texture.GetReference());
+		const FVulkanRHIAllocationInfo TextureAllocationInfo = VulkanRHI->RHIGetAllocationInfo(Texture.GetReference());
 
-	VkDevice Device = VulkanRHI->RHIGetVkDevice();
+		VkDevice Device = VulkanRHI->RHIGetVkDevice();
 
 #if PLATFORM_WINDOWS
-	HANDLE Handle;
-	// It is recommended to use NT handles where available, but these are only supported from Windows 8 onward, for earliers versions of Windows
-	// we need to use a Win7 style handle. NT handles require us to close them when we are done with them to prevent memory leaks.
-	// Refer to remarks section of https://docs.microsoft.com/en-us/windows/win32/api/dxgi1_2/nf-dxgi1_2-idxgiresource1-createsharedhandle
-	bool bUseNTHandle = IsWindows8OrGreater();
-
-	{
-		// Generate VkMemoryGetWin32HandleInfoKHR
-		VkMemoryGetWin32HandleInfoKHR MemoryGetHandleInfoKHR = {};
-		MemoryGetHandleInfoKHR.sType = VK_STRUCTURE_TYPE_MEMORY_GET_WIN32_HANDLE_INFO_KHR;
-		MemoryGetHandleInfoKHR.pNext = NULL;
-		MemoryGetHandleInfoKHR.memory = TextureAllocationInfo.Handle;
-		MemoryGetHandleInfoKHR.handleType = bUseNTHandle ? VK_EXTERNAL_MEMORY_HANDLE_TYPE_OPAQUE_WIN32_BIT : VK_EXTERNAL_MEMORY_HANDLE_TYPE_OPAQUE_WIN32_KMT_BIT;
-
-		PFN_vkGetMemoryWin32HandleKHR GetMemoryWin32HandleKHR = (PFN_vkGetMemoryWin32HandleKHR)VulkanRHI->RHIGetVkDeviceProcAddr("vkGetMemoryWin32HandleKHR");
-		VERIFYVULKANRESULT_EXTERNAL(GetMemoryWin32HandleKHR(Device, &MemoryGetHandleInfoKHR, &Handle));
-	}
-
-	FCUDAModule::CUDA().cuCtxPushCurrent(FModuleManager::GetModuleChecked<FCUDAModule>("CUDA").GetCudaContext());
-
-	CUexternalMemory MappedExternalMemory = nullptr;
+		HANDLE Handle;
+		// It is recommended to use NT handles where available, but these are only supported from Windows 8 onward, for earliers versions of Windows
+		// we need to use a Win7 style handle. NT handles require us to close them when we are done with them to prevent memory leaks.
+		// Refer to remarks section of https://docs.microsoft.com/en-us/windows/win32/api/dxgi1_2/nf-dxgi1_2-idxgiresource1-createsharedhandle
+		bool bUseNTHandle = IsWindows8OrGreater();
 
-	{
-		// generate a cudaExternalMemoryHandleDesc
-		CUDA_EXTERNAL_MEMORY_HANDLE_DESC CudaExtMemHandleDesc = {};
-		CudaExtMemHandleDesc.type = bUseNTHandle ? CU_EXTERNAL_MEMORY_HANDLE_TYPE_OPAQUE_WIN32 : CU_EXTERNAL_MEMORY_HANDLE_TYPE_OPAQUE_WIN32_KMT;
-		CudaExtMemHandleDesc.handle.win32.name = NULL;
-		CudaExtMemHandleDesc.handle.win32.handle = Handle;
-		CudaExtMemHandleDesc.size = TextureAllocationInfo.Offset + TextureAllocationInfo.Size;
-
-		// import external memory
-		CUresult Result = FCUDAModule::CUDA().cuImportExternalMemory(&MappedExternalMemory, &CudaExtMemHandleDesc);
-		if (Result != CUDA_SUCCESS)
 		{
-			UE_LOG(LogPixelStreaming, Error, TEXT("Failed to import external memory from vulkan error: %d"), Result);
+			// Generate VkMemoryGetWin32HandleInfoKHR
+			VkMemoryGetWin32HandleInfoKHR MemoryGetHandleInfoKHR = {};
+			MemoryGetHandleInfoKHR.sType = VK_STRUCTURE_TYPE_MEMORY_GET_WIN32_HANDLE_INFO_KHR;
+			MemoryGetHandleInfoKHR.pNext = NULL;
+			MemoryGetHandleInfoKHR.memory = TextureAllocationInfo.Handle;
+			MemoryGetHandleInfoKHR.handleType = bUseNTHandle ? VK_EXTERNAL_MEMORY_HANDLE_TYPE_OPAQUE_WIN32_BIT : VK_EXTERNAL_MEMORY_HANDLE_TYPE_OPAQUE_WIN32_KMT_BIT;
+
+			PFN_vkGetMemoryWin32HandleKHR GetMemoryWin32HandleKHR = (PFN_vkGetMemoryWin32HandleKHR)VulkanRHI->RHIGetVkDeviceProcAddr("vkGetMemoryWin32HandleKHR");
+			VERIFYVULKANRESULT_EXTERNAL(GetMemoryWin32HandleKHR(Device, &MemoryGetHandleInfoKHR, &Handle));
 		}
-	}
-
-	// Only store handle to be closed on frame destruction if it is an NT handle
-	Handle = bUseNTHandle ? Handle : NULL;
-	#else
-	void* Handle = nullptr;
 
-	// Get the CUarray to that textures memory making sure the clear it when done
-	int Fd;
-
-	{
-		// Generate VkMemoryGetFdInfoKHR
-		VkMemoryGetFdInfoKHR MemoryGetFdInfoKHR = {};
-		MemoryGetFdInfoKHR.sType = VK_STRUCTURE_TYPE_MEMORY_GET_FD_INFO_KHR;
-		MemoryGetFdInfoKHR.pNext = NULL;
-		MemoryGetFdInfoKHR.memory = TextureAllocationInfo.Handle;
-		MemoryGetFdInfoKHR.handleType = VK_EXTERNAL_MEMORY_HANDLE_TYPE_OPAQUE_FD_BIT_KHR;
-
-		PFN_vkGetMemoryFdKHR FPGetMemoryFdKHR = (PFN_vkGetMemoryFdKHR)VulkanRHI->RHIGetVkDeviceProcAddr("vkGetMemoryFdKHR");
-		VERIFYVULKANRESULT_EXTERNAL(FPGetMemoryFdKHR(Device, &MemoryGetFdInfoKHR, &Fd));
-	}
-
-	FCUDAModule::CUDA().cuCtxPushCurrent(FModuleManager::GetModuleChecked<FCUDAModule>("CUDA").GetCudaContext());
+		FCUDAModule::CUDA().cuCtxPushCurrent(FModuleManager::GetModuleChecked<FCUDAModule>("CUDA").GetCudaContext());
 
-	CUexternalMemory MappedExternalMemory = nullptr;
+		CUexternalMemory MappedExternalMemory = nullptr;
 
-	{
-		// generate a cudaExternalMemoryHandleDesc
-		CUDA_EXTERNAL_MEMORY_HANDLE_DESC CudaExtMemHandleDesc = {};
-		CudaExtMemHandleDesc.type = CU_EXTERNAL_MEMORY_HANDLE_TYPE_OPAQUE_FD;
-		CudaExtMemHandleDesc.handle.fd = Fd;
-		CudaExtMemHandleDesc.size = TextureAllocationInfo.Offset + TextureAllocationInfo.Size;
-
-		// import external memory
-		CUresult Result = FCUDAModule::CUDA().cuImportExternalMemory(&MappedExternalMemory, &CudaExtMemHandleDesc);
-		if (Result != CUDA_SUCCESS)
 		{
-			UE_LOG(LogPixelStreaming, Error, TEXT("Failed to import external memory from vulkan error: %d"), Result);
+			// generate a cudaExternalMemoryHandleDesc
+			CUDA_EXTERNAL_MEMORY_HANDLE_DESC CudaExtMemHandleDesc = {};
+			CudaExtMemHandleDesc.type = bUseNTHandle ? CU_EXTERNAL_MEMORY_HANDLE_TYPE_OPAQUE_WIN32 : CU_EXTERNAL_MEMORY_HANDLE_TYPE_OPAQUE_WIN32_KMT;
+			CudaExtMemHandleDesc.handle.win32.name = NULL;
+			CudaExtMemHandleDesc.handle.win32.handle = Handle;
+			CudaExtMemHandleDesc.size = TextureAllocationInfo.Offset + TextureAllocationInfo.Size;
+
+			// import external memory
+			CUresult Result = FCUDAModule::CUDA().cuImportExternalMemory(&MappedExternalMemory, &CudaExtMemHandleDesc);
+			if (Result != CUDA_SUCCESS)
+			{
+				UE_LOG(LogPixelStreaming, Error, TEXT("Failed to import external memory from vulkan error: %d"), Result);
+			}
 		}
-	}
-
-	#endif
 
-	CUmipmappedArray MappedMipArray = nullptr;
-	CUarray MappedArray = nullptr;
+		// Only store handle to be closed on frame destruction if it is an NT handle
+		Handle = bUseNTHandle ? Handle : NULL;
+#else
+		void* Handle = nullptr;
 
-	{
-		CUDA_EXTERNAL_MEMORY_MIPMAPPED_ARRAY_DESC MipmapDesc = {};
-		MipmapDesc.numLevels = 1;
-		MipmapDesc.offset = TextureAllocationInfo.Offset;
-		MipmapDesc.arrayDesc.Width = Texture->GetDesc().Extent.X;
-		MipmapDesc.arrayDesc.Height = Texture->GetDesc().Extent.Y;
-		MipmapDesc.arrayDesc.Depth = 0;
-		MipmapDesc.arrayDesc.NumChannels = 4;
-		MipmapDesc.arrayDesc.Format = CU_AD_FORMAT_UNSIGNED_INT8;
-		MipmapDesc.arrayDesc.Flags = CUDA_ARRAY3D_SURFACE_LDST | CUDA_ARRAY3D_COLOR_ATTACHMENT;
+		// Get the CUarray to that textures memory making sure the clear it when done
+		int Fd;
 
-		// get the CUarray from the external memory
-		CUresult Result = FCUDAModule::CUDA().cuExternalMemoryGetMappedMipmappedArray(&MappedMipArray, MappedExternalMemory, &MipmapDesc);
-		if (Result != CUDA_SUCCESS)
 		{
-			UE_LOG(LogPixelStreaming, Error, TEXT("Failed to bind mipmappedArray error: %d"), Result);
+			// Generate VkMemoryGetFdInfoKHR
+			VkMemoryGetFdInfoKHR MemoryGetFdInfoKHR = {};
+			MemoryGetFdInfoKHR.sType = VK_STRUCTURE_TYPE_MEMORY_GET_FD_INFO_KHR;
+			MemoryGetFdInfoKHR.pNext = NULL;
+			MemoryGetFdInfoKHR.memory = TextureAllocationInfo.Handle;
+			MemoryGetFdInfoKHR.handleType = VK_EXTERNAL_MEMORY_HANDLE_TYPE_OPAQUE_FD_BIT_KHR;
+
+			PFN_vkGetMemoryFdKHR FPGetMemoryFdKHR = (PFN_vkGetMemoryFdKHR)VulkanRHI->RHIGetVkDeviceProcAddr("vkGetMemoryFdKHR");
+			VERIFYVULKANRESULT_EXTERNAL(FPGetMemoryFdKHR(Device, &MemoryGetFdInfoKHR, &Fd));
 		}
-	}
 
-	// get the CUarray from the external memory
-	CUresult MipMapArrGetLevelErr = FCUDAModule::CUDA().cuMipmappedArrayGetLevel(&MappedArray, MappedMipArray, 0);
-	if (MipMapArrGetLevelErr != CUDA_SUCCESS)
-	{
-		UE_LOG(LogPixelStreaming, Error, TEXT("Failed to bind to mip 0."));
-	}
-
-	FCUDAModule::CUDA().cuCtxPopCurrent(NULL);
-
-	InputFrame->SetTexture(MappedArray, AVEncoder::FVideoEncoderInputFrame::EUnderlyingRHI::Vulkan, Handle, [MappedArray, MappedMipArray, MappedExternalMemory](CUarray NativeTexture) {
-		// free the cuda types
 		FCUDAModule::CUDA().cuCtxPushCurrent(FModuleManager::GetModuleChecked<FCUDAModule>("CUDA").GetCudaContext());
 
-		if (MappedArray)
+		CUexternalMemory MappedExternalMemory = nullptr;
+
 		{
-			CUresult Result = FCUDAModule::CUDA().cuArrayDestroy(MappedArray);
+			// generate a cudaExternalMemoryHandleDesc
+			CUDA_EXTERNAL_MEMORY_HANDLE_DESC CudaExtMemHandleDesc = {};
+			CudaExtMemHandleDesc.type = CU_EXTERNAL_MEMORY_HANDLE_TYPE_OPAQUE_FD;
+			CudaExtMemHandleDesc.handle.fd = Fd;
+			CudaExtMemHandleDesc.size = TextureAllocationInfo.Offset + TextureAllocationInfo.Size;
+
+			// import external memory
+			CUresult Result = FCUDAModule::CUDA().cuImportExternalMemory(&MappedExternalMemory, &CudaExtMemHandleDesc);
 			if (Result != CUDA_SUCCESS)
 			{
-				UE_LOG(LogPixelStreaming, Error, TEXT("Failed to destroy MappedArray: %d"), Result);
+				UE_LOG(LogPixelStreaming, Error, TEXT("Failed to import external memory from vulkan error: %d"), Result);
 			}
 		}
 
-		if (MappedMipArray)
+#endif
+
+		CUmipmappedArray MappedMipArray = nullptr;
+		CUarray MappedArray = nullptr;
+
 		{
-			CUresult Result = FCUDAModule::CUDA().cuMipmappedArrayDestroy(MappedMipArray);
+			CUDA_EXTERNAL_MEMORY_MIPMAPPED_ARRAY_DESC MipmapDesc = {};
+			MipmapDesc.numLevels = 1;
+			MipmapDesc.offset = TextureAllocationInfo.Offset;
+			MipmapDesc.arrayDesc.Width = Texture->GetDesc().Extent.X;
+			MipmapDesc.arrayDesc.Height = Texture->GetDesc().Extent.Y;
+			MipmapDesc.arrayDesc.Depth = 0;
+			MipmapDesc.arrayDesc.NumChannels = 4;
+			MipmapDesc.arrayDesc.Format = CU_AD_FORMAT_UNSIGNED_INT8;
+			MipmapDesc.arrayDesc.Flags = CUDA_ARRAY3D_SURFACE_LDST | CUDA_ARRAY3D_COLOR_ATTACHMENT;
+
+			// get the CUarray from the external memory
+			CUresult Result = FCUDAModule::CUDA().cuExternalMemoryGetMappedMipmappedArray(&MappedMipArray, MappedExternalMemory, &MipmapDesc);
 			if (Result != CUDA_SUCCESS)
 			{
-				UE_LOG(LogPixelStreaming, Error, TEXT("Failed to destroy MappedMipArray: %d"), Result);
+				UE_LOG(LogPixelStreaming, Error, TEXT("Failed to bind mipmappedArray error: %d"), Result);
 			}
 		}
 
-		if (MappedExternalMemory)
+		// get the CUarray from the external memory
+		CUresult MipMapArrGetLevelErr = FCUDAModule::CUDA().cuMipmappedArrayGetLevel(&MappedArray, MappedMipArray, 0);
+		if (MipMapArrGetLevelErr != CUDA_SUCCESS)
 		{
-			CUresult Result = FCUDAModule::CUDA().cuDestroyExternalMemory(MappedExternalMemory);
-			if (Result != CUDA_SUCCESS)
-			{
-				UE_LOG(LogPixelStreaming, Error, TEXT("Failed to destroy MappedExternalMemoryArray: %d"), Result);
-			}
+			UE_LOG(LogPixelStreaming, Error, TEXT("Failed to bind to mip 0."));
 		}
 
 		FCUDAModule::CUDA().cuCtxPopCurrent(NULL);
-	});
-	}
-
-	#if PLATFORM_WINDOWS
-	void FEncoderFrameFactory::SetTextureCUDAD3D11(TSharedPtr<AVEncoder::FVideoEncoderInputFrame> InputFrame, const FTextureRHIRef& Texture)
-	{
-	FD3D11Texture* D3D11Texture = GetD3D11TextureFromRHITexture(Texture);
-	unsigned long long TextureMemorySize = D3D11Texture->GetMemorySize();
-
-	ID3D11Texture2D* D3D11NativeTexture = static_cast<ID3D11Texture2D*>(D3D11Texture->GetResource());
 
-	TRefCountPtr<IDXGIResource> DXGIResource;
-	HRESULT QueryResult = D3D11NativeTexture->QueryInterface(IID_PPV_ARGS(DXGIResource.GetInitReference()));
-	if (QueryResult != S_OK)
-	{
-		UE_LOG(LogPixelStreaming, Error, TEXT("Failed to get DX texture handle for importing memory to CUDA: %d"), QueryResult);
-	}
+		InputFrame->SetTexture(MappedArray, AVEncoder::FVideoEncoderInputFrame::EUnderlyingRHI::Vulkan, Handle, [MappedArray, MappedMipArray, MappedExternalMemory](CUarray NativeTexture) {
+			// free the cuda types
+			FCUDAModule::CUDA().cuCtxPushCurrent(FModuleManager::GetModuleChecked<FCUDAModule>("CUDA").GetCudaContext());
 
-	HANDLE D3D11TextureHandle;
-	DXGIResource->GetSharedHandle(&D3D11TextureHandle);
-	DXGIResource->Release();
+			if (MappedArray)
+			{
+				CUresult Result = FCUDAModule::CUDA().cuArrayDestroy(MappedArray);
+				if (Result != CUDA_SUCCESS)
+				{
+					UE_LOG(LogPixelStreaming, Error, TEXT("Failed to destroy MappedArray: %d"), Result);
+				}
+			}
 
-	FCUDAModule::CUDA().cuCtxPushCurrent(FModuleManager::GetModuleChecked<FCUDAModule>("CUDA").GetCudaContext());
+			if (MappedMipArray)
+			{
+				CUresult Result = FCUDAModule::CUDA().cuMipmappedArrayDestroy(MappedMipArray);
+				if (Result != CUDA_SUCCESS)
+				{
+					UE_LOG(LogPixelStreaming, Error, TEXT("Failed to destroy MappedMipArray: %d"), Result);
+				}
+			}
 
-	CUexternalMemory MappedExternalMemory = nullptr;
+			if (MappedExternalMemory)
+			{
+				CUresult Result = FCUDAModule::CUDA().cuDestroyExternalMemory(MappedExternalMemory);
+				if (Result != CUDA_SUCCESS)
+				{
+					UE_LOG(LogPixelStreaming, Error, TEXT("Failed to destroy MappedExternalMemoryArray: %d"), Result);
+				}
+			}
 
-	{
-		// generate a cudaExternalMemoryHandleDesc
-		CUDA_EXTERNAL_MEMORY_HANDLE_DESC CudaExtMemHandleDesc = {};
-		CudaExtMemHandleDesc.type = CU_EXTERNAL_MEMORY_HANDLE_TYPE_D3D11_RESOURCE_KMT;
-		CudaExtMemHandleDesc.handle.win32.name = NULL;
-		CudaExtMemHandleDesc.handle.win32.handle = D3D11TextureHandle;
-		CudaExtMemHandleDesc.size = TextureMemorySize;
-		// Necessary for committed resources (DX11 and committed DX12 resources)
-		CudaExtMemHandleDesc.flags |= CUDA_EXTERNAL_MEMORY_DEDICATED;
-
-		// import external memory
-		CUresult Result = FCUDAModule::CUDA().cuImportExternalMemory(&MappedExternalMemory, &CudaExtMemHandleDesc);
-		if (Result != CUDA_SUCCESS)
-		{
-			UE_LOG(LogPixelStreaming, Error, TEXT("Failed to import external memory from vulkan error: %d"), Result);
-		}
+			FCUDAModule::CUDA().cuCtxPopCurrent(NULL);
+		});
 	}
 
-	CUmipmappedArray MappedMipArray = nullptr;
-	CUarray MappedArray = nullptr;
-
+#if PLATFORM_WINDOWS
+	void FEncoderFrameFactory::SetTextureCUDAD3D11(TSharedPtr<AVEncoder::FVideoEncoderInputFrame> InputFrame, const FTextureRHIRef& Texture)
 	{
-		CUDA_EXTERNAL_MEMORY_MIPMAPPED_ARRAY_DESC MipmapDesc = {};
-		MipmapDesc.numLevels = 1;
-		MipmapDesc.offset = 0;
-		MipmapDesc.arrayDesc.Width = Texture->GetDesc().Extent.X;
-		MipmapDesc.arrayDesc.Height = Texture->GetDesc().Extent.Y;
-		MipmapDesc.arrayDesc.Depth = 1;
-		MipmapDesc.arrayDesc.NumChannels = 4;
-		MipmapDesc.arrayDesc.Format = CU_AD_FORMAT_UNSIGNED_INT8;
-		MipmapDesc.arrayDesc.Flags = CUDA_ARRAY3D_SURFACE_LDST | CUDA_ARRAY3D_COLOR_ATTACHMENT;
+		FD3D11Texture* D3D11Texture = GetD3D11TextureFromRHITexture(Texture);
+		unsigned long long TextureMemorySize = D3D11Texture->GetMemorySize();
 
-		// get the CUarray from the external memory
-		CUresult Result = FCUDAModule::CUDA().cuExternalMemoryGetMappedMipmappedArray(&MappedMipArray, MappedExternalMemory, &MipmapDesc);
-		if (Result != CUDA_SUCCESS)
+		ID3D11Texture2D* D3D11NativeTexture = static_cast<ID3D11Texture2D*>(D3D11Texture->GetResource());
+
+		TRefCountPtr<IDXGIResource> DXGIResource;
+		HRESULT QueryResult = D3D11NativeTexture->QueryInterface(IID_PPV_ARGS(DXGIResource.GetInitReference()));
+		if (QueryResult != S_OK)
 		{
-			UE_LOG(LogPixelStreaming, Error, TEXT("Failed to bind mipmappedArray error: %d"), Result);
+			UE_LOG(LogPixelStreaming, Error, TEXT("Failed to get DX texture handle for importing memory to CUDA: %d"), QueryResult);
 		}
-	}
 
-	// get the CUarray from the external memory
-	CUresult MipMapArrGetLevelErr = FCUDAModule::CUDA().cuMipmappedArrayGetLevel(&MappedArray, MappedMipArray, 0);
-	if (MipMapArrGetLevelErr != CUDA_SUCCESS)
-	{
-		UE_LOG(LogPixelStreaming, Error, TEXT("Failed to bind to mip 0."));
-	}
-
-	FCUDAModule::CUDA().cuCtxPopCurrent(NULL);
+		HANDLE D3D11TextureHandle;
+		DXGIResource->GetSharedHandle(&D3D11TextureHandle);
+		DXGIResource->Release();
 
-	InputFrame->SetTexture(MappedArray, AVEncoder::FVideoEncoderInputFrame::EUnderlyingRHI::D3D11, nullptr, [MappedArray, MappedMipArray, MappedExternalMemory](CUarray NativeTexture) {
-		// free the cuda types
 		FCUDAModule::CUDA().cuCtxPushCurrent(FModuleManager::GetModuleChecked<FCUDAModule>("CUDA").GetCudaContext());
 
-		if (MappedArray)
+		CUexternalMemory MappedExternalMemory = nullptr;
+
 		{
-			CUresult Result = FCUDAModule::CUDA().cuArrayDestroy(MappedArray);
+			// generate a cudaExternalMemoryHandleDesc
+			CUDA_EXTERNAL_MEMORY_HANDLE_DESC CudaExtMemHandleDesc = {};
+			CudaExtMemHandleDesc.type = CU_EXTERNAL_MEMORY_HANDLE_TYPE_D3D11_RESOURCE_KMT;
+			CudaExtMemHandleDesc.handle.win32.name = NULL;
+			CudaExtMemHandleDesc.handle.win32.handle = D3D11TextureHandle;
+			CudaExtMemHandleDesc.size = TextureMemorySize;
+			// Necessary for committed resources (DX11 and committed DX12 resources)
+			CudaExtMemHandleDesc.flags |= CUDA_EXTERNAL_MEMORY_DEDICATED;
+
+			// import external memory
+			CUresult Result = FCUDAModule::CUDA().cuImportExternalMemory(&MappedExternalMemory, &CudaExtMemHandleDesc);
 			if (Result != CUDA_SUCCESS)
 			{
-				UE_LOG(LogPixelStreaming, Error, TEXT("Failed to destroy MappedArray: %d"), Result);
+				UE_LOG(LogPixelStreaming, Error, TEXT("Failed to import external memory from vulkan error: %d"), Result);
 			}
 		}
 
-		if (MappedMipArray)
+		CUmipmappedArray MappedMipArray = nullptr;
+		CUarray MappedArray = nullptr;
+
 		{
-			CUresult Result = FCUDAModule::CUDA().cuMipmappedArrayDestroy(MappedMipArray);
+			CUDA_EXTERNAL_MEMORY_MIPMAPPED_ARRAY_DESC MipmapDesc = {};
+			MipmapDesc.numLevels = 1;
+			MipmapDesc.offset = 0;
+			MipmapDesc.arrayDesc.Width = Texture->GetDesc().Extent.X;
+			MipmapDesc.arrayDesc.Height = Texture->GetDesc().Extent.Y;
+			MipmapDesc.arrayDesc.Depth = 1;
+			MipmapDesc.arrayDesc.NumChannels = 4;
+			MipmapDesc.arrayDesc.Format = CU_AD_FORMAT_UNSIGNED_INT8;
+			MipmapDesc.arrayDesc.Flags = CUDA_ARRAY3D_SURFACE_LDST | CUDA_ARRAY3D_COLOR_ATTACHMENT;
+
+			// get the CUarray from the external memory
+			CUresult Result = FCUDAModule::CUDA().cuExternalMemoryGetMappedMipmappedArray(&MappedMipArray, MappedExternalMemory, &MipmapDesc);
 			if (Result != CUDA_SUCCESS)
 			{
-				UE_LOG(LogPixelStreaming, Error, TEXT("Failed to destroy MappedMipArray: %d"), Result);
+				UE_LOG(LogPixelStreaming, Error, TEXT("Failed to bind mipmappedArray error: %d"), Result);
 			}
 		}
 
-		if (MappedExternalMemory)
+		// get the CUarray from the external memory
+		CUresult MipMapArrGetLevelErr = FCUDAModule::CUDA().cuMipmappedArrayGetLevel(&MappedArray, MappedMipArray, 0);
+		if (MipMapArrGetLevelErr != CUDA_SUCCESS)
 		{
-			CUresult Result = FCUDAModule::CUDA().cuDestroyExternalMemory(MappedExternalMemory);
-			if (Result != CUDA_SUCCESS)
-			{
-				UE_LOG(LogPixelStreaming, Error, TEXT("Failed to destroy MappedExternalMemoryArray: %d"), Result);
-			}
+			UE_LOG(LogPixelStreaming, Error, TEXT("Failed to bind to mip 0."));
 		}
 
 		FCUDAModule::CUDA().cuCtxPopCurrent(NULL);
-	});
-	}
 
-	void FEncoderFrameFactory::SetTextureCUDAD3D12(TSharedPtr<AVEncoder::FVideoEncoderInputFrame> InputFrame, const FTextureRHIRef& Texture)
-	{
-	ID3D12Resource* NativeD3D12Resource = GetID3D12DynamicRHI()->RHIGetResource(Texture);
-	const int64 TextureMemorySize = GetID3D12DynamicRHI()->RHIGetResourceMemorySize(Texture);
+		InputFrame->SetTexture(MappedArray, AVEncoder::FVideoEncoderInputFrame::EUnderlyingRHI::D3D11, nullptr, [MappedArray, MappedMipArray, MappedExternalMemory](CUarray NativeTexture) {
+			// free the cuda types
+			FCUDAModule::CUDA().cuCtxPushCurrent(FModuleManager::GetModuleChecked<FCUDAModule>("CUDA").GetCudaContext());
 
-	// Because we create our texture as RenderTargetable, it is created as a committed resource, which is what our current implementation here supports.
-	// To prevent a mystery crash in future, check that our resource is a committed resource
-	check(!GetID3D12DynamicRHI()->RHIIsResourcePlaced(Texture));
+			if (MappedArray)
+			{
+				CUresult Result = FCUDAModule::CUDA().cuArrayDestroy(MappedArray);
+				if (Result != CUDA_SUCCESS)
+				{
+					UE_LOG(LogPixelStreaming, Error, TEXT("Failed to destroy MappedArray: %d"), Result);
+				}
+			}
 
-	TRefCountPtr<ID3D12Device> OwnerDevice;
-	HRESULT QueryResult;
-	if ((QueryResult = NativeD3D12Resource->GetDevice(IID_PPV_ARGS(OwnerDevice.GetInitReference()))) != S_OK)
-	{
-		UE_LOG(LogPixelStreaming, Error, TEXT("Failed to get DX texture handle for importing memory to CUDA: %d (Get Device)"), QueryResult);
-	}
+			if (MappedMipArray)
+			{
+				CUresult Result = FCUDAModule::CUDA().cuMipmappedArrayDestroy(MappedMipArray);
+				if (Result != CUDA_SUCCESS)
+				{
+					UE_LOG(LogPixelStreaming, Error, TEXT("Failed to destroy MappedMipArray: %d"), Result);
+				}
+			}
 
-	//
-	// ID3D12Device::CreateSharedHandle gives as an NT Handle, and so we need to call CloseHandle on it
-	//
-	HANDLE D3D12TextureHandle;
-	if ((QueryResult = OwnerDevice->CreateSharedHandle(NativeD3D12Resource, NULL, GENERIC_ALL, NULL, &D3D12TextureHandle)) != S_OK)
-	{
-		UE_LOG(LogPixelStreaming, Error, TEXT("Failed to get DX texture handle for importing memory to CUDA: %d"), QueryResult);
+			if (MappedExternalMemory)
+			{
+				CUresult Result = FCUDAModule::CUDA().cuDestroyExternalMemory(MappedExternalMemory);
+				if (Result != CUDA_SUCCESS)
+				{
+					UE_LOG(LogPixelStreaming, Error, TEXT("Failed to destroy MappedExternalMemoryArray: %d"), Result);
+				}
+			}
+
+			FCUDAModule::CUDA().cuCtxPopCurrent(NULL);
+		});
 	}
 
-	FCUDAModule::CUDA().cuCtxPushCurrent(FModuleManager::GetModuleChecked<FCUDAModule>("CUDA").GetCudaContext());
+	void FEncoderFrameFactory::SetTextureCUDAD3D12(TSharedPtr<AVEncoder::FVideoEncoderInputFrame> InputFrame, const FTextureRHIRef& Texture)
+	{
+		ID3D12Resource* NativeD3D12Resource = GetID3D12DynamicRHI()->RHIGetResource(Texture);
+		const int64 TextureMemorySize = GetID3D12DynamicRHI()->RHIGetResourceMemorySize(Texture);
 
-	CUexternalMemory MappedExternalMemory = nullptr;
+		// Because we create our texture as RenderTargetable, it is created as a committed resource, which is what our current implementation here supports.
+		// To prevent a mystery crash in future, check that our resource is a committed resource
+		check(!GetID3D12DynamicRHI()->RHIIsResourcePlaced(Texture));
 
-	{
-		// generate a cudaExternalMemoryHandleDesc
-		CUDA_EXTERNAL_MEMORY_HANDLE_DESC CudaExtMemHandleDesc = {};
-		CudaExtMemHandleDesc.type = CU_EXTERNAL_MEMORY_HANDLE_TYPE_D3D12_RESOURCE;
-		CudaExtMemHandleDesc.handle.win32.name = NULL;
-		CudaExtMemHandleDesc.handle.win32.handle = D3D12TextureHandle;
-		CudaExtMemHandleDesc.size = TextureMemorySize;
-		// Necessary for committed resources (DX11 and committed DX12 resources)
-		CudaExtMemHandleDesc.flags |= CUDA_EXTERNAL_MEMORY_DEDICATED;
-
-		// import external memory
-		CUresult Result = FCUDAModule::CUDA().cuImportExternalMemory(&MappedExternalMemory, &CudaExtMemHandleDesc);
-		if (Result != CUDA_SUCCESS)
+		TRefCountPtr<ID3D12Device> OwnerDevice;
+		HRESULT QueryResult;
+		if ((QueryResult = NativeD3D12Resource->GetDevice(IID_PPV_ARGS(OwnerDevice.GetInitReference()))) != S_OK)
 		{
-			UE_LOG(LogPixelStreaming, Error, TEXT("Failed to import external memory from vulkan error: %d"), Result);
+			UE_LOG(LogPixelStreaming, Error, TEXT("Failed to get DX texture handle for importing memory to CUDA: %d (Get Device)"), QueryResult);
 		}
-	}
-
-	CUmipmappedArray MappedMipArray = nullptr;
-	CUarray MappedArray = nullptr;
-
-	{
-		CUDA_EXTERNAL_MEMORY_MIPMAPPED_ARRAY_DESC MipmapDesc = {};
-		MipmapDesc.numLevels = 1;
-		MipmapDesc.offset = 0;
-		MipmapDesc.arrayDesc.Width = Texture->GetDesc().Extent.X;
-		MipmapDesc.arrayDesc.Height = Texture->GetDesc().Extent.Y;
-		MipmapDesc.arrayDesc.Depth = 1;
-		MipmapDesc.arrayDesc.NumChannels = 4;
-		MipmapDesc.arrayDesc.Format = CU_AD_FORMAT_UNSIGNED_INT8;
-		MipmapDesc.arrayDesc.Flags = CUDA_ARRAY3D_SURFACE_LDST | CUDA_ARRAY3D_COLOR_ATTACHMENT;
 
-		// get the CUarray from the external memory
-		CUresult Result = FCUDAModule::CUDA().cuExternalMemoryGetMappedMipmappedArray(&MappedMipArray, MappedExternalMemory, &MipmapDesc);
-		if (Result != CUDA_SUCCESS)
+		//
+		// ID3D12Device::CreateSharedHandle gives as an NT Handle, and so we need to call CloseHandle on it
+		//
+		HANDLE D3D12TextureHandle;
+		if ((QueryResult = OwnerDevice->CreateSharedHandle(NativeD3D12Resource, NULL, GENERIC_ALL, NULL, &D3D12TextureHandle)) != S_OK)
 		{
-			UE_LOG(LogPixelStreaming, Error, TEXT("Failed to bind mipmappedArray error: %d"), Result);
+			UE_LOG(LogPixelStreaming, Error, TEXT("Failed to get DX texture handle for importing memory to CUDA: %d"), QueryResult);
 		}
-	}
-
-	// get the CUarray from the external memory
-	CUresult MipMapArrGetLevelErr = FCUDAModule::CUDA().cuMipmappedArrayGetLevel(&MappedArray, MappedMipArray, 0);
-	if (MipMapArrGetLevelErr != CUDA_SUCCESS)
-	{
-		UE_LOG(LogPixelStreaming, Error, TEXT("Failed to bind to mip 0."));
-	}
-
-	FCUDAModule::CUDA().cuCtxPopCurrent(NULL);
 
-	InputFrame->SetTexture(MappedArray, AVEncoder::FVideoEncoderInputFrame::EUnderlyingRHI::D3D12, D3D12TextureHandle, [MappedArray, MappedMipArray, MappedExternalMemory](CUarray NativeTexture) {
-		// free the cuda types
 		FCUDAModule::CUDA().cuCtxPushCurrent(FModuleManager::GetModuleChecked<FCUDAModule>("CUDA").GetCudaContext());
 
-		if (MappedArray)
+		CUexternalMemory MappedExternalMemory = nullptr;
+
 		{
-			CUresult Result = FCUDAModule::CUDA().cuArrayDestroy(MappedArray);
+			// generate a cudaExternalMemoryHandleDesc
+			CUDA_EXTERNAL_MEMORY_HANDLE_DESC CudaExtMemHandleDesc = {};
+			CudaExtMemHandleDesc.type = CU_EXTERNAL_MEMORY_HANDLE_TYPE_D3D12_RESOURCE;
+			CudaExtMemHandleDesc.handle.win32.name = NULL;
+			CudaExtMemHandleDesc.handle.win32.handle = D3D12TextureHandle;
+			CudaExtMemHandleDesc.size = TextureMemorySize;
+			// Necessary for committed resources (DX11 and committed DX12 resources)
+			CudaExtMemHandleDesc.flags |= CUDA_EXTERNAL_MEMORY_DEDICATED;
+
+			// import external memory
+			CUresult Result = FCUDAModule::CUDA().cuImportExternalMemory(&MappedExternalMemory, &CudaExtMemHandleDesc);
 			if (Result != CUDA_SUCCESS)
 			{
-				UE_LOG(LogPixelStreaming, Error, TEXT("Failed to destroy MappedArray: %d"), Result);
+				UE_LOG(LogPixelStreaming, Error, TEXT("Failed to import external memory from vulkan error: %d"), Result);
 			}
 		}
 
-		if (MappedMipArray)
+		CUmipmappedArray MappedMipArray = nullptr;
+		CUarray MappedArray = nullptr;
+
 		{
-			CUresult Result = FCUDAModule::CUDA().cuMipmappedArrayDestroy(MappedMipArray);
+			CUDA_EXTERNAL_MEMORY_MIPMAPPED_ARRAY_DESC MipmapDesc = {};
+			MipmapDesc.numLevels = 1;
+			MipmapDesc.offset = 0;
+			MipmapDesc.arrayDesc.Width = Texture->GetDesc().Extent.X;
+			MipmapDesc.arrayDesc.Height = Texture->GetDesc().Extent.Y;
+			MipmapDesc.arrayDesc.Depth = 1;
+			MipmapDesc.arrayDesc.NumChannels = 4;
+			MipmapDesc.arrayDesc.Format = CU_AD_FORMAT_UNSIGNED_INT8;
+			MipmapDesc.arrayDesc.Flags = CUDA_ARRAY3D_SURFACE_LDST | CUDA_ARRAY3D_COLOR_ATTACHMENT;
+
+			// get the CUarray from the external memory
+			CUresult Result = FCUDAModule::CUDA().cuExternalMemoryGetMappedMipmappedArray(&MappedMipArray, MappedExternalMemory, &MipmapDesc);
 			if (Result != CUDA_SUCCESS)
 			{
-				UE_LOG(LogPixelStreaming, Error, TEXT("Failed to destroy MappedMipArray: %d"), Result);
+				UE_LOG(LogPixelStreaming, Error, TEXT("Failed to bind mipmappedArray error: %d"), Result);
 			}
 		}
 
-		if (MappedExternalMemory)
+		// get the CUarray from the external memory
+		CUresult MipMapArrGetLevelErr = FCUDAModule::CUDA().cuMipmappedArrayGetLevel(&MappedArray, MappedMipArray, 0);
+		if (MipMapArrGetLevelErr != CUDA_SUCCESS)
 		{
-			CUresult Result = FCUDAModule::CUDA().cuDestroyExternalMemory(MappedExternalMemory);
-			if (Result != CUDA_SUCCESS)
-			{
-				UE_LOG(LogPixelStreaming, Error, TEXT("Failed to destroy MappedExternalMemoryArray: %d"), Result);
-			}
+			UE_LOG(LogPixelStreaming, Error, TEXT("Failed to bind to mip 0."));
 		}
 
 		FCUDAModule::CUDA().cuCtxPopCurrent(NULL);
-	});
+
+		InputFrame->SetTexture(MappedArray, AVEncoder::FVideoEncoderInputFrame::EUnderlyingRHI::D3D12, D3D12TextureHandle, [MappedArray, MappedMipArray, MappedExternalMemory](CUarray NativeTexture) {
+			// free the cuda types
+			FCUDAModule::CUDA().cuCtxPushCurrent(FModuleManager::GetModuleChecked<FCUDAModule>("CUDA").GetCudaContext());
+
+			if (MappedArray)
+			{
+				CUresult Result = FCUDAModule::CUDA().cuArrayDestroy(MappedArray);
+				if (Result != CUDA_SUCCESS)
+				{
+					UE_LOG(LogPixelStreaming, Error, TEXT("Failed to destroy MappedArray: %d"), Result);
+				}
+			}
+
+			if (MappedMipArray)
+			{
+				CUresult Result = FCUDAModule::CUDA().cuMipmappedArrayDestroy(MappedMipArray);
+				if (Result != CUDA_SUCCESS)
+				{
+					UE_LOG(LogPixelStreaming, Error, TEXT("Failed to destroy MappedMipArray: %d"), Result);
+				}
+			}
+
+			if (MappedExternalMemory)
+			{
+				CUresult Result = FCUDAModule::CUDA().cuDestroyExternalMemory(MappedExternalMemory);
+				if (Result != CUDA_SUCCESS)
+				{
+					UE_LOG(LogPixelStreaming, Error, TEXT("Failed to destroy MappedExternalMemoryArray: %d"), Result);
+				}
+			}
+
+			FCUDAModule::CUDA().cuCtxPopCurrent(NULL);
+		});
 	}
-	#endif //PLATFORM_WINDOWS
+#endif // PLATFORM_WINDOWS
 } // namespace UE::PixelStreaming
\ No newline at end of file
diff --git a/Engine/Plugins/Media/PixelStreaming/Source/PixelStreaming/Private/FrameAdapterH264.cpp b/Engine/Plugins/Media/PixelStreaming/Source/PixelStreaming/Private/FrameAdapterH264.cpp
new file mode 100644
--- /dev/null
+++ b/Engine/Plugins/Media/PixelStreaming/Source/PixelStreaming/Private/FrameAdapterH264.cpp
@@ -0,0 +1,22 @@
+// Copyright Epic Games, Inc. All Rights Reserved.
+
+#include "FrameAdapterH264.h"
+#include "FrameAdapterProcessH264.h"
+
+namespace UE::PixelStreaming
+{
+	FFrameAdapterH264::FFrameAdapterH264(TSharedPtr<FPixelStreamingVideoInput> VideoInput)
+		: FPixelStreamingFrameAdapter(VideoInput)
+	{
+	}
+
+	FFrameAdapterH264::FFrameAdapterH264(TSharedPtr<FPixelStreamingVideoInput> VideoInput, TArray<float> LayerScales)
+		: FPixelStreamingFrameAdapter(VideoInput, LayerScales)
+	{
+	}
+
+	TSharedPtr<FPixelStreamingFrameAdapterProcess> FFrameAdapterH264::CreateAdaptProcess(float Scale)
+	{
+		return MakeShared<FFrameAdapterProcessH264>(Scale);
+	}
+} // namespace UE::PixelStreaming
diff --git a/Engine/Plugins/Media/PixelStreaming/Source/PixelStreaming/Private/FrameAdapterH264.h b/Engine/Plugins/Media/PixelStreaming/Source/PixelStreaming/Private/FrameAdapterH264.h
new file mode 100644
--- /dev/null
+++ b/Engine/Plugins/Media/PixelStreaming/Source/PixelStreaming/Private/FrameAdapterH264.h
@@ -0,0 +1,40 @@
+// Copyright Epic Games, Inc. All Rights Reserved.
+
+#pragma once
+
+#include "PixelStreamingFrameAdapter.h"
+#include "PixelStreamingFrameMetadata.h"
+
+namespace UE::PixelStreaming
+{
+	class FAdaptedVideoFrameLayerH264 : public IPixelStreamingAdaptedVideoFrameLayer
+	{
+	public:
+		FAdaptedVideoFrameLayerH264(FTexture2DRHIRef InFrameTexture)
+			: FrameTexture(InFrameTexture)
+		{
+		}
+		virtual ~FAdaptedVideoFrameLayerH264() = default;
+
+		virtual int32 GetWidth() const override { return FrameTexture->GetDesc().Extent.X; }
+		virtual int32 GetHeight() const override { return FrameTexture->GetDesc().Extent.Y; }
+
+		FTexture2DRHIRef GetFrameTexture() const { return FrameTexture; }
+
+		FPixelStreamingFrameMetadata Metadata;
+		
+	private:
+		FTexture2DRHIRef FrameTexture;
+	};
+	
+	class FFrameAdapterH264 : public FPixelStreamingFrameAdapter
+	{
+	public:
+		FFrameAdapterH264(TSharedPtr<FPixelStreamingVideoInput> VideoInput);
+		FFrameAdapterH264(TSharedPtr<FPixelStreamingVideoInput> VideoInput, TArray<float> LayerScales);
+		virtual ~FFrameAdapterH264() = default;
+
+	private:
+		virtual TSharedPtr<FPixelStreamingFrameAdapterProcess> CreateAdaptProcess(float Scale) override;
+	};
+} // namespace UE::PixelStreaming
diff --git a/Engine/Plugins/Media/PixelStreaming/Source/PixelStreaming/Private/FrameAdapterI420.cpp b/Engine/Plugins/Media/PixelStreaming/Source/PixelStreaming/Private/FrameAdapterI420.cpp
new file mode 100644
--- /dev/null
+++ b/Engine/Plugins/Media/PixelStreaming/Source/PixelStreaming/Private/FrameAdapterI420.cpp
@@ -0,0 +1,48 @@
+// Copyright Epic Games, Inc. All Rights Reserved.
+
+#include "FrameAdapterI420.h"
+#include "FrameAdapterProcessI420CPU.h"
+#include "FrameAdapterProcessI420Compute.h"
+
+namespace UE::PixelStreaming
+{
+	FFrameAdapterI420::FFrameAdapterI420(TSharedPtr<FPixelStreamingVideoInput> VideoInput)
+		: FPixelStreamingFrameAdapter(VideoInput)
+	{
+	}
+
+	FFrameAdapterI420::FFrameAdapterI420(TSharedPtr<FPixelStreamingVideoInput> VideoInput, TArray<float> LayerScales)
+		: FPixelStreamingFrameAdapter(VideoInput, LayerScales)
+	{
+	}
+
+	FFrameAdapterI420CPU::FFrameAdapterI420CPU(TSharedPtr<FPixelStreamingVideoInput> VideoInput)
+		: FFrameAdapterI420(VideoInput)
+	{
+	}
+
+	FFrameAdapterI420CPU::FFrameAdapterI420CPU(TSharedPtr<FPixelStreamingVideoInput> VideoInput, TArray<float> LayerScales)
+		: FFrameAdapterI420(VideoInput, LayerScales)
+	{
+	}
+
+	TSharedPtr<FPixelStreamingFrameAdapterProcess> FFrameAdapterI420CPU::CreateAdaptProcess(float Scale)
+	{
+		return MakeShared<FFrameAdapterProcessI420CPU>(Scale);
+	}
+
+	FFrameAdapterI420Compute::FFrameAdapterI420Compute(TSharedPtr<FPixelStreamingVideoInput> VideoInput)
+		: FFrameAdapterI420(VideoInput)
+	{
+	}
+
+	FFrameAdapterI420Compute::FFrameAdapterI420Compute(TSharedPtr<FPixelStreamingVideoInput> VideoInput, TArray<float> LayerScales)
+		: FFrameAdapterI420(VideoInput, LayerScales)
+	{
+	}
+
+	TSharedPtr<FPixelStreamingFrameAdapterProcess> FFrameAdapterI420Compute::CreateAdaptProcess(float Scale)
+	{
+		return MakeShared<FFrameAdapterProcessI420Compute>(Scale);
+	}
+} // namespace UE::PixelStreaming
diff --git a/Engine/Plugins/Media/PixelStreaming/Source/PixelStreaming/Private/FrameAdapterI420.h b/Engine/Plugins/Media/PixelStreaming/Source/PixelStreaming/Private/FrameAdapterI420.h
new file mode 100644
--- /dev/null
+++ b/Engine/Plugins/Media/PixelStreaming/Source/PixelStreaming/Private/FrameAdapterI420.h
@@ -0,0 +1,60 @@
+// Copyright Epic Games, Inc. All Rights Reserved.
+
+#pragma once
+
+#include "WebRTCIncludes.h"
+#include "PixelStreamingFrameAdapter.h"
+#include "PixelStreamingFrameMetadata.h"
+
+namespace UE::PixelStreaming
+{
+	class FAdaptedVideoFrameLayerI420 : public IPixelStreamingAdaptedVideoFrameLayer
+	{
+	public:
+		FAdaptedVideoFrameLayerI420(rtc::scoped_refptr<webrtc::I420Buffer> InI420Buffer)
+			: I420Buffer(InI420Buffer)
+		{
+		}
+		virtual ~FAdaptedVideoFrameLayerI420() = default;
+
+		virtual int32 GetWidth() const override { return I420Buffer->width(); }
+		virtual int32 GetHeight() const override { return I420Buffer->height(); }
+
+		rtc::scoped_refptr<webrtc::I420Buffer> GetI420Buffer() const { return I420Buffer; }
+
+		FPixelStreamingFrameMetadata Metadata;
+		
+	private:
+		rtc::scoped_refptr<webrtc::I420Buffer> I420Buffer;
+	};
+	
+	class FFrameAdapterI420 : public FPixelStreamingFrameAdapter
+	{
+	public:
+		FFrameAdapterI420(TSharedPtr<FPixelStreamingVideoInput> VideoInput);
+		FFrameAdapterI420(TSharedPtr<FPixelStreamingVideoInput> VideoInput, TArray<float> LayerScales);
+		virtual ~FFrameAdapterI420() = default;
+	};
+
+	class FFrameAdapterI420CPU : public FFrameAdapterI420
+	{
+	public:
+		FFrameAdapterI420CPU(TSharedPtr<FPixelStreamingVideoInput> VideoInput);
+		FFrameAdapterI420CPU(TSharedPtr<FPixelStreamingVideoInput> VideoInput, TArray<float> LayerScales);
+		virtual ~FFrameAdapterI420CPU() = default;
+
+	private:
+		virtual TSharedPtr<FPixelStreamingFrameAdapterProcess> CreateAdaptProcess(float Scale) override;
+	};
+
+	class FFrameAdapterI420Compute : public FFrameAdapterI420
+	{
+	public:
+		FFrameAdapterI420Compute(TSharedPtr<FPixelStreamingVideoInput> VideoInput);
+		FFrameAdapterI420Compute(TSharedPtr<FPixelStreamingVideoInput> VideoInput, TArray<float> LayerScales);
+		virtual ~FFrameAdapterI420Compute() = default;
+
+	private:
+		virtual TSharedPtr<FPixelStreamingFrameAdapterProcess> CreateAdaptProcess(float Scale) override;
+	};
+} // namespace UE::PixelStreaming
diff --git a/Engine/Plugins/Media/PixelStreaming/Source/PixelStreaming/Private/FrameAdapterProcessH264.cpp b/Engine/Plugins/Media/PixelStreaming/Source/PixelStreaming/Private/FrameAdapterProcessH264.cpp
new file mode 100644
--- /dev/null
+++ b/Engine/Plugins/Media/PixelStreaming/Source/PixelStreaming/Private/FrameAdapterProcessH264.cpp
@@ -0,0 +1,84 @@
+// Copyright Epic Games, Inc. All Rights Reserved.
+
+#include "FrameAdapterProcessH264.h"
+#include "PixelStreamingSourceFrame.h"
+#include "FrameAdapterH264.h"
+#include "UtilsRender.h"
+#include "Stats.h"
+
+namespace UE::PixelStreaming
+{
+	FFrameAdapterProcessH264::FFrameAdapterProcessH264(float InScale)
+		: Scale(InScale)
+	{
+		Fence = GDynamicRHI->RHICreateGPUFence(TEXT("FFrameAdapterProcessH264 Fence"));
+	}
+
+	TSharedPtr<IPixelStreamingAdaptedVideoFrameLayer> FFrameAdapterProcessH264::CreateOutputBuffer(int32 SourceWidth, int32 SourceHeight)
+	{
+		const int32 Width = SourceWidth * Scale;
+		const int32 Height = SourceHeight * Scale;
+
+		FRHITextureCreateDesc TextureDesc =
+			FRHITextureCreateDesc::Create2D(TEXT("FFrameDataH264 Texture"), Width, Height, EPixelFormat::PF_B8G8R8A8)
+			.SetClearValue(FClearValueBinding::None)
+			.SetFlags(ETextureCreateFlags::RenderTargetable)
+			.SetInitialState(ERHIAccess::Present)
+			.DetermineInititialState();
+
+		if (RHIGetInterfaceType() == ERHIInterfaceType::Vulkan)
+		{
+			TextureDesc.AddFlags(ETextureCreateFlags::External);
+		}
+		else
+		{
+			TextureDesc.AddFlags(ETextureCreateFlags::Shared);
+		}
+
+		return MakeShared<FAdaptedVideoFrameLayerH264>(GDynamicRHI->RHICreateTexture(TextureDesc));
+	}
+
+	void FFrameAdapterProcessH264::BeginProcess(const FPixelStreamingSourceFrame& SourceFrame)
+	{
+		CurrentOuputBuffer = StaticCast<FAdaptedVideoFrameLayerH264*>(GetWriteBuffer().Get());
+		CurrentOuputBuffer->Metadata = SourceFrame.Metadata.Copy();
+		CurrentOuputBuffer->Metadata.ProcessName = "H264";
+		CurrentOuputBuffer->Metadata.AdaptCallTime = FPlatformTime::Cycles64();
+
+		FRHICommandListImmediate& RHICmdList = FRHICommandListExecutor::GetImmediateCommandList();
+
+		RHICmdList.EnqueueLambda([CurrentOuputBuffer = CurrentOuputBuffer](FRHICommandListImmediate&) { CurrentOuputBuffer->Metadata.AdaptProcessStartTime = FPlatformTime::Cycles64(); });
+
+		CopyTexture(RHICmdList, SourceFrame.FrameTexture, CurrentOuputBuffer->GetFrameTexture(), Fence);
+
+		// by adding this shared ref to the rhi lambda we can ensure that 'this' will not be destroyed
+		// until after the rhi thread is done with it, so all the commands will still have valid references.
+		TSharedRef<FFrameAdapterProcessH264> ThisRHIRef = StaticCastSharedRef<FFrameAdapterProcessH264>(AsShared());
+		RHICmdList.EnqueueLambda([ThisRHIRef](FRHICommandListImmediate&) { ThisRHIRef->CheckComplete(); });
+	}
+
+	void FFrameAdapterProcessH264::CheckComplete()
+	{
+		if (!Fence->Poll())
+		{
+			TSharedRef<FFrameAdapterProcessH264> ThisRHIRef = StaticCastSharedRef<FFrameAdapterProcessH264>(AsShared());
+			AsyncTask(ENamedThreads::AnyHiPriThreadHiPriTask, [ThisRHIRef]() {
+				ThisRHIRef->CheckComplete();
+			});
+		}
+		else
+		{
+			OnRHIStageComplete();
+		}
+	}
+
+	void FFrameAdapterProcessH264::OnRHIStageComplete()
+	{
+		checkf(Fence->Poll(), TEXT("Fence was not set. Backbuffer copy may not have completed."));
+		Fence->Clear();
+		CurrentOuputBuffer->Metadata.AdaptProcessFinalizeTime = FPlatformTime::Cycles64();
+		CurrentOuputBuffer->Metadata.AdaptProcessEndTime = FPlatformTime::Cycles64();
+		CurrentOuputBuffer = nullptr;
+		EndProcess();
+	}
+} // namespace UE::PixelStreaming
diff --git a/Engine/Plugins/Media/PixelStreaming/Source/PixelStreaming/Private/FrameAdapterProcessH264.h b/Engine/Plugins/Media/PixelStreaming/Source/PixelStreaming/Private/FrameAdapterProcessH264.h
new file mode 100644
--- /dev/null
+++ b/Engine/Plugins/Media/PixelStreaming/Source/PixelStreaming/Private/FrameAdapterProcessH264.h
@@ -0,0 +1,30 @@
+// Copyright Epic Games, Inc. All Rights Reserved.
+
+#pragma once
+
+#include "PixelStreamingFrameAdapterProcess.h"
+
+namespace UE::PixelStreaming
+{
+	class FAdaptedVideoFrameLayerH264;
+	
+	class FFrameAdapterProcessH264 : public FPixelStreamingFrameAdapterProcess
+	{
+	public:
+		FFrameAdapterProcessH264(float InScale);
+		virtual ~FFrameAdapterProcessH264() = default;
+
+	protected:
+		virtual TSharedPtr<IPixelStreamingAdaptedVideoFrameLayer> CreateOutputBuffer(int32 SourceWidth, int32 SourceHeight) override;
+		virtual void BeginProcess(const FPixelStreamingSourceFrame& SourceFrame) override;
+
+	private:
+		float Scale = 1.0f;
+
+		FGPUFenceRHIRef Fence;
+		FAdaptedVideoFrameLayerH264* CurrentOuputBuffer = nullptr;
+
+		void CheckComplete();
+		void OnRHIStageComplete();
+	};
+} // namespace UE::PixelStreaming
diff --git a/Engine/Plugins/Media/PixelStreaming/Source/PixelStreaming/Private/FrameAdapterProcessI420CPU.cpp b/Engine/Plugins/Media/PixelStreaming/Source/PixelStreaming/Private/FrameAdapterProcessI420CPU.cpp
new file mode 100644
--- /dev/null
+++ b/Engine/Plugins/Media/PixelStreaming/Source/PixelStreaming/Private/FrameAdapterProcessI420CPU.cpp
@@ -0,0 +1,131 @@
+// Copyright Epic Games, Inc. All Rights Reserved.
+
+#include "FrameAdapterProcessI420CPU.h"
+#include "FrameAdapterI420.h"
+#include "libyuv/convert.h"
+#include "UtilsRender.h"
+#include "PixelStreamingSourceFrame.h"
+
+namespace UE::PixelStreaming
+{
+	FFrameAdapterProcessI420CPU::FFrameAdapterProcessI420CPU(float InScale)
+		: Scale(InScale)
+	{
+	}
+
+	FFrameAdapterProcessI420CPU::~FFrameAdapterProcessI420CPU()
+	{
+		CleanUp();
+	}
+
+	void FFrameAdapterProcessI420CPU::Initialize(int32 SourceWidth, int32 SourceHeight)
+	{
+		const int32 Width = SourceWidth * Scale;
+		const int32 Height = SourceHeight * Scale;
+
+		FRHITextureCreateDesc TextureDesc =
+			FRHITextureCreateDesc::Create2D(TEXT("FFrameAdapterProcessI420CPU StagingTexture"), Width, Height, EPixelFormat::PF_B8G8R8A8)
+			.SetClearValue(FClearValueBinding::None)
+			.SetFlags(ETextureCreateFlags::RenderTargetable)
+			.SetInitialState(ERHIAccess::CopySrc)
+			.DetermineInititialState();
+
+		if (RHIGetInterfaceType() == ERHIInterfaceType::Vulkan)
+		{
+			TextureDesc.AddFlags(ETextureCreateFlags::External);
+		}
+		else
+		{
+			TextureDesc.AddFlags(ETextureCreateFlags::Shared);
+		}
+
+		StagingTexture = GDynamicRHI->RHICreateTexture(TextureDesc);
+
+		FRHITextureCreateDesc ReadbackDesc =
+			FRHITextureCreateDesc::Create2D(TEXT("FFrameAdapterProcessI420CPU ReadbackTexture"), Width, Height, EPixelFormat::PF_B8G8R8A8)
+			.SetClearValue(FClearValueBinding::None)
+			.SetFlags(ETextureCreateFlags::CPUReadback)
+			.SetInitialState(ERHIAccess::CPURead)
+			.DetermineInititialState();
+
+		ReadbackTexture = GDynamicRHI->RHICreateTexture(ReadbackDesc);
+
+		int32 BufferWidth = 0, BufferHeight = 0;
+		GDynamicRHI->RHIMapStagingSurface(ReadbackTexture, nullptr, ResultsBuffer, BufferWidth, BufferHeight);
+		MappedStride = BufferWidth;
+
+		FPixelStreamingFrameAdapterProcess::Initialize(SourceWidth, SourceHeight);
+	}
+
+	void FFrameAdapterProcessI420CPU::OnSourceResolutionChanged(int32 OldWidth, int32 OldHeight, int32 NewWidth, int32 NewHeight)
+	{
+		CleanUp();
+	}
+
+	TSharedPtr<IPixelStreamingAdaptedVideoFrameLayer> FFrameAdapterProcessI420CPU::CreateOutputBuffer(int32 SourceWidth, int32 SourceHeight)
+	{
+		const int32 Width = SourceWidth * Scale;
+		const int32 Height = SourceHeight * Scale;
+		return MakeShared<FAdaptedVideoFrameLayerI420>(webrtc::I420Buffer::Create(Width, Height));
+	}
+
+	void FFrameAdapterProcessI420CPU::BeginProcess(const FPixelStreamingSourceFrame& SourceFrame)
+	{
+		CurrentOuputBuffer = StaticCast<FAdaptedVideoFrameLayerI420*>(GetWriteBuffer().Get());
+		CurrentOuputBuffer->Metadata = SourceFrame.Metadata.Copy();
+		CurrentOuputBuffer->Metadata.ProcessName = "I420 CPU";
+		CurrentOuputBuffer->Metadata.AdaptCallTime = FPlatformTime::Cycles64();
+
+		FTexture2DRHIRef SourceTexture = SourceFrame.FrameTexture;
+
+		FRHICommandListImmediate& RHICmdList = FRHICommandListExecutor::GetImmediateCommandList();
+
+		RHICmdList.EnqueueLambda([CurrentOuputBuffer = CurrentOuputBuffer](FRHICommandListImmediate&) { CurrentOuputBuffer->Metadata.AdaptProcessStartTime = FPlatformTime::Cycles64(); });
+
+		RHICmdList.Transition(FRHITransitionInfo(SourceFrame.FrameTexture, ERHIAccess::Unknown, ERHIAccess::CopySrc));
+		RHICmdList.Transition(FRHITransitionInfo(StagingTexture, ERHIAccess::CopySrc, ERHIAccess::CopyDest));
+		CopyTexture(RHICmdList, SourceFrame.FrameTexture, StagingTexture, nullptr);
+
+		RHICmdList.Transition(FRHITransitionInfo(StagingTexture, ERHIAccess::CopyDest, ERHIAccess::CopySrc));
+		RHICmdList.Transition(FRHITransitionInfo(ReadbackTexture, ERHIAccess::CPURead, ERHIAccess::CopyDest));
+		RHICmdList.CopyTexture(StagingTexture, ReadbackTexture, {});
+
+		RHICmdList.Transition(FRHITransitionInfo(ReadbackTexture, ERHIAccess::CopyDest, ERHIAccess::CPURead));
+
+		// by adding this shared ref to the rhi lambda we can ensure that 'this' will not be destroyed
+		// until after the rhi thread is done with it, so all the commands will still have valid references.
+		TSharedRef<FFrameAdapterProcessI420CPU> ThisRHIRef = StaticCastSharedRef<FFrameAdapterProcessI420CPU>(AsShared());
+		RHICmdList.EnqueueLambda([ThisRHIRef](FRHICommandListImmediate&) {
+			ThisRHIRef->OnRHIStageComplete();
+		});
+	}
+
+	void FFrameAdapterProcessI420CPU::OnRHIStageComplete()
+	{
+		CurrentOuputBuffer->Metadata.AdaptProcessFinalizeTime = FPlatformTime::Cycles64();
+
+		rtc::scoped_refptr<webrtc::I420Buffer> I420Buffer = CurrentOuputBuffer->GetI420Buffer();
+		libyuv::ARGBToI420(
+			static_cast<uint8*>(ResultsBuffer),
+			MappedStride * 4,
+			I420Buffer->MutableDataY(),
+			I420Buffer->StrideY(),
+			I420Buffer->MutableDataU(),
+			I420Buffer->StrideU(),
+			I420Buffer->MutableDataV(),
+			I420Buffer->StrideV(),
+			I420Buffer->width(),
+			I420Buffer->height());
+
+		CurrentOuputBuffer->Metadata.AdaptProcessEndTime = FPlatformTime::Cycles64();
+		CurrentOuputBuffer = nullptr;
+
+		EndProcess();
+	}
+
+	void FFrameAdapterProcessI420CPU::CleanUp()
+	{
+		GDynamicRHI->RHIUnmapStagingSurface(ReadbackTexture);
+		ResultsBuffer = nullptr;
+	}
+} // namespace UE::PixelStreaming
diff --git a/Engine/Plugins/Media/PixelStreaming/Source/PixelStreaming/Private/FrameAdapterProcessI420CPU.h b/Engine/Plugins/Media/PixelStreaming/Source/PixelStreaming/Private/FrameAdapterProcessI420CPU.h
new file mode 100644
--- /dev/null
+++ b/Engine/Plugins/Media/PixelStreaming/Source/PixelStreaming/Private/FrameAdapterProcessI420CPU.h
@@ -0,0 +1,36 @@
+// Copyright Epic Games, Inc. All Rights Reserved.
+
+#pragma once
+
+#include "PixelStreamingFrameAdapterProcess.h"
+
+namespace UE::PixelStreaming
+{
+	class FAdaptedVideoFrameLayerI420;
+	
+	class FFrameAdapterProcessI420CPU : public FPixelStreamingFrameAdapterProcess
+	{
+	public:
+		FFrameAdapterProcessI420CPU(float InScale);
+		virtual ~FFrameAdapterProcessI420CPU();
+
+	protected:
+		virtual void Initialize(int32 SourceWidth, int32 SourceHeight) override;
+		virtual void OnSourceResolutionChanged(int32 OldWidth, int32 OldHeight, int32 NewWidth, int32 NewHeight) override;
+		virtual TSharedPtr<IPixelStreamingAdaptedVideoFrameLayer> CreateOutputBuffer(int32 SourceWidth, int32 SourceHeight) override;
+		virtual void BeginProcess(const FPixelStreamingSourceFrame& SourceFrame) override;
+
+	private:
+		float Scale = 1.0f;
+
+		FTextureRHIRef StagingTexture;
+		FTextureRHIRef ReadbackTexture;
+		void* ResultsBuffer = nullptr;
+		int32 MappedStride = 0;
+
+		FAdaptedVideoFrameLayerI420* CurrentOuputBuffer = nullptr;
+
+		void OnRHIStageComplete();
+		void CleanUp();
+	};
+} // namespace UE::PixelStreaming
diff --git a/Engine/Plugins/Media/PixelStreaming/Source/PixelStreaming/Private/FrameAdapterProcessI420Compute.cpp b/Engine/Plugins/Media/PixelStreaming/Source/PixelStreaming/Private/FrameAdapterProcessI420Compute.cpp
new file mode 100644
--- /dev/null
+++ b/Engine/Plugins/Media/PixelStreaming/Source/PixelStreaming/Private/FrameAdapterProcessI420Compute.cpp
@@ -0,0 +1,151 @@
+// Copyright Epic Games, Inc. All Rights Reserved.
+
+#include "FrameAdapterProcessI420Compute.h"
+#include "FrameAdapterI420.h"
+#include "RGBToYUVShader.h"
+#include "Stats.h"
+#include "Utils.h"
+#include "PixelStreamingSourceFrame.h"
+
+namespace UE::PixelStreaming
+{
+	FFrameAdapterProcessI420Compute::FFrameAdapterProcessI420Compute(float InScale)
+		: Scale(InScale)
+	{
+	}
+
+	FFrameAdapterProcessI420Compute ::~FFrameAdapterProcessI420Compute()
+	{
+		CleanUp();
+	}
+
+	void FFrameAdapterProcessI420Compute::Initialize(int32 SourceWidth, int32 SourceHeight)
+	{
+		const int32 Width = SourceWidth * Scale;
+		const int32 Height = SourceHeight * Scale;
+
+		PlaneYDimensions = { Width, Height };
+		PlaneUVDimensions = { Width / 2, Height / 2 };
+
+		FRHITextureCreateDesc TextureDescY =
+			FRHITextureCreateDesc::Create2D(TEXT("Compute YUV Target"), PlaneYDimensions.X, PlaneYDimensions.Y, EPixelFormat::PF_R8)
+			.SetClearValue(FClearValueBinding::None)
+			.SetFlags(ETextureCreateFlags::UAV)
+			.SetInitialState(ERHIAccess::UAVCompute)
+			.DetermineInititialState();
+
+		FRHITextureCreateDesc TextureDescUV =
+			FRHITextureCreateDesc::Create2D(TEXT("Compute YUV Target"), PlaneUVDimensions.X, PlaneUVDimensions.Y, EPixelFormat::PF_R8)
+			.SetClearValue(FClearValueBinding::None)
+			.SetFlags(ETextureCreateFlags::UAV)
+			.SetInitialState(ERHIAccess::UAVCompute)
+			.DetermineInititialState();
+
+		TextureY = GDynamicRHI->RHICreateTexture(TextureDescY);
+		TextureU = GDynamicRHI->RHICreateTexture(TextureDescUV);
+		TextureV = GDynamicRHI->RHICreateTexture(TextureDescUV);
+
+		FRHITextureCreateDesc StagingDescY =
+			FRHITextureCreateDesc::Create2D(TEXT("YUV Output CPU Texture"), PlaneYDimensions.X, PlaneYDimensions.Y, EPixelFormat::PF_R8)
+			.SetClearValue(FClearValueBinding::None)
+			.SetFlags(ETextureCreateFlags::CPUReadback)
+			.SetInitialState(ERHIAccess::Unknown)
+			.DetermineInititialState();
+
+		FRHITextureCreateDesc StagingDescUV =
+			FRHITextureCreateDesc::Create2D(TEXT("YUV Output CPU Texture"), PlaneUVDimensions.X, PlaneUVDimensions.Y, EPixelFormat::PF_R8)
+			.SetClearValue(FClearValueBinding::None)
+			.SetFlags(ETextureCreateFlags::CPUReadback)
+			.SetInitialState(ERHIAccess::Unknown)
+			.DetermineInititialState();
+
+		StagingTextureY = GDynamicRHI->RHICreateTexture(StagingDescY);
+		StagingTextureU = GDynamicRHI->RHICreateTexture(StagingDescUV);
+		StagingTextureV = GDynamicRHI->RHICreateTexture(StagingDescUV);
+
+		TextureYUAV = GDynamicRHI->RHICreateUnorderedAccessView(TextureY, 0, 0, 0);
+		TextureUUAV = GDynamicRHI->RHICreateUnorderedAccessView(TextureU, 0, 0, 0);
+		TextureVUAV = GDynamicRHI->RHICreateUnorderedAccessView(TextureV, 0, 0, 0);
+
+		int32 OutWidth, OutHeight;
+		GDynamicRHI->RHIMapStagingSurface(StagingTextureY, nullptr, MappedY, OutWidth, OutHeight);
+		YStride = OutWidth;
+		GDynamicRHI->RHIMapStagingSurface(StagingTextureU, nullptr, MappedU, OutWidth, OutHeight);
+		UStride = OutWidth;
+		GDynamicRHI->RHIMapStagingSurface(StagingTextureV, nullptr, MappedV, OutWidth, OutHeight);
+		VStride = OutWidth;
+
+		FPixelStreamingFrameAdapterProcess::Initialize(SourceWidth, SourceHeight);
+	}
+
+	void FFrameAdapterProcessI420Compute::OnSourceResolutionChanged(int32 OldWidth, int32 OldHeight, int32 NewWidth, int32 NewHeight)
+	{
+		CleanUp();
+	}
+
+	TSharedPtr<IPixelStreamingAdaptedVideoFrameLayer> FFrameAdapterProcessI420Compute::CreateOutputBuffer(int32 SourceWidth, int32 SourceHeight)
+	{
+		const int32 Width = SourceWidth * Scale;
+		const int32 Height = SourceHeight * Scale;
+		return MakeShared<FAdaptedVideoFrameLayerI420>(webrtc::I420Buffer::Create(Width, Height));
+	}
+
+	void FFrameAdapterProcessI420Compute::BeginProcess(const FPixelStreamingSourceFrame& SourceFrame)
+	{
+		CurrentOuputBuffer = StaticCast<FAdaptedVideoFrameLayerI420*>(GetWriteBuffer().Get());
+		CurrentOuputBuffer->Metadata = SourceFrame.Metadata.Copy();
+		CurrentOuputBuffer->Metadata.ProcessName = "I420 Compute";
+		CurrentOuputBuffer->Metadata.AdaptCallTime = FPlatformTime::Cycles64();
+
+		FTexture2DRHIRef SourceTexture = SourceFrame.FrameTexture;
+
+		FRHICommandListImmediate& RHICmdList = FRHICommandListExecutor::GetImmediateCommandList();
+
+		RHICmdList.EnqueueLambda([CurrentOuputBuffer = CurrentOuputBuffer](FRHICommandListImmediate&) { CurrentOuputBuffer->Metadata.AdaptProcessStartTime = FPlatformTime::Cycles64(); });
+
+		FRGBToYUVShaderParameters ShaderParameters;
+		ShaderParameters.SourceTexture = SourceTexture;
+		ShaderParameters.DestPlaneYDimensions = PlaneYDimensions;
+		ShaderParameters.DestPlaneUVDimensions = PlaneUVDimensions;
+		ShaderParameters.DestPlaneY = TextureYUAV;
+		ShaderParameters.DestPlaneU = TextureUUAV;
+		ShaderParameters.DestPlaneV = TextureVUAV;
+		FRGBToYUVShader::Dispatch(RHICmdList, ShaderParameters);
+
+		RHICmdList.CopyTexture(TextureY, StagingTextureY, {});
+		RHICmdList.CopyTexture(TextureU, StagingTextureU, {});
+		RHICmdList.CopyTexture(TextureV, StagingTextureV, {});
+
+		// by adding this shared ref to the rhi lambda we can ensure that 'this' will not be destroyed
+		// until after the rhi thread is done with it, so all the commands will still have valid references.
+		TSharedRef<FFrameAdapterProcessI420Compute> ThisRHIRef = StaticCastSharedRef<FFrameAdapterProcessI420Compute>(AsShared());
+		RHICmdList.EnqueueLambda([ThisRHIRef](FRHICommandListImmediate&) { ThisRHIRef->OnRHIStageComplete(); });
+	}
+
+	void FFrameAdapterProcessI420Compute::OnRHIStageComplete()
+	{
+		CurrentOuputBuffer->Metadata.AdaptProcessFinalizeTime = FPlatformTime::Cycles64();
+
+		rtc::scoped_refptr<webrtc::I420Buffer> I420Buffer = CurrentOuputBuffer->GetI420Buffer();
+
+		MemCpyStride(I420Buffer->MutableDataY(), MappedY, I420Buffer->StrideY(), YStride, PlaneYDimensions.Y);
+		MemCpyStride(I420Buffer->MutableDataU(), MappedU, I420Buffer->StrideU(), UStride, PlaneUVDimensions.Y);
+		MemCpyStride(I420Buffer->MutableDataV(), MappedV, I420Buffer->StrideV(), VStride, PlaneUVDimensions.Y);
+
+		CurrentOuputBuffer->Metadata.AdaptProcessEndTime = FPlatformTime::Cycles64();
+		CurrentOuputBuffer = nullptr;
+
+		EndProcess();
+	}
+
+	void FFrameAdapterProcessI420Compute::CleanUp()
+	{
+		GDynamicRHI->RHIUnmapStagingSurface(StagingTextureY);
+		GDynamicRHI->RHIUnmapStagingSurface(StagingTextureU);
+		GDynamicRHI->RHIUnmapStagingSurface(StagingTextureV);
+
+		MappedY = nullptr;
+		MappedU = nullptr;
+		MappedV = nullptr;
+	}
+} // namespace UE::PixelStreaming
diff --git a/Engine/Plugins/Media/PixelStreaming/Source/PixelStreaming/Private/FrameAdapterProcessI420Compute.h b/Engine/Plugins/Media/PixelStreaming/Source/PixelStreaming/Private/FrameAdapterProcessI420Compute.h
new file mode 100644
--- /dev/null
+++ b/Engine/Plugins/Media/PixelStreaming/Source/PixelStreaming/Private/FrameAdapterProcessI420Compute.h
@@ -0,0 +1,59 @@
+// Copyright Epic Games, Inc. All Rights Reserved.
+
+#pragma once
+
+#include "PixelStreamingFrameAdapterProcess.h"
+
+namespace UE::PixelStreaming
+{
+	class FAdaptedVideoFrameLayerI420;
+	
+	class FFrameAdapterProcessI420Compute : public FPixelStreamingFrameAdapterProcess
+	{
+	public:
+		FFrameAdapterProcessI420Compute(float InScale);
+		virtual ~FFrameAdapterProcessI420Compute();
+
+	protected:
+		virtual void Initialize(int32 SourceWidth, int32 SourceHeight) override;
+		virtual void OnSourceResolutionChanged(int32 OldWidth, int32 OldHeight, int32 NewWidth, int32 NewHeight) override;
+		virtual TSharedPtr<IPixelStreamingAdaptedVideoFrameLayer> CreateOutputBuffer(int32 SourceWidth, int32 SourceHeight) override;
+		virtual void BeginProcess(const FPixelStreamingSourceFrame& SourceFrame) override;
+
+	private:
+		float Scale = 1.0f;
+
+		// dimensions of the texures
+		FIntPoint PlaneYDimensions;
+		FIntPoint PlaneUVDimensions;
+
+		// used as targets for the compute shader
+		FTextureRHIRef TextureY;
+		FTextureRHIRef TextureU;
+		FTextureRHIRef TextureV;
+
+		// the UAVs of the targets
+		FUnorderedAccessViewRHIRef TextureYUAV;
+		FUnorderedAccessViewRHIRef TextureUUAV;
+		FUnorderedAccessViewRHIRef TextureVUAV;
+
+		// cpu readable copies of the targets above
+		FTextureRHIRef StagingTextureY;
+		FTextureRHIRef StagingTextureU;
+		FTextureRHIRef StagingTextureV;
+
+		// memory mapped pointers of the staging textures
+		void* MappedY = nullptr;
+		void* MappedU = nullptr;
+		void* MappedV = nullptr;
+
+		int32 YStride = 0;
+		int32 UStride = 0;
+		int32 VStride = 0;
+
+		FAdaptedVideoFrameLayerI420* CurrentOuputBuffer = nullptr;
+
+		void OnRHIStageComplete();
+		void CleanUp();
+	};
+} // namespace UE::PixelStreaming
diff --git a/Engine/Plugins/Media/PixelStreaming/Source/PixelStreaming/Private/FrameBufferH264.cpp b/Engine/Plugins/Media/PixelStreaming/Source/PixelStreaming/Private/FrameBufferH264.cpp
new file mode 100644
--- /dev/null
+++ b/Engine/Plugins/Media/PixelStreaming/Source/PixelStreaming/Private/FrameBufferH264.cpp
@@ -0,0 +1,64 @@
+// Copyright Epic Games, Inc. All Rights Reserved.
+
+#include "FrameBufferH264.h"
+#include "FrameAdapterH264.h"
+
+namespace UE::PixelStreaming
+{
+	FFrameBufferH264Base::FFrameBufferH264Base(TSharedPtr<IPixelStreamingFrameSource> InFrameSource)
+		: FrameSource(InFrameSource)
+	{
+	}
+
+	FFrameBufferH264Simulcast::FFrameBufferH264Simulcast(TSharedPtr<IPixelStreamingFrameSource> InFrameSource)
+		: FFrameBufferH264Base(InFrameSource)
+	{
+	}
+
+	int FFrameBufferH264Simulcast::width() const
+	{
+		return FrameSource->GetWidth(GetNumLayers() - 1);
+	}
+
+	int FFrameBufferH264Simulcast::height() const
+	{
+		return FrameSource->GetHeight(GetNumLayers() - 1);
+	}
+
+	int FFrameBufferH264Simulcast::GetNumLayers() const
+	{
+		return FrameSource->GetNumLayers();
+	}
+
+	FFrameBufferH264::FFrameBufferH264(TSharedPtr<IPixelStreamingFrameSource> InFrameSource, int InLayerIndex)
+		: FFrameBufferH264Base(InFrameSource)
+		, LayerIndex(InLayerIndex)
+	{
+	}
+
+	int FFrameBufferH264::width() const
+	{
+		return FrameSource->GetWidth(LayerIndex);
+	}
+
+	int FFrameBufferH264::height() const
+	{
+		return FrameSource->GetHeight(LayerIndex);
+	}
+
+	FAdaptedVideoFrameLayerH264* FFrameBufferH264::GetAdaptedLayer() const
+	{
+		EnsureCachedAdaptedLayer();
+		return CachedAdaptedLayer.Get();
+	}
+
+	void FFrameBufferH264::EnsureCachedAdaptedLayer() const
+	{
+		if (CachedAdaptedLayer == nullptr)
+		{
+			FPixelStreamingFrameAdapter* FrameAdapter = StaticCast<FPixelStreamingFrameAdapter*>(FrameSource.Get());
+			CachedAdaptedLayer = StaticCastSharedPtr<FAdaptedVideoFrameLayerH264>(FrameAdapter->ReadOutput(LayerIndex));
+			CachedAdaptedLayer->Metadata.Layer = LayerIndex;
+		}
+	}
+} // namespace UE::PixelStreaming
diff --git a/Engine/Plugins/Media/PixelStreaming/Source/PixelStreaming/Private/FrameBufferH264.h b/Engine/Plugins/Media/PixelStreaming/Source/PixelStreaming/Private/FrameBufferH264.h
new file mode 100644
--- /dev/null
+++ b/Engine/Plugins/Media/PixelStreaming/Source/PixelStreaming/Private/FrameBufferH264.h
@@ -0,0 +1,70 @@
+// Copyright Epic Games, Inc. All Rights Reserved.
+
+#pragma once
+
+#include "IPixelStreamingFrameBuffer.h"
+#include "RHI.h"
+
+namespace UE::PixelStreaming
+{
+	class FAdaptedVideoFrameLayerH264;
+	
+	/*
+	 * The base of the H264 frame buffer. Shouldn't be used directly.
+	 */
+	class FFrameBufferH264Base : public IPixelStreamingFrameBuffer
+	{
+	public:
+		FFrameBufferH264Base(TSharedPtr<IPixelStreamingFrameSource> InFrameSource);
+		virtual ~FFrameBufferH264Base() = default;
+
+	protected:
+		TSharedPtr<IPixelStreamingFrameSource> FrameSource;
+	};
+
+	/*
+	 * A frame buffer for simulcast encoders. The adapter held within should provide multiple layers
+	 * to be streamed. Should only be used to create single layer frame buffers.
+	 */
+	class FFrameBufferH264Simulcast : public FFrameBufferH264Base
+	{
+	public:
+		FFrameBufferH264Simulcast(TSharedPtr<IPixelStreamingFrameSource> InFrameSource);
+		virtual ~FFrameBufferH264Simulcast() = default;
+
+		virtual EPixelStreamingFrameBufferType GetFrameBufferType() const override { return EPixelStreamingFrameBufferType::Simulcast; }
+
+		virtual int width() const override;
+		virtual int height() const override;
+
+		TSharedPtr<IPixelStreamingFrameSource> GetFrameSource() const { return FrameSource; }
+		int GetNumLayers() const;
+	};
+
+	/*
+	 * A frame buffer for single layer encoders. Contains a layer index into the adapter. Should be
+	 * used in the encoder.
+	 */
+	class FFrameBufferH264 : public FFrameBufferH264Base
+	{
+	public:
+		FFrameBufferH264(TSharedPtr<IPixelStreamingFrameSource> InFrameSource, int InLayerIndex = 0);
+		virtual ~FFrameBufferH264() = default;
+
+		virtual EPixelStreamingFrameBufferType GetFrameBufferType() const override { return EPixelStreamingFrameBufferType::Layer; }
+
+		virtual int width() const override;
+		virtual int height() const override;
+
+		FAdaptedVideoFrameLayerH264* GetAdaptedLayer() const;
+
+	private:
+		int LayerIndex;
+
+		// used so we dont call ReadOutput multiple times and possibly flip the read buffer and get a different frame
+		// probably not a big deal but i would like to make sure that the buffer always refers to the same frame.
+		// this way allows the user to call GetAdaptedLayer() repeatedly and will always get the same frame.
+		void EnsureCachedAdaptedLayer() const;
+		mutable TSharedPtr<FAdaptedVideoFrameLayerH264> CachedAdaptedLayer;
+	};
+} // namespace UE::PixelStreaming
diff --git a/Engine/Plugins/Media/PixelStreaming/Source/PixelStreaming/Private/FrameBufferI420.cpp b/Engine/Plugins/Media/PixelStreaming/Source/PixelStreaming/Private/FrameBufferI420.cpp
new file mode 100644
--- /dev/null
+++ b/Engine/Plugins/Media/PixelStreaming/Source/PixelStreaming/Private/FrameBufferI420.cpp
@@ -0,0 +1,74 @@
+// Copyright Epic Games, Inc. All Rights Reserved.
+
+#include "FrameBufferI420.h"
+#include "FrameAdapterI420.h"
+
+namespace UE::PixelStreaming
+{
+	FFrameBufferI420Base::FFrameBufferI420Base(TSharedPtr<IPixelStreamingFrameSource> InFrameSource)
+		: FrameSource(InFrameSource)
+	{
+	}
+
+	FFrameBufferI420Simulcast::FFrameBufferI420Simulcast(TSharedPtr<IPixelStreamingFrameSource> InFrameSource)
+		: FFrameBufferI420Base(InFrameSource)
+	{
+	}
+
+	int FFrameBufferI420Simulcast::width() const
+	{
+		return FrameSource->GetWidth(GetNumLayers() - 1);
+	}
+
+	int FFrameBufferI420Simulcast::height() const
+	{
+		return FrameSource->GetHeight(GetNumLayers() - 1);
+	}
+
+	int FFrameBufferI420Simulcast::GetNumLayers() const
+	{
+		return FrameSource->GetNumLayers();
+	}
+
+	FFrameBufferI420::FFrameBufferI420(TSharedPtr<IPixelStreamingFrameSource> InFrameSource, int InLayerIndex)
+		: FFrameBufferI420Base(InFrameSource)
+		, LayerIndex(InLayerIndex)
+	{
+	}
+
+	int FFrameBufferI420::width() const
+	{
+		return FrameSource->GetWidth(LayerIndex);
+	}
+
+	int FFrameBufferI420::height() const
+	{
+		return FrameSource->GetHeight(LayerIndex);
+	}
+
+	rtc::scoped_refptr<webrtc::I420BufferInterface> FFrameBufferI420::ToI420()
+	{
+		return GetAdaptedLayer()->GetI420Buffer();
+	}
+
+	const webrtc::I420BufferInterface* FFrameBufferI420::GetI420() const
+	{
+		return GetAdaptedLayer()->GetI420Buffer().get();
+	}
+
+	FAdaptedVideoFrameLayerI420* FFrameBufferI420::GetAdaptedLayer() const
+	{
+		EnsureCachedAdaptedLayer();
+		return CachedAdaptedLayer.Get();
+	}
+
+	void FFrameBufferI420::EnsureCachedAdaptedLayer() const
+	{
+		if (CachedAdaptedLayer == nullptr)
+		{
+			FPixelStreamingFrameAdapter* FrameAdapter = StaticCast<FPixelStreamingFrameAdapter*>(FrameSource.Get());
+			CachedAdaptedLayer = StaticCastSharedPtr<FAdaptedVideoFrameLayerI420>(FrameAdapter->ReadOutput(LayerIndex));
+			CachedAdaptedLayer->Metadata.Layer = LayerIndex;
+		}
+	}
+} // namespace UE::PixelStreaming
diff --git a/Engine/Plugins/Media/PixelStreaming/Source/PixelStreaming/Private/FrameBufferI420.h b/Engine/Plugins/Media/PixelStreaming/Source/PixelStreaming/Private/FrameBufferI420.h
new file mode 100644
--- /dev/null
+++ b/Engine/Plugins/Media/PixelStreaming/Source/PixelStreaming/Private/FrameBufferI420.h
@@ -0,0 +1,62 @@
+// Copyright Epic Games, Inc. All Rights Reserved.
+
+#pragma once
+
+#include "IPixelStreamingFrameBuffer.h"
+
+namespace UE::PixelStreaming
+{
+	class FAdaptedVideoFrameLayerI420;
+
+	class FFrameBufferI420Base : public IPixelStreamingFrameBuffer
+	{
+	public:
+		FFrameBufferI420Base(TSharedPtr<IPixelStreamingFrameSource> InFrameSource);
+		virtual ~FFrameBufferI420Base() = default;
+
+	protected:
+		TSharedPtr<IPixelStreamingFrameSource> FrameSource;
+	};
+
+	class FFrameBufferI420Simulcast : public FFrameBufferI420Base
+	{
+	public:
+		FFrameBufferI420Simulcast(TSharedPtr<IPixelStreamingFrameSource> InFrameSource);
+		virtual ~FFrameBufferI420Simulcast() = default;
+
+		virtual EPixelStreamingFrameBufferType GetFrameBufferType() const { return EPixelStreamingFrameBufferType::Simulcast; }
+
+		virtual int width() const override;
+		virtual int height() const override;
+
+		TSharedPtr<IPixelStreamingFrameSource> GetFrameAdapter() const { return FrameSource; }
+		int GetNumLayers() const;
+	};
+
+	/*
+	* A single frame buffer for I420 frames
+	*/
+	class FFrameBufferI420 : public FFrameBufferI420Base
+	{
+	public:
+		FFrameBufferI420(TSharedPtr<IPixelStreamingFrameSource> InFrameSource, int InLayerIndex);
+		virtual ~FFrameBufferI420() = default;
+
+		virtual EPixelStreamingFrameBufferType GetFrameBufferType() const { return EPixelStreamingFrameBufferType::Layer; }
+
+		virtual int width() const override;
+		virtual int height() const override;
+
+		virtual rtc::scoped_refptr<webrtc::I420BufferInterface> ToI420() override;
+		virtual const webrtc::I420BufferInterface* GetI420() const override;
+
+		FAdaptedVideoFrameLayerI420* GetAdaptedLayer() const;
+
+	private:
+		int32 LayerIndex;
+
+		// see comments in FrameBufferH264.h
+		void EnsureCachedAdaptedLayer() const;
+		mutable TSharedPtr<FAdaptedVideoFrameLayerI420> CachedAdaptedLayer;
+	};
+} // namespace UE::PixelStreaming
diff --git a/Engine/Plugins/Media/PixelStreaming/Source/PixelStreaming/Private/FrameBufferInitialize.cpp b/Engine/Plugins/Media/PixelStreaming/Source/PixelStreaming/Private/FrameBufferInitialize.cpp
new file mode 100644
--- /dev/null
+++ b/Engine/Plugins/Media/PixelStreaming/Source/PixelStreaming/Private/FrameBufferInitialize.cpp
@@ -0,0 +1,22 @@
+// Copyright Epic Games, Inc. All Rights Reserved.
+
+#include "FrameBufferInitialize.h"
+#include "IPixelStreamingFrameSource.h"
+
+namespace UE::PixelStreaming
+{
+	FFrameBufferInitialize::FFrameBufferInitialize(TSharedPtr<IPixelStreamingFrameSource> InFrameSource)
+		: FrameSource(InFrameSource)
+	{
+	}
+
+	int FFrameBufferInitialize::width() const
+	{
+		return FrameSource->GetWidth(0);
+	}
+
+	int FFrameBufferInitialize::height() const
+	{
+		return FrameSource->GetHeight(0);
+	}
+} // namespace UE::PixelStreaming
diff --git a/Engine/Plugins/Media/PixelStreaming/Source/PixelStreaming/Private/FrameBufferInitialize.h b/Engine/Plugins/Media/PixelStreaming/Source/PixelStreaming/Private/FrameBufferInitialize.h
new file mode 100644
--- /dev/null
+++ b/Engine/Plugins/Media/PixelStreaming/Source/PixelStreaming/Private/FrameBufferInitialize.h
@@ -0,0 +1,26 @@
+// Copyright Epic Games, Inc. All Rights Reserved.
+
+#pragma once
+
+#include "IPixelStreamingFrameBuffer.h"
+
+namespace UE::PixelStreaming
+{
+	/*
+	 * Used only to initialize hardware encoders and query for size. Will never actually provide image data.
+	 */
+	class FFrameBufferInitialize : public IPixelStreamingFrameBuffer
+	{
+	public:
+		FFrameBufferInitialize(TSharedPtr<IPixelStreamingFrameSource> InFrameSource);
+		virtual ~FFrameBufferInitialize() = default;
+
+		virtual EPixelStreamingFrameBufferType GetFrameBufferType() const { return EPixelStreamingFrameBufferType::Initialize; }
+
+		virtual int width() const override;
+		virtual int height() const override;
+
+	private:
+		TSharedPtr<IPixelStreamingFrameSource> FrameSource;
+	};
+} // namespace UE::PixelStreaming
