Early Pixel Streaming 5.1 features, bugfixes, and refactors that benefit main. (From various authors)

diff --git a/Engine/Plugins/Media/PixelStreaming/Source/PixelStreaming/Private/IPlayerSession.h b/Engine/Plugins/Media/PixelStreaming/Source/PixelStreaming/Private/IPlayerSession.h
--- a/Engine/Plugins/Media/PixelStreaming/Source/PixelStreaming/Private/IPlayerSession.h
+++ b/Engine/Plugins/Media/PixelStreaming/Source/PixelStreaming/Private/IPlayerSession.h
@@ -5,12 +5,13 @@
 #include "WebRTCIncludes.h"
 #include "PixelStreamingProtocolDefs.h"
 #include "PixelStreamingPlayerId.h"
+#include "VideoSourceBase.h"
+#include "Templates/SharedPointer.h"
 
 class IPixelStreamingAudioSink;
 
 namespace UE::PixelStreaming
 {
-	class FSignallingServerConnection;
 	class FDataChannelObserver;
 
 	class IPlayerSession : public webrtc::PeerConnectionObserver
@@ -18,10 +19,11 @@ namespace UE::PixelStreaming
 	public:
 		virtual ~IPlayerSession() {}
 
-		virtual webrtc::PeerConnectionInterface& GetPeerConnection() = 0;
+		virtual rtc::scoped_refptr<webrtc::PeerConnectionInterface> GetPeerConnection() = 0;
 		virtual void SetPeerConnection(const rtc::scoped_refptr<webrtc::PeerConnectionInterface>& InPeerConnection) = 0;
 		virtual void SetDataChannel(const rtc::scoped_refptr<webrtc::DataChannelInterface>& InDataChannel) = 0;
 		virtual void SetVideoSource(const rtc::scoped_refptr<webrtc::VideoTrackSourceInterface>& InVideoSource) = 0;
+		virtual rtc::scoped_refptr<webrtc::VideoTrackSourceInterface> GetVideoSource() const = 0;
 
 		virtual void OnAnswer(FString Sdp) = 0;
 		virtual void OnRemoteIceCandidate(const FString& SdpMid, int SdpMLineIndex, const FString& Sdp) = 0;
@@ -30,15 +32,16 @@ namespace UE::PixelStreaming
 		virtual FName GetSessionType() const = 0;
 		virtual FPixelStreamingPlayerId GetPlayerId() const = 0;
 		virtual IPixelStreamingAudioSink* GetAudioSink() = 0;
-		virtual FDataChannelObserver* GetDataChannelObserver() = 0;
+		virtual TSharedPtr<FDataChannelObserver> GetDataChannelObserver() = 0;
 
 		virtual bool SendMessage(Protocol::EToPlayerMsg Type, const FString& Descriptor) const = 0;
-		virtual void SendQualityControlStatus(bool bIsQualityController) const = 0;
-		virtual void SendFreezeFrame(const TArray64<uint8>& JpegBytes) const = 0;
-		virtual void SendFileData(const TArray<uint8>& ByteData, const FString& MimeType, const FString& FileExtension) const = 0;
-		virtual void SendUnfreezeFrame() const = 0;
-		virtual void SendArbitraryData(const TArray<uint8>& DataBytes, const uint8 MessageType) const = 0;
-		virtual void SendVideoEncoderQP(double QP) const = 0;
+		virtual bool SendQualityControlStatus(bool bIsQualityController) const = 0;
+		virtual bool SendInputControlStatus(bool bIsInputController) const = 0;
+		virtual bool SendFreezeFrame(const TArray64<uint8>& JpegBytes) const = 0;
+		virtual bool SendFileData(const TArray<uint8>& ByteData, const FString& MimeType, const FString& FileExtension) const = 0;
+		virtual bool SendUnfreezeFrame() const = 0;
+		virtual bool SendArbitraryData(const TArray<uint8>& DataBytes, const uint8 MessageType) const = 0;
+		virtual bool SendVideoEncoderQP(double QP) const = 0;
 		virtual void PollWebRTCStats() const = 0;
 		virtual void OnDataChannelClosed() const = 0;
 	};
diff --git a/Engine/Plugins/Media/PixelStreaming/Source/PixelStreaming/Private/InputDevice.cpp b/Engine/Plugins/Media/PixelStreaming/Source/PixelStreaming/Private/InputDevice.cpp
--- a/Engine/Plugins/Media/PixelStreaming/Source/PixelStreaming/Private/InputDevice.cpp
+++ b/Engine/Plugins/Media/PixelStreaming/Source/PixelStreaming/Private/InputDevice.cpp
@@ -17,6 +17,19 @@
 #include "Serialization/JsonSerializer.h"
 #include "Framework/Application/SlateUser.h"
 #include "EditorPixelStreamingSettings.h"
+#include "IPixelStreamingStreamer.h"
+
+#if WITH_EDITOR
+	#include "AnalyticsEventAttribute.h"
+	#include "Editor.h"
+	#include "EngineAnalytics.h"
+	#include "IAssetViewport.h"
+	#include "LevelEditor.h"
+	#include "Editor/EditorEngine.h"
+	#include "Engine/EngineTypes.h"
+	#include "Framework/Notifications/NotificationManager.h"
+	#include "Widgets/Notifications/SNotificationList.h"
+#endif // WITH_EDITOR
 
 DECLARE_LOG_CATEGORY_EXTERN(LogPixelStreamingInputDevice, Log, VeryVerbose);
 DEFINE_LOG_CATEGORY(LogPixelStreamingInputDevice);
@@ -146,36 +159,20 @@ namespace UE::PixelStreaming
 		, FocusedPos(UnfocusedPos)
 		, PixelStreamingModule(FPixelStreamingModule::GetModule())
 	{
-		if (GEngine->GameViewport && !GEngine->GameViewport->HasSoftwareCursor(EMouseCursor::Default))
+		// default to the scene viewport if we have a game engine. if we are
+		// running editor we require the user to set the viewport via FStreamer
+		if (UGameEngine* GameEngine = Cast<UGameEngine>(GEngine))
 		{
-			// Pixel streaming always requires a default software cursor as it needs
-			// to be shown on the browser to allow the user to click UI elements.
-			const UPixelStreamingSettings* Settings = GetDefault<UPixelStreamingSettings>();
-			check(Settings);
-
-			// Check to see if we want to hide the cursor (by making it invisible).
-			// This is required if we want to make the cursor client-side (displayed
-			// only by the the browser).
-			bool bHideCursor = UE::PixelStreaming::Settings::IsPixelStreamingHideCursor();
-			if (bHideCursor)
-			{
-				GEngine->GameViewport->AddSoftwareCursor(EMouseCursor::Default, Settings->HiddenCursorClassName);
-				GEngine->GameViewport->AddSoftwareCursor(EMouseCursor::TextEditBeam, Settings->HiddenCursorClassName);
-			}
-			else
-			{
-				GEngine->GameViewport->AddSoftwareCursor(EMouseCursor::Default, Settings->DefaultCursorClassName);
-				GEngine->GameViewport->AddSoftwareCursor(EMouseCursor::TextEditBeam, Settings->TextEditBeamCursorClassName);
-			}
+			SetTargetViewport(GameEngine->SceneViewport);
 		}
 	}
 
 	// XY positions are the ratio (0.0..1.0) along a viewport axis, quantized
 	// into an uint16 (0..65536). This allows the browser viewport and player
 	// viewport to have a different size.
-	void QuantizeAndNormalize(const FVector2D& InPos, uint16& OutX, uint16& OutY)
+	void QuantizeAndNormalize(TSharedPtr<FSceneViewport> TargetViewport, const FVector2D& InPos, uint16& OutX, uint16& OutY)
 	{
-		FIntPoint SizeXY = GEngine->GameViewport->Viewport->GetSizeXY();
+		FIntPoint SizeXY = TargetViewport->GetSizeXY();
 		OutX = InPos.X / SizeXY.X * 65536.0f;
 		OutY = InPos.Y / SizeXY.Y * 65536.0f;
 	}
@@ -275,7 +272,7 @@ namespace UE::PixelStreaming
 					int16 DeltaX;
 					int16 DeltaY;
 					Event.GetMouseDelta(PosX, PosY, DeltaX, DeltaY);
-					FVector2D CursorPos = GEngine->GameViewport->GetWindow()->GetPositionInScreen() + FVector2D(PosX, PosY);
+					FVector2D CursorPos = TargetViewport->GetCachedGeometry().GetAbsolutePosition() + FVector2D(PosX, PosY);
 					PixelStreamerApplicationWrapper->Cursor->SetPosition(CursorPos.X, CursorPos.Y);
 					MessageHandler->OnRawMouseMove(DeltaX, DeltaY);
 					UE_LOG(LogPixelStreamingInputDevice, VeryVerbose, TEXT("MOUSE_MOVE: Pos = (%d, %d); CursorPos = (%d, %d); Delta = (%d, %d)"), PosX, PosY, static_cast<int>(CursorPos.X), static_cast<int>(CursorPos.Y), DeltaX, DeltaY);
@@ -296,9 +293,9 @@ namespace UE::PixelStreaming
 					uint16 PosX;
 					uint16 PosY;
 					Event.GetMouseClick(Button, PosX, PosY);
-					FVector2D CursorPos = GEngine->GameViewport->GetWindow()->GetPositionInScreen() + FVector2D(PosX, PosY);
+					FVector2D CursorPos = TargetViewport->GetCachedGeometry().GetAbsolutePosition() + FVector2D(PosX, PosY);
 					PixelStreamerApplicationWrapper->Cursor->SetPosition(CursorPos.X, CursorPos.Y);
-					MessageHandler->OnMouseDown(GEngine->GameViewport->GetWindow()->GetNativeWindow(), Button, CursorPos);
+					MessageHandler->OnMouseDown(TargetViewport->FindWindow()->GetNativeWindow(), Button, CursorPos);
 					UE_LOG(LogPixelStreamingInputDevice, Verbose, TEXT("MOUSE_DOWN: Button = %d; Pos = (%d, %d); CursorPos = (%d, %d)"), Button, PosX, PosY, static_cast<int>(CursorPos.X), static_cast<int>(CursorPos.Y));
 
 					// The browser may be faking a mouse when touching so it will send
@@ -312,7 +309,7 @@ namespace UE::PixelStreaming
 					uint16 PosX;
 					uint16 PosY;
 					Event.GetMouseClick(Button, PosX, PosY);
-					FVector2D CursorPos = GEngine->GameViewport->GetWindow()->GetPositionInScreen() + FVector2D(PosX, PosY);
+					FVector2D CursorPos = TargetViewport->GetCachedGeometry().GetAbsolutePosition() + FVector2D(PosX, PosY);
 					PixelStreamerApplicationWrapper->Cursor->SetPosition(CursorPos.X, CursorPos.Y);
 					MessageHandler->OnMouseUp(Button);
 					UE_LOG(LogPixelStreamingInputDevice, Verbose, TEXT("MOUSE_UP: Button = %d; Pos = (%d, %d); CursorPos = (%d, %d)"), Button, PosX, PosY, static_cast<int>(CursorPos.X), static_cast<int>(CursorPos.Y));
@@ -325,7 +322,7 @@ namespace UE::PixelStreaming
 					uint16 PosY;
 					Event.GetMouseWheel(Delta, PosX, PosY);
 					const float SpinFactor = 1 / 120.0f;
-					FVector2D CursorPos = GEngine->GameViewport->GetWindow()->GetPositionInScreen() + FVector2D(PosX, PosY);
+					FVector2D CursorPos = TargetViewport->GetCachedGeometry().GetAbsolutePosition() + FVector2D(PosX, PosY);
 					MessageHandler->OnMouseWheel(Delta * SpinFactor, CursorPos);
 					UE_LOG(LogPixelStreamingInputDevice, Verbose, TEXT("MOUSE_WHEEL: Delta = %d; Pos = (%d, %d); CursorPos = (%d, %d)"), Delta, PosX, PosY, static_cast<int>(CursorPos.X), static_cast<int>(CursorPos.Y));
 				}
@@ -337,8 +334,8 @@ namespace UE::PixelStreaming
 					uint16 PosY;
 					uint8 Force; // Force is between 0.0 and 1.0 so will need to unquantize from byte.
 					Event.GetTouch(TouchIndex, PosX, PosY, Force);
-					FVector2D CursorPos = GEngine->GameViewport->GetWindow()->GetPositionInScreen() + FVector2D(PosX, PosY);
-					MessageHandler->OnTouchStarted(GEngine->GameViewport->GetWindow()->GetNativeWindow(), CursorPos, Force / 255.0f, TouchIndex, 0); // TODO: ControllerId?
+					FVector2D CursorPos = TargetViewport->GetCachedGeometry().GetAbsolutePosition() + FVector2D(PosX, PosY);
+					MessageHandler->OnTouchStarted(TargetViewport->FindWindow()->GetNativeWindow(), CursorPos, Force / 255.0f, TouchIndex, 0); // TODO: ControllerId?
 					UE_LOG(LogPixelStreamingInputDevice, Verbose, TEXT("TOUCH_START: TouchIndex = %d; Pos = (%d, %d); CursorPos = (%d, %d); Force = %.3f"), TouchIndex, PosX, PosY, static_cast<int>(CursorPos.X), static_cast<int>(CursorPos.Y), Force / 255.0f);
 
 					// If the user starts a touch then they may be focusing an editable
@@ -353,7 +350,7 @@ namespace UE::PixelStreaming
 					uint16 PosY;
 					uint8 Force;
 					Event.GetTouch(TouchIndex, PosX, PosY, Force);
-					FVector2D CursorPos = GEngine->GameViewport->GetWindow()->GetPositionInScreen() + FVector2D(PosX, PosY);
+					FVector2D CursorPos = TargetViewport->GetCachedGeometry().GetAbsolutePosition() + FVector2D(PosX, PosY);
 					MessageHandler->OnTouchEnded(CursorPos, TouchIndex, 0); // TODO: ControllerId?
 					UE_LOG(LogPixelStreamingInputDevice, Verbose, TEXT("TOUCH_END: TouchIndex = %d; Pos = (%d, %d); CursorPos = (%d, %d)"), TouchIndex, PosX, PosY, static_cast<int>(CursorPos.X), static_cast<int>(CursorPos.Y));
 				}
@@ -365,7 +362,7 @@ namespace UE::PixelStreaming
 					uint16 PosY;
 					uint8 Force; // Force is between 0.0 and 1.0 so will need to unquantize from byte.
 					Event.GetTouch(TouchIndex, PosX, PosY, Force);
-					FVector2D CursorPos = GEngine->GameViewport->GetWindow()->GetPositionInScreen() + FVector2D(PosX, PosY);
+					FVector2D CursorPos = TargetViewport->GetCachedGeometry().GetAbsolutePosition() + FVector2D(PosX, PosY);
 					MessageHandler->OnTouchMoved(CursorPos, Force / 255.0f, TouchIndex, 0); // TODO: ControllerId?
 					UE_LOG(LogPixelStreamingInputDevice, VeryVerbose, TEXT("TOUCH_MOVE: TouchIndex = %d; Pos = (%d, %d); CursorPos = (%d, %d); Force = %.3f"), TouchIndex, PosX, PosY, static_cast<int>(CursorPos.X), static_cast<int>(CursorPos.Y), Force / 255.0f);
 				}
@@ -626,10 +623,10 @@ namespace UE::PixelStreaming
 
 				if (bEditable)
 				{
-					Pos = Pos - GEngine->GameViewport->GetWindow()->GetPositionInScreen();
+					Pos = Pos - TargetViewport->FindWindow()->GetPositionInScreen();
 					uint16 PosX;
 					uint16 PosY;
-					QuantizeAndNormalize(Pos, PosX, PosY);
+					QuantizeAndNormalize(TargetViewport, Pos, PosX, PosY);
 					JsonObject->SetNumberField(TEXT("x"), PosX);
 					JsonObject->SetNumberField(TEXT("y"), PosY);
 				}
@@ -638,7 +635,9 @@ namespace UE::PixelStreaming
 				TSharedRef<TJsonWriter<TCHAR, TCondensedJsonPrintPolicy<TCHAR>>> JsonWriter = TJsonWriterFactory<TCHAR, TCondensedJsonPrintPolicy<TCHAR>>::Create(&Descriptor);
 				FJsonSerializer::Serialize(JsonObject.ToSharedRef(), JsonWriter);
 
-				PixelStreamingModule->SendCommand(Descriptor);
+				PixelStreamingModule->ForEachStreamer([&Descriptor](TSharedPtr<IPixelStreamingStreamer> Streamer){
+					Streamer->SendPlayerMessage(Protocol::EToPlayerMsg::Command, Descriptor);
+				});
 			}
 		});
 	}
@@ -648,15 +647,14 @@ namespace UE::PixelStreaming
 	namespace
 	{
 
-	#define GET(Type, Var) Type Var = UE::PixelStreaming::ParseBuffer<Type>(Data, Size)
+#define GET(Type, Var) Type Var = UE::PixelStreaming::ParseBuffer<Type>(Data, Size)
 
 		// XY positions are the ratio (0.0..1.0) along a viewport axis, quantized
 		// into an uint16 (0..65536). This allows the browser viewport and player
 		// viewport to have a different size.
-		void UnquantizeAndDenormalize(uint16& InOutX, uint16& InOutY)
+		void UnquantizeAndDenormalize(TSharedPtr<FSceneViewport> TargetViewport, uint16& InOutX, uint16& InOutY)
 		{
-			checkf(GEngine->GameViewport != nullptr, TEXT("PixelStreaming is not supported in editor, and should have been disabled earlier"));
-			FIntPoint SizeXY = GEngine->GameViewport->Viewport->GetSizeXY();
+			FIntPoint SizeXY = TargetViewport->GetSizeXY();
 			InOutX = InOutX / 65536.0f * SizeXY.X;
 			InOutY = InOutY / 65536.0f * SizeXY.Y;
 		}
@@ -664,9 +662,9 @@ namespace UE::PixelStreaming
 		// XY deltas are the ratio (-1.0..1.0) along a viewport axis, quantized
 		// into an int16 (-32767..32767). This allows the browser viewport and
 		// player viewport to have a different size.
-		void UnquantizeAndDenormalize(int16& InOutX, int16& InOutY)
+		void UnquantizeAndDenormalize(TSharedPtr<FSceneViewport> TargetViewport, int16& InOutX, int16& InOutY)
 		{
-			FIntPoint SizeXY = GEngine->GameViewport->Viewport->GetSizeXY();
+			FIntPoint SizeXY = TargetViewport->GetSizeXY();
 			InOutX = InOutX / 32767.0f * SizeXY.X;
 			InOutY = InOutY / 32767.0f * SizeXY.Y;
 		}
@@ -702,7 +700,7 @@ namespace UE::PixelStreaming
 		* @param OutTouches - The array of touches.
 		* @return False if there insufficient room in the receive buffer to read the entire event.
 		*/
-		FTouchesType GetTouches(const uint8*& Data, uint32& Size)
+		FTouchesType GetTouches(TSharedPtr<FSceneViewport> TargetViewport, const uint8*& Data, uint32& Size)
 		{
 			// Get the number of touches in the array.
 			GET(uint8, NumTouches);
@@ -713,7 +711,7 @@ namespace UE::PixelStreaming
 			{
 				GET(FPosType, PosX);
 				GET(FPosType, PosY);
-				UnquantizeAndDenormalize(PosX, PosY);
+				UnquantizeAndDenormalize(TargetViewport, PosX, PosY);
 				GET(uint8, TouchIndex);
 				GET(uint8, Force);
 				GET(uint8, Valid);
@@ -838,7 +836,7 @@ namespace UE::PixelStreaming
 				checkf(Size == 0, TEXT("%d"), Size);
 				UE_LOG(LogPixelStreamingInput, Verbose, TEXT("mouseDown at (%d, %d), button %d"), PosX, PosY, Button);
 
-				UnquantizeAndDenormalize(PosX, PosY);
+				UnquantizeAndDenormalize(TargetViewport, PosX, PosY);
 
 				FEvent MouseDownEvent(EventType::MOUSE_DOWN);
 				MouseDownEvent.SetMouseClick(Button, PosX, PosY);
@@ -853,7 +851,7 @@ namespace UE::PixelStreaming
 				checkf(Size == 0, TEXT("%d"), Size);
 				UE_LOG(LogPixelStreamingInput, Verbose, TEXT("mouseUp at (%d, %d), button %d"), PosX, PosY, Button);
 
-				UnquantizeAndDenormalize(PosX, PosY);
+				UnquantizeAndDenormalize(TargetViewport, PosX, PosY);
 
 				FEvent MouseDownEvent(EventType::MOUSE_UP);
 				MouseDownEvent.SetMouseClick(Button, PosX, PosY);
@@ -869,8 +867,8 @@ namespace UE::PixelStreaming
 				checkf(Size == 0, TEXT("%d"), Size);
 				UE_LOG(LogPixelStreamingInput, Verbose, TEXT("mouseMove to (%d, %d), delta (%d, %d)"), PosX, PosY, DeltaX, DeltaY);
 
-				UnquantizeAndDenormalize(PosX, PosY);
-				UnquantizeAndDenormalize(DeltaX, DeltaY);
+				UnquantizeAndDenormalize(TargetViewport, PosX, PosY);
+				UnquantizeAndDenormalize(TargetViewport, DeltaX, DeltaY);
 
 				FEvent MouseMoveEvent(EventType::MOUSE_MOVE);
 				MouseMoveEvent.SetMouseDelta(PosX, PosY, DeltaX, DeltaY);
@@ -885,7 +883,7 @@ namespace UE::PixelStreaming
 				checkf(Size == 0, TEXT("%d"), Size);
 				UE_LOG(LogPixelStreamingInput, Verbose, TEXT("mouseWheel, delta %d"), Delta);
 
-				UnquantizeAndDenormalize(PosX, PosY);
+				UnquantizeAndDenormalize(TargetViewport, PosX, PosY);
 
 				FEvent MouseWheelEvent(EventType::MOUSE_WHEEL);
 				MouseWheelEvent.SetMouseWheel(Delta, PosX, PosY);
@@ -894,7 +892,7 @@ namespace UE::PixelStreaming
 			}
 			case EToStreamerMsg::TouchStart:
 			{
-				FTouchesType Touches = GetTouches(Data, Size);
+				FTouchesType Touches = GetTouches(TargetViewport, Data, Size);
 				checkf(Size == 0, TEXT("%d"), Size);
 				UE_LOG(LogPixelStreamingInput, Verbose, TEXT("TouchStart: %s"), *TouchesToString(Touches));
 
@@ -911,7 +909,7 @@ namespace UE::PixelStreaming
 			}
 			case EToStreamerMsg::TouchEnd:
 			{
-				FTouchesType Touches = GetTouches(Data, Size);
+				FTouchesType Touches = GetTouches(TargetViewport, Data, Size);
 				checkf(Size == 0, TEXT("%d"), Size);
 				UE_LOG(LogPixelStreamingInput, Verbose, TEXT("TouchEnd: %s"), *TouchesToString(Touches));
 
@@ -930,7 +928,7 @@ namespace UE::PixelStreaming
 			}
 			case EToStreamerMsg::TouchMove:
 			{
-				FTouchesType Touches = GetTouches(Data, Size);
+				FTouchesType Touches = GetTouches(TargetViewport, Data, Size);
 				checkf(Size == 0, TEXT("%d"), Size);
 				UE_LOG(LogPixelStreamingInput, Verbose, TEXT("TouchMove: %s"), *TouchesToString(Touches));
 
@@ -983,7 +981,7 @@ namespace UE::PixelStreaming
 		}
 	}
 
-	#undef GET
+#undef GET
 
 	void FInputDevice::FEvent::GetMouseClick(EMouseButtons::Type& OutButton, uint16& OutPosX, uint16& OutPosY)
 	{
@@ -1026,4 +1024,44 @@ namespace UE::PixelStreaming
 		OutPosX = Data.MouseButton.PosX;
 		OutPosY = Data.MouseButton.PosY;
 	}
+
+	void FInputDevice::SetTargetViewport(TSharedPtr<FSceneViewport> InTargetViewport)
+	{
+		TargetViewport = InTargetViewport;
+
+		if (TargetViewport)
+		{
+			if (UGameEngine* GameEngine = Cast<UGameEngine>(GEngine))
+			{
+				if (TargetViewport == GameEngine->SceneViewport)
+				{
+					if (UGameViewportClient* GameViewport = static_cast<UGameViewportClient*>(TargetViewport->GetClient()))
+					{
+						if (!GameViewport->HasSoftwareCursor(EMouseCursor::Default))
+						{
+							// Pixel streaming always requires a default software cursor as it needs
+							// to be shown on the browser to allow the user to click UI elements.
+							const UPixelStreamingSettings* Settings = GetDefault<UPixelStreamingSettings>();
+							check(Settings);
+
+							// Check to see if we want to hide the cursor (by making it invisible).
+							// This is required if we want to make the cursor client-side (displayed
+							// only by the the browser).
+							bool bHideCursor = UE::PixelStreaming::Settings::IsPixelStreamingHideCursor();
+							if (bHideCursor)
+							{
+								GameViewport->AddSoftwareCursor(EMouseCursor::Default, Settings->HiddenCursorClassName);
+								GameViewport->AddSoftwareCursor(EMouseCursor::TextEditBeam, Settings->HiddenCursorClassName);
+							}
+							else
+							{
+								GameViewport->AddSoftwareCursor(EMouseCursor::Default, Settings->DefaultCursorClassName);
+								GameViewport->AddSoftwareCursor(EMouseCursor::TextEditBeam, Settings->TextEditBeamCursorClassName);
+							}
+						}
+					}
+				}
+			}
+		}
+	}
 } // namespace UE::PixelStreaming
\ No newline at end of file
diff --git a/Engine/Plugins/Media/PixelStreaming/Source/PixelStreaming/Private/InputDevice.h b/Engine/Plugins/Media/PixelStreaming/Source/PixelStreaming/Private/InputDevice.h
--- a/Engine/Plugins/Media/PixelStreaming/Source/PixelStreaming/Private/InputDevice.h
+++ b/Engine/Plugins/Media/PixelStreaming/Source/PixelStreaming/Private/InputDevice.h
@@ -2,16 +2,17 @@
 
 #pragma once
 
-#include "IInputDevice.h"
+#include "IPixelStreamingInputDevice.h"
 #include "Containers/Queue.h"
 #include "GenericPlatform/GenericApplication.h"
 #include "api/data_channel_interface.h"
+#include "Slate/SceneViewport.h"
 
 class IPixelStreamingModule;
 
-namespace UE::PixelStreaming 
+namespace UE::PixelStreaming
 {
-	class FInputDevice : public IInputDevice
+	class FInputDevice : public IPixelStreamingInputDevice
 	{
 	public:
 		/**
@@ -433,7 +434,7 @@ namespace UE::PixelStreaming
 		* Is the application faking touch events?
 		* @return True if the application is faking touch events.
 		*/
-		bool IsFakingTouchEvents() const
+		virtual bool IsFakingTouchEvents() const override
 		{
 			return bFakingTouchEvents;
 		}
@@ -442,7 +443,7 @@ namespace UE::PixelStreaming
 		* Handle incoming messages
 		* @param Buffer - Incoming data buffer.
 		*/
-		void OnMessage(const webrtc::DataBuffer& Buffer);
+		virtual void OnMessage(const webrtc::DataBuffer& Buffer) override;
 
 		/**
 		* How long is the message type specifier and the length specifier in bytes?
@@ -453,6 +454,8 @@ namespace UE::PixelStreaming
 			return MessageHeaderOffset;
 		}
 
+		virtual void SetTargetViewport(TSharedPtr<FSceneViewport> InTargetViewport) override;
+
 	private:
 		/**
 		* Touch only. If the user touches the UI and focuses an editable widget
@@ -523,5 +526,9 @@ namespace UE::PixelStreaming
 		* 2 characters for the length which are skipped
 		*/
 		static const size_t MessageHeaderOffset;
+
+		//bool FindSceneViewportAndLevel();
+		TSharedPtr<FSceneViewport> TargetViewport;
+		bool bIsGameViewport = false;
 	};
 } // namespace UE::PixelStreaming
\ No newline at end of file
diff --git a/Engine/Plugins/Media/PixelStreaming/Source/PixelStreaming/Private/PixelStreamingAudioComponent.cpp b/Engine/Plugins/Media/PixelStreaming/Source/PixelStreaming/Private/PixelStreamingAudioComponent.cpp
--- a/Engine/Plugins/Media/PixelStreaming/Source/PixelStreaming/Private/PixelStreamingAudioComponent.cpp
+++ b/Engine/Plugins/Media/PixelStreaming/Source/PixelStreaming/Private/PixelStreamingAudioComponent.cpp
@@ -2,12 +2,13 @@
 
 #include "PixelStreamingAudioComponent.h"
 #include "IPixelStreamingModule.h"
+#include "IPixelStreamingStreamer.h"
 #include "PixelStreamingPrivate.h"
 #include "CoreMinimal.h"
 
 /*
-* Component that recieves audio from a remote webrtc connection and outputs it into UE using a "synth component".
-*/
+ * Component that recieves audio from a remote webrtc connection and outputs it into UE using a "synth component".
+ */
 UPixelStreamingAudioComponent::UPixelStreamingAudioComponent(const FObjectInitializer& ObjectInitializer)
 	: Super(ObjectInitializer)
 	, PlayerToHear(FPixelStreamingPlayerId())
@@ -15,19 +16,10 @@ UPixelStreamingAudioComponent::UPixelStreamingAudioComponent(const FObjectInitia
 	, AudioSink(nullptr)
 	, SoundGenerator(MakeShared<FWebRTCSoundGenerator, ESPMode::ThreadSafe>())
 {
-
-	bool bPixelStreamingLoaded = IPixelStreamingModule::IsAvailable();
-
-	// Only output this warning if we are actually running this component (not in commandlet).
-	if (!bPixelStreamingLoaded && !IsRunningCommandlet())
-	{
-		UE_LOG(LogPixelStreaming, Warning, TEXT("Pixel Streaming audio component will not tick because Pixel Streaming module is not loaded. This is expected on dedicated servers."));
-	}
-
 	PreferredBufferLength = 512u;
 	NumChannels = 2;
-	PrimaryComponentTick.bCanEverTick = bPixelStreamingLoaded;
-	SetComponentTickEnabled(bPixelStreamingLoaded);
+	PrimaryComponentTick.bCanEverTick = true;
+	SetComponentTickEnabled(true);
 	bAutoActivate = true;
 };
 
@@ -56,7 +48,16 @@ void UPixelStreamingAudioComponent::BeginDestroy()
 
 bool UPixelStreamingAudioComponent::ListenTo(FString PlayerToListenTo)
 {
+	IPixelStreamingModule& PixelStreamingModule = IPixelStreamingModule::Get();
+	if (!PixelStreamingModule.IsReady())
+	{
+		return false;
+	}
+	return StreamerListenTo(PixelStreamingModule.GetDefaultStreamerID(), PlayerToListenTo);
+}
 
+bool UPixelStreamingAudioComponent::StreamerListenTo(FString StreamerId, FString PlayerToListenTo)
+{
 	if (!IPixelStreamingModule::IsAvailable())
 	{
 		UE_LOG(LogPixelStreaming, Verbose, TEXT("Pixel Streaming audio component could not listen to anything because Pixel Streaming module is not loaded. This is expected on dedicated servers."));
@@ -71,7 +72,29 @@ bool UPixelStreamingAudioComponent::ListenTo(FString PlayerToListenTo)
 
 	PlayerToHear = PlayerToListenTo;
 
-	IPixelStreamingAudioSink* CandidateSink = WillListenToAnyPlayer() ? PixelStreamingModule.GetUnlistenedAudioSink() : PixelStreamingModule.GetPeerAudioSink(FPixelStreamingPlayerId(PlayerToHear));
+	if (StreamerId == FString())
+	{
+		TArray<FString> StreamerIds = PixelStreamingModule.GetStreamerIds();
+		if (StreamerIds.Num() > 0)
+		{
+			StreamerToHear = StreamerIds[0];
+		}
+		else
+		{
+			StreamerToHear = PixelStreamingModule.GetDefaultStreamerID();
+		}
+	}
+	else
+	{
+		StreamerToHear = StreamerId;
+	}
+
+	TSharedPtr<IPixelStreamingStreamer> Streamer = PixelStreamingModule.GetStreamer(StreamerToHear);
+	if (!Streamer)
+	{
+		return false;
+	}
+	IPixelStreamingAudioSink* CandidateSink = WillListenToAnyPlayer() ? Streamer->GetUnlistenedAudioSink() : Streamer->GetPeerAudioSink(FPixelStreamingPlayerId(PlayerToHear));
 
 	if (CandidateSink == nullptr)
 	{
@@ -87,6 +110,7 @@ bool UPixelStreamingAudioComponent::ListenTo(FString PlayerToListenTo)
 void UPixelStreamingAudioComponent::Reset()
 {
 	PlayerToHear = FString();
+	StreamerToHear = FString();
 	SoundGenerator->bShouldGenerateAudio = false;
 	if (AudioSink)
 	{
@@ -141,6 +165,14 @@ void UPixelStreamingAudioComponent::OnConsumerRemoved()
 void UPixelStreamingAudioComponent::TickComponent(float DeltaTime, ELevelTick TickType, FActorComponentTickFunction* ThisTickFunction)
 {
 
+	bool bPixelStreamingLoaded = IPixelStreamingModule::IsAvailable();
+
+	// Early out if running in commandlet
+	if (IsRunningCommandlet())
+	{
+		return;
+	}
+
 	// if auto connect turned off don't bother
 	if (!bAutoFindPeer)
 	{
@@ -153,15 +185,15 @@ void UPixelStreamingAudioComponent::TickComponent(float DeltaTime, ELevelTick Ti
 		return;
 	}
 
-	if (ListenTo(PlayerToHear))
+	if (StreamerListenTo(StreamerToHear, PlayerToHear))
 	{
 		UE_LOG(LogPixelStreaming, Log, TEXT("PixelStreaming audio component found a WebRTC peer to listen to."));
 	}
 }
 
 /*
-* ---------------- FWebRTCSoundGenerator -------------------------
-*/
+ * ---------------- FWebRTCSoundGenerator -------------------------
+ */
 
 FWebRTCSoundGenerator::FWebRTCSoundGenerator()
 	: Params()
diff --git a/Engine/Plugins/Media/PixelStreaming/Source/PixelStreaming/Private/PixelStreamingAudioDeviceModule.cpp b/Engine/Plugins/Media/PixelStreaming/Source/PixelStreaming/Private/PixelStreamingAudioDeviceModule.cpp
new file mode 100644
--- /dev/null
+++ b/Engine/Plugins/Media/PixelStreaming/Source/PixelStreaming/Private/PixelStreamingAudioDeviceModule.cpp
@@ -0,0 +1,406 @@
+// Copyright Epic Games, Inc. All Rights Reserved.
+
+#include "PixelStreamingAudioDeviceModule.h"
+#include "AudioSubmixCapturer.h"
+#include "AudioPlayoutRequester.h"
+#include "PixelStreamingPrivate.h"
+#include "Settings.h"
+
+// These are copied from webrtc internals
+#define CHECKinitialized_() \
+	{                       \
+		if (!bInitialized)  \
+		{                   \
+			return -1;      \
+		};                  \
+	}
+#define CHECKinitialized__BOOL() \
+	{                            \
+		if (!bInitialized)       \
+		{                        \
+			return false;        \
+		};                       \
+	}
+
+FPixelStreamingAudioDeviceModule::FPixelStreamingAudioDeviceModule()
+	: bInitialized(false)
+	, Capturer(MakeUnique<UE::PixelStreaming::FAudioSubmixCapturer>())
+	, Requester(MakeUnique<UE::PixelStreaming::FAudioPlayoutRequester>())
+{
+}
+
+int32 FPixelStreamingAudioDeviceModule::ActiveAudioLayer(AudioLayer* audioLayer) const
+{
+	*audioLayer = AudioDeviceModule::kDummyAudio;
+	return 0;
+}
+
+int32 FPixelStreamingAudioDeviceModule::RegisterAudioCallback(webrtc::AudioTransport* audioCallback)
+{
+	Requester->RegisterAudioCallback(audioCallback);
+	Capturer->RegisterAudioCallback(audioCallback);
+	return 0;
+}
+
+int32 FPixelStreamingAudioDeviceModule::Init()
+{
+	InitRecording();
+
+	UE_LOG(LogPixelStreaming, Verbose, TEXT("Init PixelStreamingAudioDeviceModule"));
+	bInitialized = true;
+	return 0;
+}
+
+int32 FPixelStreamingAudioDeviceModule::Terminate()
+{
+	if (!bInitialized)
+	{
+		return 0;
+	}
+
+	Requester->Uninitialise();
+	Capturer->Uninitialise();
+
+	UE_LOG(LogPixelStreaming, Verbose, TEXT("Terminate"));
+
+	return 0;
+}
+
+bool FPixelStreamingAudioDeviceModule::Initialized() const
+{
+	return bInitialized;
+}
+
+int16 FPixelStreamingAudioDeviceModule::PlayoutDevices()
+{
+	CHECKinitialized_();
+	return -1;
+}
+
+int16 FPixelStreamingAudioDeviceModule::RecordingDevices()
+{
+	CHECKinitialized_();
+	return -1;
+}
+
+int32 FPixelStreamingAudioDeviceModule::PlayoutDeviceName(uint16 index, char name[webrtc::kAdmMaxDeviceNameSize], char guid[webrtc::kAdmMaxGuidSize])
+{
+	CHECKinitialized_();
+	return -1;
+}
+
+int32 FPixelStreamingAudioDeviceModule::RecordingDeviceName(uint16 index, char name[webrtc::kAdmMaxDeviceNameSize], char guid[webrtc::kAdmMaxGuidSize])
+{
+	CHECKinitialized_();
+	return -1;
+}
+
+int32 FPixelStreamingAudioDeviceModule::SetPlayoutDevice(uint16 index)
+{
+	CHECKinitialized_();
+	return 0;
+}
+
+int32 FPixelStreamingAudioDeviceModule::SetPlayoutDevice(WindowsDeviceType device)
+{
+	CHECKinitialized_();
+	return 0;
+}
+
+int32 FPixelStreamingAudioDeviceModule::SetRecordingDevice(uint16 index)
+{
+	CHECKinitialized_();
+	return 0;
+}
+
+int32 FPixelStreamingAudioDeviceModule::SetRecordingDevice(WindowsDeviceType device)
+{
+	CHECKinitialized_();
+	return 0;
+}
+
+int32 FPixelStreamingAudioDeviceModule::PlayoutIsAvailable(bool* available)
+{
+	CHECKinitialized_();
+	*available = !UE::PixelStreaming::Settings::CVarPixelStreamingWebRTCDisableReceiveAudio.GetValueOnAnyThread();
+	return 0;
+}
+
+int32 FPixelStreamingAudioDeviceModule::InitPlayout()
+{
+	CHECKinitialized_();
+	bool bIsPlayoutAvailable = false;
+	PlayoutIsAvailable(&bIsPlayoutAvailable);
+	if (!bIsPlayoutAvailable)
+	{
+		return -1;
+	}
+
+	Requester->InitPlayout();
+	return 0;
+}
+
+bool FPixelStreamingAudioDeviceModule::PlayoutIsInitialized() const
+{
+	CHECKinitialized__BOOL();
+	return Requester->PlayoutIsInitialized();
+}
+
+int32 FPixelStreamingAudioDeviceModule::StartPlayout()
+{
+	CHECKinitialized_();
+
+	if (!Requester->PlayoutIsInitialized())
+	{
+		return -1;
+	}
+
+	Requester->StartPlayout();
+	return 0;
+}
+
+int32 FPixelStreamingAudioDeviceModule::StopPlayout()
+{
+	CHECKinitialized_();
+	if (!Requester->PlayoutIsInitialized())
+	{
+		return -1;
+	}
+
+	Requester->StopPlayout();
+	return 0;
+}
+
+bool FPixelStreamingAudioDeviceModule::Playing() const
+{
+	CHECKinitialized__BOOL();
+	return Requester->Playing();
+}
+
+int32 FPixelStreamingAudioDeviceModule::RecordingIsAvailable(bool* available)
+{
+	CHECKinitialized_();
+	*available = !UE::PixelStreaming::Settings::CVarPixelStreamingWebRTCDisableTransmitAudio.GetValueOnAnyThread();
+	return 0;
+}
+
+int32 FPixelStreamingAudioDeviceModule::InitRecording()
+{
+	CHECKinitialized_();
+
+	bool bIsRecordingAvailable = false;
+	RecordingIsAvailable(&bIsRecordingAvailable);
+	if (!bIsRecordingAvailable)
+	{
+		return -1;
+	}
+
+	if (!Capturer->IsInitialised())
+	{
+		Capturer->Init();
+	}
+
+	return 0;
+}
+
+bool FPixelStreamingAudioDeviceModule::RecordingIsInitialized() const
+{
+	CHECKinitialized__BOOL();
+	return Capturer->IsInitialised();
+}
+
+int32 FPixelStreamingAudioDeviceModule::StartRecording()
+{
+	CHECKinitialized_();
+	if (!Capturer->IsCapturing())
+	{
+		Capturer->StartCapturing();
+	}
+	return 0;
+}
+
+int32 FPixelStreamingAudioDeviceModule::StopRecording()
+{
+	CHECKinitialized_();
+	if (Capturer->IsCapturing())
+	{
+		Capturer->EndCapturing();
+	}
+	return 0;
+}
+
+bool FPixelStreamingAudioDeviceModule::Recording() const
+{
+	CHECKinitialized__BOOL();
+	return Capturer->IsCapturing();
+}
+
+int32 FPixelStreamingAudioDeviceModule::InitSpeaker()
+{
+	CHECKinitialized_();
+	return -1;
+}
+
+bool FPixelStreamingAudioDeviceModule::SpeakerIsInitialized() const
+{
+	CHECKinitialized__BOOL();
+	return false;
+}
+
+int32 FPixelStreamingAudioDeviceModule::InitMicrophone()
+{
+	CHECKinitialized_();
+	return 0;
+}
+
+bool FPixelStreamingAudioDeviceModule::MicrophoneIsInitialized() const
+{
+	CHECKinitialized__BOOL();
+	return true;
+}
+
+int32 FPixelStreamingAudioDeviceModule::SpeakerVolumeIsAvailable(bool* available)
+{
+	return -1;
+}
+int32 FPixelStreamingAudioDeviceModule::SetSpeakerVolume(uint32 volume)
+{
+	return -1;
+}
+int32 FPixelStreamingAudioDeviceModule::SpeakerVolume(uint32* volume) const
+{
+	return -1;
+}
+int32 FPixelStreamingAudioDeviceModule::MaxSpeakerVolume(uint32* maxVolume) const
+{
+	return -1;
+}
+int32 FPixelStreamingAudioDeviceModule::MinSpeakerVolume(uint32* minVolume) const
+{
+	return -1;
+}
+
+int32 FPixelStreamingAudioDeviceModule::MicrophoneVolumeIsAvailable(bool* available)
+{
+	return 0;
+}
+int32 FPixelStreamingAudioDeviceModule::SetMicrophoneVolume(uint32 volume)
+{
+	// Capturer.SetVolume(volume);
+	return 0;
+}
+int32 FPixelStreamingAudioDeviceModule::MicrophoneVolume(uint32* volume) const
+{
+	//*volume = Capturer.GetVolume();
+	return 0;
+}
+int32 FPixelStreamingAudioDeviceModule::MaxMicrophoneVolume(uint32* maxVolume) const
+{
+	*maxVolume = UE::PixelStreaming::FAudioSubmixCapturer::MaxVolumeLevel;
+	return 0;
+}
+int32 FPixelStreamingAudioDeviceModule::MinMicrophoneVolume(uint32* minVolume) const
+{
+	return 0;
+}
+
+int32 FPixelStreamingAudioDeviceModule::SpeakerMuteIsAvailable(bool* available)
+{
+	return -1;
+}
+int32 FPixelStreamingAudioDeviceModule::SetSpeakerMute(bool enable)
+{
+	return -1;
+}
+int32 FPixelStreamingAudioDeviceModule::SpeakerMute(bool* enabled) const
+{
+	return -1;
+}
+
+int32 FPixelStreamingAudioDeviceModule::MicrophoneMuteIsAvailable(bool* available)
+{
+	*available = false;
+	return -1;
+}
+int32 FPixelStreamingAudioDeviceModule::SetMicrophoneMute(bool enable)
+{
+	return -1;
+}
+int32 FPixelStreamingAudioDeviceModule::MicrophoneMute(bool* enabled) const
+{
+	return -1;
+}
+
+int32 FPixelStreamingAudioDeviceModule::StereoPlayoutIsAvailable(bool* available) const
+{
+	CHECKinitialized_();
+	*available = true;
+	return 0;
+}
+
+int32 FPixelStreamingAudioDeviceModule::SetStereoPlayout(bool enable)
+{
+	CHECKinitialized_();
+	FString AudioChannelStr = enable ? TEXT("stereo") : TEXT("mono");
+	UE_LOG(LogPixelStreaming, Verbose, TEXT("WebRTC has requested browser audio playout in UE be: %s"), *AudioChannelStr);
+	return 0;
+}
+
+int32 FPixelStreamingAudioDeviceModule::StereoPlayout(bool* enabled) const
+{
+	CHECKinitialized_();
+	*enabled = true;
+	return 0;
+}
+
+int32 FPixelStreamingAudioDeviceModule::StereoRecordingIsAvailable(bool* available) const
+{
+	CHECKinitialized_();
+	*available = true;
+	return 0;
+}
+
+int32 FPixelStreamingAudioDeviceModule::SetStereoRecording(bool enable)
+{
+	CHECKinitialized_();
+	return 0;
+}
+
+int32 FPixelStreamingAudioDeviceModule::StereoRecording(bool* enabled) const
+{
+	CHECKinitialized_();
+	*enabled = true;
+	return 0;
+}
+
+int32 FPixelStreamingAudioDeviceModule::PlayoutDelay(uint16* delayMS) const
+{
+	*delayMS = 0;
+	return 0;
+}
+
+bool FPixelStreamingAudioDeviceModule::BuiltInAECIsAvailable() const
+{
+	return false;
+}
+bool FPixelStreamingAudioDeviceModule::BuiltInAGCIsAvailable() const
+{
+	return false;
+}
+bool FPixelStreamingAudioDeviceModule::BuiltInNSIsAvailable() const
+{
+	return false;
+}
+
+int32 FPixelStreamingAudioDeviceModule::EnableBuiltInAEC(bool enable)
+{
+	return -1;
+}
+int32 FPixelStreamingAudioDeviceModule::EnableBuiltInAGC(bool enable)
+{
+	return -1;
+}
+int32 FPixelStreamingAudioDeviceModule::EnableBuiltInNS(bool enable)
+{
+	return -1;
+}
\ No newline at end of file
diff --git a/Engine/Plugins/Media/PixelStreaming/Source/PixelStreaming/Private/PixelStreamingBlueprints.cpp b/Engine/Plugins/Media/PixelStreaming/Source/PixelStreaming/Private/PixelStreamingBlueprints.cpp
--- a/Engine/Plugins/Media/PixelStreaming/Source/PixelStreaming/Private/PixelStreamingBlueprints.cpp
+++ b/Engine/Plugins/Media/PixelStreaming/Source/PixelStreaming/Private/PixelStreamingBlueprints.cpp
@@ -6,59 +6,186 @@
 #include "PixelStreamingInputComponent.h"
 #include "PixelStreamingAudioComponent.h"
 #include "PixelStreamingPlayerId.h"
+#include "IPixelStreamingStreamer.h"
 
 void UPixelStreamingBlueprints::SendFileAsByteArray(TArray<uint8> ByteArray, FString MimeType, FString FileExtension)
 {
 	IPixelStreamingModule* Module = UE::PixelStreaming::FPixelStreamingModule::GetModule();
-	if (Module)
+	if(!Module)
 	{
-		Module->SendFileData(ByteArray, MimeType, FileExtension);
+		return;
 	}
+	TSharedPtr<IPixelStreamingStreamer> Streamer = Module->GetStreamer(Module->GetDefaultStreamerID());
+	if(!Streamer)
+	{
+		return;
+	}
+	Streamer->SendFileData(ByteArray, MimeType, FileExtension);
+}
+
+void UPixelStreamingBlueprints::StreamerSendFileAsByteArray(FString StreamerId, TArray<uint8> ByteArray, FString MimeType, FString FileExtension)
+{
+	IPixelStreamingModule* Module = UE::PixelStreaming::FPixelStreamingModule::GetModule();
+	if(!Module)
+	{
+		return;
+	}
+	TSharedPtr<IPixelStreamingStreamer> Streamer = Module->GetStreamer(StreamerId);
+	if(!Streamer)
+	{
+		return;
+	}
+	Streamer->SendFileData(ByteArray, MimeType, FileExtension);
 }
 
 void UPixelStreamingBlueprints::SendFile(FString FilePath, FString MimeType, FString FileExtension)
 {
 	IPixelStreamingModule* Module = UE::PixelStreaming::FPixelStreamingModule::GetModule();
-	if (Module)
-	{
-		TArray<uint8> ByteData;
-		bool bSuccess = FFileHelper::LoadFileToArray(ByteData, *FilePath);
-		if (bSuccess)
-		{
-			Module->SendFileData(ByteData, MimeType, FileExtension);
-		}
-		else
-		{
-			UE_LOG(LogPixelStreaming, Error, TEXT("FileHelper failed to load file data"));
-		}
+	if(!Module)
+	{
+		return;
+	}
+	TSharedPtr<IPixelStreamingStreamer> Streamer = Module->GetStreamer(Module->GetDefaultStreamerID());
+	if(!Streamer)
+	{
+		return;
+	}
+	
+	TArray<uint8> ByteData;
+	bool bSuccess = FFileHelper::LoadFileToArray(ByteData, *FilePath);
+	if(bSuccess)
+	{
+		Streamer->SendFileData(ByteData, MimeType, FileExtension);
+	}
+	else
+	{
+		UE_LOG(LogPixelStreaming, Error, TEXT("FileHelper failed to load file data"));
+	}	
+}
+
+void UPixelStreamingBlueprints::StreamerSendFile(FString StreamerId, FString FilePath, FString MimeType, FString FileExtension)
+{
+	IPixelStreamingModule* Module = UE::PixelStreaming::FPixelStreamingModule::GetModule();
+	if(!Module)
+	{
+		return;
+	}
+	TSharedPtr<IPixelStreamingStreamer> Streamer = Module->GetStreamer(StreamerId);
+	if(!Streamer)
+	{
+		return;
 	}
+	
+	TArray<uint8> ByteData;
+	bool bSuccess = FFileHelper::LoadFileToArray(ByteData, *FilePath);
+	if(bSuccess)
+	{
+		Streamer->SendFileData(ByteData, MimeType, FileExtension);
+	}
+	else
+	{
+		UE_LOG(LogPixelStreaming, Error, TEXT("FileHelper failed to load file data"));
+	}	
 }
 
 void UPixelStreamingBlueprints::FreezeFrame(UTexture2D* Texture)
 {
 	IPixelStreamingModule* Module = UE::PixelStreaming::FPixelStreamingModule::GetModule();
-	if (Module)
+	if(!Module)
+	{
+		return;
+	}
+	TSharedPtr<IPixelStreamingStreamer> Streamer = Module->GetStreamer(Module->GetDefaultStreamerID());
+	if(!Streamer)
 	{
-		Module->FreezeFrame(Texture);
+		return;
 	}
+	Streamer->FreezeStream(Texture);	
+}
+
+void UPixelStreamingBlueprints::StreamerFreezeStream(FString StreamerId, UTexture2D* Texture)
+{
+	IPixelStreamingModule* Module = UE::PixelStreaming::FPixelStreamingModule::GetModule();
+	if(!Module)
+	{
+		return;
+	}
+	TSharedPtr<IPixelStreamingStreamer> Streamer = Module->GetStreamer(StreamerId);
+	if(!Streamer)
+	{
+		return;
+	}
+	Streamer->FreezeStream(Texture);	
 }
 
 void UPixelStreamingBlueprints::UnfreezeFrame()
 {
 	IPixelStreamingModule* Module = UE::PixelStreaming::FPixelStreamingModule::GetModule();
-	if (Module)
+	if(!Module)
+	{
+		return;
+	}
+	TSharedPtr<IPixelStreamingStreamer> Streamer = Module->GetStreamer(Module->GetDefaultStreamerID());
+	if(!Streamer)
+	{
+		return;
+	}
+	Streamer->UnfreezeStream();
+}
+
+void UPixelStreamingBlueprints::StreamerUnfreezeStream(FString StreamerId)
+{
+	IPixelStreamingModule* Module = UE::PixelStreaming::FPixelStreamingModule::GetModule();
+	if(!Module)
+	{
+		return;
+	}
+	TSharedPtr<IPixelStreamingStreamer> Streamer = Module->GetStreamer(StreamerId);
+	if(!Streamer)
 	{
-		Module->UnfreezeFrame();
+		return;
 	}
+	Streamer->UnfreezeStream();
 }
 
 void UPixelStreamingBlueprints::KickPlayer(FString PlayerId)
 {
 	IPixelStreamingModule* Module = UE::PixelStreaming::FPixelStreamingModule::GetModule();
-	if (Module)
+	if(!Module)
+	{
+		return;
+	}
+	TSharedPtr<IPixelStreamingStreamer> Streamer = Module->GetStreamer(Module->GetDefaultStreamerID());
+	if(!Streamer)
+	{
+		return;
+	}
+	Streamer->KickPlayer(ToPlayerId(PlayerId));
+}
+
+void UPixelStreamingBlueprints::StreamerKickPlayer(FString StreamerId, FString PlayerId)
+{
+	IPixelStreamingModule* Module = UE::PixelStreaming::FPixelStreamingModule::GetModule();
+	if(!Module)
+	{
+		return;
+	}
+	TSharedPtr<IPixelStreamingStreamer> Streamer = Module->GetStreamer(StreamerId);
+	if(!Streamer)
+	{
+		return;
+	}
+	Streamer->KickPlayer(ToPlayerId(PlayerId));
+}
+
+FString UPixelStreamingBlueprints::GetDefaultStreamerID()
+{
+	IPixelStreamingModule* Module = UE::PixelStreaming::FPixelStreamingModule::GetModule();
+	if(!Module)
 	{
-		Module->KickPlayer(ToPlayerId(PlayerId));
+		FString();
 	}
+	return Module->GetDefaultStreamerID();
 }
 
 UPixelStreamingDelegates* UPixelStreamingBlueprints::GetPixelStreamingDelegates()
diff --git a/Engine/Plugins/Media/PixelStreaming/Source/PixelStreaming/Private/PixelStreamingBlueprints.h b/Engine/Plugins/Media/PixelStreaming/Source/PixelStreaming/Private/PixelStreamingBlueprints.h
--- a/Engine/Plugins/Media/PixelStreaming/Source/PixelStreaming/Private/PixelStreamingBlueprints.h
+++ b/Engine/Plugins/Media/PixelStreaming/Source/PixelStreaming/Private/PixelStreamingBlueprints.h
@@ -22,9 +22,20 @@ public:
      * @param   MimeType        The mime type of the file. Used for reconstruction on the front end
      * @param   FileExtension   The file extension. Used for file reconstruction on the front end
      */
-	UFUNCTION(BlueprintCallable, Category = "Pixel Streaming Transmit")
+	UFUNCTION(BlueprintCallable, Category = "Pixel Streaming")
 	static void SendFileAsByteArray(TArray<uint8> ByteArray, FString MimeType, FString FileExtension);
 
+	/**
+     * Send a specified byte array over the WebRTC peer connection data channel. The extension and mime type are used for file reconstruction on the front end 
+     * 
+	 * @param	StreamerId		The streamer use when sending the data
+     * @param   ByteArray       The raw data that will be sent over the data channel
+     * @param   MimeType        The mime type of the file. Used for reconstruction on the front end
+     * @param   FileExtension   The file extension. Used for file reconstruction on the front end
+     */
+	UFUNCTION(BlueprintCallable, Category = "Pixel Streaming")
+	static void StreamerSendFileAsByteArray(FString StreamerId, TArray<uint8> ByteArray, FString MimeType, FString FileExtension);
+
 	/**
      * Send a specified file over the WebRTC peer connection data channel. The extension and mime type are used for file reconstruction on the front end 
      * 
@@ -32,28 +43,56 @@ public:
      * @param   MimeType        The mime type of the file. Used for file reconstruction on the front end
      * @param   FileExtension   The file extension. Used for file reconstruction on the front end
      */
-	UFUNCTION(BlueprintCallable, Category = "Pixel Streaming Transmit")
+	UFUNCTION(BlueprintCallable, Category = "Pixel Streaming")
 	static void SendFile(FString Filepath, FString MimeType, FString FileExtension);
 
+
+	/**
+     * Send a specified file over the WebRTC peer connection data channel. The extension and mime type are used for file reconstruction on the front end 
+     * 
+	 * @param	StreamerId		The streamer use when sending the data
+     * @param   FilePath        The path to the file that will be sent
+     * @param   MimeType        The mime type of the file. Used for file reconstruction on the front end
+     * @param   FileExtension   The file extension. Used for file reconstruction on the front end
+     */
+	UFUNCTION(BlueprintCallable, Category = "Pixel Streaming")
+	static void StreamerSendFile(FString StreamerId, FString Filepath, FString MimeType, FString FileExtension);
+
 	/**
     * The functions allow Pixel Streaming to be frozen and unfrozen from
     * Blueprint. When frozen, a freeze frame (a still image) will be used by the
     * browser instead of the video stream.
     */
 
-	/**
+   /**
 	 * Freeze Pixel Streaming.
 	 * @param   Texture         The freeze frame to display. If null then the back buffer is captured.
 	 */
-	UFUNCTION(BlueprintCallable, Category = "Pixel Streaming Freeze Frame")
+	UFUNCTION(BlueprintCallable, Category = "Pixel Streaming")
 	static void FreezeFrame(UTexture2D* Texture);
 
 	/**
 	 * Unfreeze Pixel Streaming. 
 	 */
-	UFUNCTION(BlueprintCallable, Category = "Pixel Streaming Freeze Frame")
+	UFUNCTION(BlueprintCallable, Category = "Pixel Streaming")
 	static void UnfreezeFrame();
 
+
+	/**
+	 * Freeze Pixel Streaming.
+	 * @param	StreamerId		The id of the streamer to freeze.
+	 * @param   Texture         The freeze frame to display. If null then the back buffer is captured.
+	 */
+	UFUNCTION(BlueprintCallable, Category = "Pixel Streaming")
+	static void StreamerFreezeStream(FString StreamerId, UTexture2D* Texture);
+
+	/**
+	 * Unfreeze Pixel Streaming. 
+	 * @param StreamerId		The id of the streamer to unfreeze.
+	 */
+	UFUNCTION(BlueprintCallable, Category = "Pixel Streaming")
+	static void StreamerUnfreezeStream(FString StreamerId);
+
 	/**
 	 * Kick a player.
 	 * @param   PlayerId         The ID of the player to kick.
@@ -61,11 +100,25 @@ public:
 	UFUNCTION(BlueprintCallable, Category = "Pixel Streaming")
 	static void KickPlayer(FString PlayerId);
 
+	/**
+	 * Kick a player.
+	 * @param	StreamerId		The streamer which the player belongs
+	 * @param   PlayerId        The ID of the player to kick.
+	 */
+	UFUNCTION(BlueprintCallable, Category = "Pixel Streaming")
+	static void StreamerKickPlayer(FString StreamerId, FString PlayerId);
+
+	/**
+	 * Get the default Streamer ID
+	 */
+	UFUNCTION(BlueprintCallable, Category = "Pixel Streaming")
+	static FString GetDefaultStreamerID();
+
 	// PixelStreamingDelegates
 	/**
 	 * Get the singleton. This allows application-specific blueprints to bind
 	 * to delegates of interest.
 	 */
-	UFUNCTION(BlueprintCallable, Category = "Pixel Streaming Delegates")
+	UFUNCTION(BlueprintCallable, Category = "Pixel Streaming")
 	static UPixelStreamingDelegates* GetPixelStreamingDelegates();
 };
\ No newline at end of file
diff --git a/Engine/Plugins/Media/PixelStreaming/Source/PixelStreaming/Private/PixelStreamingFrameAdapter.cpp b/Engine/Plugins/Media/PixelStreaming/Source/PixelStreaming/Private/PixelStreamingFrameAdapter.cpp
new file mode 100644
--- /dev/null
+++ b/Engine/Plugins/Media/PixelStreaming/Source/PixelStreaming/Private/PixelStreamingFrameAdapter.cpp
@@ -0,0 +1,88 @@
+// Copyright Epic Games, Inc. All Rights Reserved.
+
+#include "PixelStreamingFrameAdapter.h"
+#include "PixelStreamingFrameAdapterProcess.h"
+#include "PixelStreamingVideoInput.h"
+
+FPixelStreamingFrameAdapter::FPixelStreamingFrameAdapter(TSharedPtr<FPixelStreamingVideoInput> VideoInput)
+	: VideoInputPtr(VideoInput)
+{
+	LayerScales.Add(1.0f);
+	OnFrameDelegateHandle = VideoInput->OnFrame.AddRaw(this, &FPixelStreamingFrameAdapter::OnFrame);
+}
+
+FPixelStreamingFrameAdapter::FPixelStreamingFrameAdapter(TSharedPtr<FPixelStreamingVideoInput> VideoInput, TArray<float> InLayerScales)
+	: VideoInputPtr(VideoInput)
+	, LayerScales(InLayerScales)
+{
+	OnFrameDelegateHandle = VideoInput->OnFrame.AddRaw(this, &FPixelStreamingFrameAdapter::OnFrame);
+}
+
+FPixelStreamingFrameAdapter::~FPixelStreamingFrameAdapter()
+{
+	if (TSharedPtr<FPixelStreamingVideoInput> VideoInput = VideoInputPtr.Pin())
+	{
+		VideoInput->OnFrame.Remove(OnFrameDelegateHandle);
+	}
+}
+
+bool FPixelStreamingFrameAdapter::IsReady() const
+{
+	FScopeLock LayersLock(&LayersGuard);
+
+	for (auto& LayerAdapter : LayerAdapters)
+	{
+		if (!LayerAdapter->HasOutput())
+		{
+			return false;
+		}
+	}
+	return !LayerAdapters.IsEmpty();
+}
+
+int32 FPixelStreamingFrameAdapter::GetWidth(int LayerIndex) const
+{
+	FScopeLock LayersLock(&LayersGuard);
+	return LayerAdapters[LayerIndex]->GetOutputLayerWidth();
+}
+
+int32 FPixelStreamingFrameAdapter::GetHeight(int LayerIndex) const
+{
+	FScopeLock LayersLock(&LayersGuard);
+	return LayerAdapters[LayerIndex]->GetOutputLayerHeight();
+}
+
+TSharedPtr<IPixelStreamingAdaptedVideoFrameLayer> FPixelStreamingFrameAdapter::ReadOutput(int32 LayerIndex)
+{
+	FScopeLock LayersLock(&LayersGuard);
+	return LayerAdapters[LayerIndex]->ReadOutput();
+}
+
+void FPixelStreamingFrameAdapter::AddLayer(float Scale)
+{
+	LayerAdapters.Add(CreateAdaptProcess(Scale));
+}
+
+void FPixelStreamingFrameAdapter::OnFrame(const FPixelStreamingSourceFrame& SourceFrame)
+{
+	FScopeLock LayersLock(&LayersGuard);
+
+	// initial setup. Should only be called on the first frame.
+	// We do this lazily because CreateAdaptProcess is a pure virtual function
+	// so we cannot call it in our constructor. Another option would be to
+	// require the user to call a SetupLayers method or something but that
+	// could be prone to errors.
+	if (LayerAdapters.IsEmpty())
+	{
+		for (auto& Scale : LayerScales)
+		{
+			AddLayer(Scale);
+		}
+	}
+
+	// adapt the frame for encoder use
+	for (auto& LayerAdapter : LayerAdapters)
+	{
+		LayerAdapter->Process(SourceFrame);
+	}
+}
diff --git a/Engine/Plugins/Media/PixelStreaming/Source/PixelStreaming/Private/PixelStreamingFrameAdapterProcess.cpp b/Engine/Plugins/Media/PixelStreaming/Source/PixelStreaming/Private/PixelStreamingFrameAdapterProcess.cpp
new file mode 100644
--- /dev/null
+++ b/Engine/Plugins/Media/PixelStreaming/Source/PixelStreaming/Private/PixelStreamingFrameAdapterProcess.cpp
@@ -0,0 +1,79 @@
+// Copyright Epic Games, Inc. All Rights Reserved.
+
+#include "PixelStreamingFrameAdapterProcess.h"
+#include "Stats.h"
+#include "PixelStreamingFrameAdapter.h"
+#include "PixelStreamingSourceFrame.h"
+#include "Stats.h"
+
+TSharedPtr<IPixelStreamingAdaptedVideoFrameLayer> FPixelStreamingFrameAdapterProcess::ReadOutput()
+{
+	return Buffer.SwapAndRead();
+}
+
+void FPixelStreamingFrameAdapterProcess::Process(const FPixelStreamingSourceFrame& SourceFrame)
+{
+	if (IsBusy())
+	{
+		return;
+	}
+
+	bBusy = true;
+
+	const int32 SourceWidth = SourceFrame.GetWidth();
+	const int32 SourceHeight = SourceFrame.GetHeight();
+	if (!IsInitialized())
+	{
+		Initialize(SourceWidth, SourceHeight);
+	}
+	else if (ResolutionChanged(SourceWidth, SourceHeight))
+	{
+		OnSourceResolutionChanged(ExpectedSourceWidth, ExpectedSourceHeight, SourceWidth, SourceHeight);
+		Initialize(SourceWidth, SourceHeight);
+	}
+
+	BeginProcess(SourceFrame);
+}
+
+int32 FPixelStreamingFrameAdapterProcess::GetOutputLayerWidth()
+{
+	checkf(HasOutput(), TEXT("No output"));
+	return Buffer.Read()->GetWidth();
+}
+
+int32 FPixelStreamingFrameAdapterProcess::GetOutputLayerHeight()
+{
+	checkf(HasOutput(), TEXT("No output"));
+	return Buffer.Read()->GetHeight();
+}
+
+void FPixelStreamingFrameAdapterProcess::Initialize(int32 SourceWidth, int32 SourceHeight)
+{
+	Buffer.Write(CreateOutputBuffer(SourceWidth, SourceHeight));
+	Buffer.SwapWriteBuffers();
+	Buffer.Write(CreateOutputBuffer(SourceWidth, SourceHeight));
+	Buffer.SwapReadBuffers();
+	Buffer.SwapWriteBuffers();
+	Buffer.Write(CreateOutputBuffer(SourceWidth, SourceHeight));
+	ExpectedSourceWidth = SourceWidth;
+	ExpectedSourceHeight = SourceHeight;
+	bHasOutput = false;
+	bInitialized = true;
+}
+
+TSharedPtr<IPixelStreamingAdaptedVideoFrameLayer> FPixelStreamingFrameAdapterProcess::GetWriteBuffer()
+{
+	return Buffer.GetWriteBuffer();
+}
+
+void FPixelStreamingFrameAdapterProcess::EndProcess()
+{
+	Buffer.SwapWriteBuffers();
+	bBusy = false;
+	bHasOutput = true;
+}
+
+bool FPixelStreamingFrameAdapterProcess::ResolutionChanged(int32 SourceWidth, int32 SourceHeight) const
+{
+	return ExpectedSourceWidth != SourceWidth || ExpectedSourceHeight != SourceHeight;
+}
diff --git a/Engine/Plugins/Media/PixelStreaming/Source/PixelStreaming/Private/PixelStreamingInputComponent.cpp b/Engine/Plugins/Media/PixelStreaming/Source/PixelStreaming/Private/PixelStreamingInputComponent.cpp
--- a/Engine/Plugins/Media/PixelStreaming/Source/PixelStreaming/Private/PixelStreamingInputComponent.cpp
+++ b/Engine/Plugins/Media/PixelStreaming/Source/PixelStreaming/Private/PixelStreamingInputComponent.cpp
@@ -11,6 +11,8 @@
 #include "Serialization/JsonSerializer.h"
 #include "GameFramework/GameUserSettings.h"
 #include "PixelStreamingPrivate.h"
+#include "PixelStreamingProtocolDefs.h"
+#include "IPixelStreamingStreamer.h"
 
 UPixelStreamingInput::UPixelStreamingInput(const FObjectInitializer& ObjectInitializer)
 	: Super(ObjectInitializer)
@@ -87,7 +89,10 @@ void UPixelStreamingInput::SendPixelStreamingResponse(const FString& Descriptor)
 {
 	if (PixelStreamingModule)
 	{
-		PixelStreamingModule->SendResponse(Descriptor);
+		PixelStreamingModule->ForEachStreamer([&Descriptor](TSharedPtr<IPixelStreamingStreamer> Streamer)
+		{
+			Streamer->SendPlayerMessage(UE::PixelStreaming::Protocol::EToPlayerMsg::Response, Descriptor);
+		});
 	}
 	else
 	{
diff --git a/Engine/Plugins/Media/PixelStreaming/Source/PixelStreaming/Private/PixelStreamingModule.cpp b/Engine/Plugins/Media/PixelStreaming/Source/PixelStreaming/Private/PixelStreamingModule.cpp
--- a/Engine/Plugins/Media/PixelStreaming/Source/PixelStreaming/Private/PixelStreamingModule.cpp
+++ b/Engine/Plugins/Media/PixelStreaming/Source/PixelStreaming/Private/PixelStreamingModule.cpp
@@ -5,7 +5,7 @@
 #include "InputDevice.h"
 #include "PixelStreamingInputComponent.h"
 #include "PixelStreamingDelegates.h"
-#include "SignallingServerConnection.h"
+#include "PixelStreamingSignallingConnection.h"
 #include "Settings.h"
 #include "PixelStreamingPrivate.h"
 #include "PlayerSession.h"
@@ -16,19 +16,20 @@
 #include "Engine/Texture2D.h"
 #include "Slate/SceneViewport.h"
 #include "Utils.h"
+#include "UtilsRender.h"
 
 #if PLATFORM_WINDOWS || PLATFORM_XBOXONE
-	#include "Windows/WindowsHWrapper.h"
+#include "Windows/WindowsHWrapper.h"
 #elif PLATFORM_LINUX
-	#include "CudaModule.h"
+#include "CudaModule.h"
 #endif
 
 #if PLATFORM_WINDOWS
-	#include "Windows/AllowWindowsPlatformTypes.h"
+#include "Windows/AllowWindowsPlatformTypes.h"
 THIRD_PARTY_INCLUDES_START
-	#include <VersionHelpers.h>
+#include <VersionHelpers.h>
 THIRD_PARTY_INCLUDES_END
-	#include "Windows/HideWindowsPlatformTypes.h"
+#include "Windows/HideWindowsPlatformTypes.h"
 #endif
 
 #include "RenderingThread.h"
@@ -42,172 +43,95 @@ THIRD_PARTY_INCLUDES_END
 #include "Dom/JsonObject.h"
 #include "Misc/App.h"
 #include "Misc/MessageDialog.h"
-#include "IImageWrapper.h"
-#include "IImageWrapperModule.h"
 #include "Async/Async.h"
 #include "Engine/Engine.h"
-#include "EncoderFactory.h"
-#include "VideoSource.h"
+#include "VideoEncoderFactory.h"
+#include "VideoEncoderFactorySimple.h"
+#include "WebRTCLogging.h"
+#include "WebSocketsModule.h"
 
 #if !UE_BUILD_SHIPPING
-	#include "DrawDebugHelpers.h"
+#include "DrawDebugHelpers.h"
 #endif
 
+#include "VideoInputBackBuffer.h"
+#include "VideoSourceGroup.h"
+#include "PixelStreamingPeerConnection.h"
+
 DEFINE_LOG_CATEGORY(LogPixelStreaming);
 
 IPixelStreamingModule* UE::PixelStreaming::FPixelStreamingModule::PixelStreamingModule = nullptr;
 
-namespace
-{
-
-#if PLATFORM_WINDOWS || PLATFORM_XBOXONE
-	// required for WMF video decoding
-	// some Windows versions don't have Media Foundation preinstalled. We configure MF DLLs as delay-loaded and load them manually here
-	// checking the result and avoiding error message box if failed
-	bool LoadMediaFoundationDLLs()
-	{
-		// Ensure that all required modules are preloaded so they are not loaded just-in-time, causing a hitch.
-		if (IsWindows8OrGreater())
-		{
-			return FPlatformProcess::GetDllHandle(TEXT("mf.dll"))
-				&& FPlatformProcess::GetDllHandle(TEXT("mfplat.dll"))
-				&& FPlatformProcess::GetDllHandle(TEXT("msmpeg2vdec.dll"))
-				&& FPlatformProcess::GetDllHandle(TEXT("MSAudDecMFT.dll"));
-		}
-		else // Windows 7
-		{
-			return FPlatformProcess::GetDllHandle(TEXT("mf.dll"))
-				&& FPlatformProcess::GetDllHandle(TEXT("mfplat.dll"))
-				&& FPlatformProcess::GetDllHandle(TEXT("msmpeg2vdec.dll"))
-				&& FPlatformProcess::GetDllHandle(TEXT("msmpeg2adec.dll"));
-		}
-	}
-#endif
-} // namespace
-
 namespace UE::PixelStreaming
 {
-	void FPixelStreamingModule::InitStreamer()
-	{
-		FString StreamerId;
-		FParse::Value(FCommandLine::Get(), TEXT("PixelStreamingID="), StreamerId);
-		UE_LOG(LogPixelStreaming, Log, TEXT("PixelStreaming endpoint ID: %s"), *StreamerId);
-
-		// Check to see if we can use the Pixel Streaming plugin on this platform.
-		// If not then we avoid setting up our delegates to prevent access to the
-		// plugin. Note that Pixel Streaming is not currently performed in the
-		// Editor.
-		if (!IsPlatformCompatible())
-		{
-			return;
-		}
-
-		if (!ensure(GEngine != nullptr))
-		{
-			return;
-		}
-
-		Poller = MakeUnique<FPoller>();
-		FramePump = MakeUnique<FFixedFPSPump>();
-		bDecoupleFrameRate = UE::PixelStreaming::Settings::CVarPixelStreamingDecoupleFrameRate.GetValueOnAnyThread();
-
-		// subscribe to engine delegates here for init / framebuffer creation / whatever
-		// TODO check if there is a better callback to attach so that we can use with editor
-		if (FSlateApplication::IsInitialized())
-		{
-			FSlateApplication::Get().GetRenderer()->OnBackBufferReadyToPresent().AddRaw(this, &FPixelStreamingModule::OnBackBufferReady_RenderThread);
-		}
-
-		IModularFeatures::Get().RegisterModularFeature(GetModularFeatureName(), this);
-
-		FApp::SetUnfocusedVolumeMultiplier(1.0f);
-
-		// Allow Pixel Streaming to broadcast to various delegates bound in the
-		// application-specific blueprint.
-		UPixelStreamingDelegates::CreateInstance();
-
-		verify(FModuleManager::Get().LoadModule(FName("ImageWrapper")));
-
-		Streamer = MakeUnique<FStreamer>(StreamerId);
-
-		// Streamer has been created, so module is now "ready" for external use.
-		ReadyEvent.Broadcast(*this);
-
-		FString SignallingServerURL;
-		if (!FParse::Value(FCommandLine::Get(), TEXT("PixelStreamingURL="), SignallingServerURL))
-		{
-			// didnt get the startup URL for pixel streaming. Check deprecated options...
-			FString SignallingServerIP;
-			uint16 SignallingServerPort;
-			if (FParse::Value(FCommandLine::Get(), TEXT("PixelStreamingIP="), SignallingServerIP)
-				&& FParse::Value(FCommandLine::Get(), TEXT("PixelStreamingPort="), SignallingServerPort))
-			{
-				// got both old parameters. Warn about deprecation and build the proper url.
-				UE_LOG(LogPixelStreaming, Warning, TEXT("PixelStreamingIP and PixelStreamingPort are deprecated flags. Use PixelStreamingURL instead. eg. -PixelStreamingURL=ws://%s:%d"), *SignallingServerIP, SignallingServerPort);
-				SignallingServerURL = FString::Printf(TEXT("ws://%s:%d"), *SignallingServerIP, SignallingServerPort);
-			}
-		}
-
-		if (!SignallingServerURL.IsEmpty())
-		{
-			// have a startup url. dont start in editor though.
-			if (GIsEditor)
-			{
-				FText TitleText = FText::FromString(TEXT("Pixel Streaming Plugin"));
-				FString ErrorString = TEXT("Pixel Streaming Plugin is not supported in editor, but it was explicitly enabled by command-line arguments. Please remove `PixelStreamingURL` or `PixelStreamingIP` and `PixelStreamingPort` args from editor command line.");
-				FText ErrorText = FText::FromString(ErrorString);
-				FMessageDialog::Open(EAppMsgType::Ok, ErrorText, &TitleText);
-				UE_LOG(LogPixelStreaming, Error, TEXT("%s"), *ErrorString);
-			}
-			else
-			{
-				StartStreaming(SignallingServerURL);
-			}
-		}
-	}
-
-	/** IModuleInterface implementation */
+	/** 
+	 * IModuleInterface implementation 
+	 */
 	void FPixelStreamingModule::StartupModule()
 	{
+		// Initialise all settings from command line args etc
+		Settings::InitialiseSettings();
+
 		// Pixel Streaming does not make sense without an RHI so we don't run in commandlets without one.
 		if (IsRunningCommandlet() && !IsAllowCommandletRendering())
 		{
 			return;
 		}
 
-		// Initialise all settings from command line args etc
-		Settings::InitialiseSettings();
-
 		const ERHIInterfaceType RHIType = GDynamicRHI ? RHIGetInterfaceType() : ERHIInterfaceType::Hidden;
 
-		// only D3D11/D3D12 is supported
+		// only D3D11/D3D12/Vulkan is supported
 		if (RHIType == ERHIInterfaceType::D3D11 || RHIType == ERHIInterfaceType::D3D12 || RHIType == ERHIInterfaceType::Vulkan)
 		{
-			// By calling InitStreamer post engine init we can use pixel streaming in standalone editor mode
-			FCoreDelegates::OnFEngineLoopInitComplete.AddRaw(this, &FPixelStreamingModule::InitStreamer);
+			// By calling InitDefaultStreamer post engine init we can use pixel streaming in standalone editor mode
+			FCoreDelegates::OnFEngineLoopInitComplete.AddLambda([this]()
+			{
+				InitDefaultStreamer();
+				bModuleReady = true;
+				ReadyEvent.Broadcast(*this);
+			});
 		}
 		else
 		{
 			UE_LOG(LogPixelStreaming, Warning, TEXT("Only D3D11/D3D12/Vulkan Dynamic RHI is supported. Detected %s"), GDynamicRHI != nullptr ? GDynamicRHI->GetName() : TEXT("[null]"));
 		}
 
-		TextureSourceFactory = MakeUnique<FTextureSourceFactory>();
+		rtc::InitializeSSL();
+		RedirectWebRtcLogsToUnreal(rtc::LoggingSeverity::LS_VERBOSE);
+		FModuleManager::LoadModuleChecked<IModuleInterface>(TEXT("AVEncoder"));
+		FModuleManager::LoadModuleChecked<FWebSocketsModule>("WebSockets");
+
+		// ExternalVideoSourceGroup is used so that we can have a video source without a streamer
+		ExternalVideoSourceGroup = MakeUnique<FVideoSourceGroup>();
+		ExternalVideoSourceGroup->SetVideoInput(MakeShared<FVideoInputBackBuffer>());
+		ExternalVideoSourceGroup->Start();
 	}
 
 	void FPixelStreamingModule::ShutdownModule()
 	{
-		if (FSlateApplication::IsInitialized())
+		// Pixel Streaming does not make sense without an RHI so we don't run in commandlets without one.
+		if (IsRunningCommandlet() && !IsAllowCommandletRendering())
 		{
-			FSlateApplication::Get().GetRenderer()->OnBackBufferReadyToPresent().RemoveAll(this);
-			FSlateApplication::Get().GetRenderer()->OnPreResizeWindowBackBuffer().RemoveAll(this);
+			return;
 		}
 
-		IModularFeatures::Get().UnregisterModularFeature(GetModularFeatureName(), this);
-
 		// We explicitly call release on streamer so WebRTC gets shutdown before our module is deleted
-		Streamer.Reset(nullptr);
+		Streamers.Empty();
+		ExternalVideoSourceGroup->Stop();
+
+		FPixelStreamingPeerConnection::Shutdown();
+
+		rtc::CleanupSSL();
 	}
+	/** 
+	 * End IModuleInterface implementation 
+	 */
 
+
+
+	/** 
+	 * IPixelStreamingModule implementation
+	 */
 	IPixelStreamingModule* FPixelStreamingModule::GetModule()
 	{
 		if (PixelStreamingModule)
@@ -222,125 +146,104 @@ namespace UE::PixelStreaming
 		return PixelStreamingModule;
 	}
 
-	bool FPixelStreamingModule::StartStreaming(const FString& SignallingServerUrl)
+	IPixelStreamingModule::FReadyEvent& FPixelStreamingModule::OnReady()
 	{
-		if (Streamer && IsReady())
-		{
-			return Streamer->StartStreaming(SignallingServerUrl);
-		}
-		return false;
+		return ReadyEvent;
 	}
 
-	void FPixelStreamingModule::StopStreaming()
+	bool FPixelStreamingModule::IsReady()
 	{
-		if (Streamer)
-		{
-			Streamer->StopStreaming();
-		}
+		return bModuleReady;
 	}
 
-	bool FPixelStreamingModule::IsPlatformCompatible() const
+	bool FPixelStreamingModule::StartStreaming()
 	{
-		bool bCompatible = true;
-
-#if PLATFORM_WINDOWS || PLATFORM_XBOXONE
-		bool bWin8OrHigher = IsWindows8OrGreater();
-		if (!bWin8OrHigher)
+		bool bSuccess = true;
+		TMap<FString, TSharedPtr<IPixelStreamingStreamer>>::TIterator Iter = Streamers.CreateIterator();
+		for (; Iter; ++Iter)
 		{
-			FString ErrorString(TEXT("Failed to initialize Pixel Streaming plugin because minimum requirement is Windows 8"));
-			FText ErrorText = FText::FromString(ErrorString);
-			FText TitleText = FText::FromString(TEXT("Pixel Streaming Plugin"));
-			FMessageDialog::Open(EAppMsgType::Ok, ErrorText, &TitleText);
-			UE_LOG(LogPixelStreaming, Error, TEXT("%s"), *ErrorString);
-			bCompatible = false;
-		}
-#endif
+			TSharedPtr<IPixelStreamingStreamer> Streamer = Iter.Value();
 
-		if (!FStreamer::IsPlatformCompatible())
-		{
-			FText TitleText = FText::FromString(TEXT("Pixel Streaming Plugin"));
-			FString ErrorString = TEXT("No compatible GPU found, or failed to load their respective encoder libraries");
-			FText ErrorText = FText::FromString(ErrorString);
-			FMessageDialog::Open(EAppMsgType::Ok, ErrorText, &TitleText);
-			UE_LOG(LogPixelStreaming, Error, TEXT("%s"), *ErrorString);
-			bCompatible = false;
+			if (Streamer.IsValid())
+			{
+				Streamer->SetStreamFPS(Settings::CVarPixelStreamingWebRTCFps.GetValueOnAnyThread());
+				Streamer->StartStreaming();
+				bSuccess &= true;
+				continue;
+			}
+			bSuccess &= false;
 		}
-
-		return bCompatible;
+		return bSuccess;
 	}
 
-	void FPixelStreamingModule::UpdateViewport(FSceneViewport* Viewport)
-	{
-		FRHIViewport* const ViewportRHI = Viewport->GetViewportRHI().GetReference();
-	}
-
-	void FPixelStreamingModule::OnBackBufferReady_RenderThread(SWindow& SlateWindow, const FTexture2DRHIRef& BackBuffer)
+	void FPixelStreamingModule::StopStreaming()
 	{
-		// enable streaming explicitly by providing `PixelStreamingIP` and `PixelStreamingPort` cmd-args
-		if (!Streamer)
+		TMap<FString, TSharedPtr<IPixelStreamingStreamer>>::TIterator Iter = Streamers.CreateIterator();
+		for (; Iter; ++Iter)
 		{
-			return;
-		}
+			TSharedPtr<IPixelStreamingStreamer> Streamer = Iter.Value();
 
-		check(IsInRenderingThread());
+			if (Streamer.IsValid())
+			{
+				Streamer->StopStreaming();
+			}
+		}
+	}
 
-		// Framerate isn't decoupled, manually pump
-		if (!bDecoupleFrameRate)
+	TSharedPtr<IPixelStreamingStreamer> FPixelStreamingModule::CreateStreamer(const FString& StreamerId)
+	{
+		TSharedPtr<IPixelStreamingStreamer> ExistingStreamer = GetStreamer(StreamerId);
+		if(ExistingStreamer)
 		{
-			FramePump->PumpAll();
+			return ExistingStreamer;
 		}
 
-		// Check to see if we have been instructed to capture the back buffer as a
-		// freeze frame.
-		if (bCaptureNextBackBufferAndStream && Streamer->IsStreaming())
+		TSharedPtr<FStreamer> NewStreamer = MakeShared<FStreamer>(StreamerId);
 		{
-			bCaptureNextBackBufferAndStream = false;
-
-			// Read the data out of the back buffer and send as a JPEG.
-			FRHICommandListImmediate& RHICmdList = FRHICommandListExecutor::GetImmediateCommandList();
-			FIntRect Rect(0, 0, BackBuffer->GetDesc().Extent.X, BackBuffer->GetDesc().Extent.Y);
-			TArray<FColor> Data;
-
-			RHICmdList.ReadSurfaceData(BackBuffer, Rect, Data, FReadSurfaceDataFlags());
-			SendJpeg(MoveTemp(Data), Rect);
+			FScopeLock Lock(&StreamersCS);
+			Streamers.Add(StreamerId, NewStreamer);
 		}
+		return NewStreamer;
 	}
 
-	TSharedPtr<class IInputDevice> FPixelStreamingModule::CreateInputDevice(const TSharedRef<FGenericApplicationMessageHandler>& InMessageHandler)
-	{
-		FInputDevice* InputDevicePtr = new FInputDevice(InMessageHandler);
-		InputDevice = TSharedPtr<FInputDevice>(InputDevicePtr);
-		return InputDevice;
-	}
-
-	IPixelStreamingModule::FReadyEvent& FPixelStreamingModule::OnReady()
-	{
-		return ReadyEvent;
-	}
-
-	IPixelStreamingModule::FStreamingStartedEvent& FPixelStreamingModule::OnStreamingStarted()
+	TArray<FString> FPixelStreamingModule::GetStreamerIds()
 	{
-		return StreamingStartedEvent;
+		TArray<FString> StreamerKeys;
+		FScopeLock Lock(&StreamersCS);
+		Streamers.GenerateKeyArray(StreamerKeys);
+		return StreamerKeys;
 	}
 
-	IPixelStreamingModule::FStreamingStoppedEvent& FPixelStreamingModule::OnStreamingStopped()
+	TSharedPtr<IPixelStreamingStreamer> FPixelStreamingModule::GetStreamer(const FString& StreamerId)
 	{
-		return StreamingStoppedEvent;
+		FScopeLock Lock(&StreamersCS);
+		if (Streamers.Contains(StreamerId))
+		{
+			return Streamers[StreamerId];
+		}
+		return nullptr;
 	}
 
-	bool FPixelStreamingModule::IsReady()
+	TSharedPtr<IPixelStreamingStreamer> FPixelStreamingModule::DeleteStreamer(const FString& StreamerId)
 	{
-		return Streamer.IsValid();
+		TSharedPtr<IPixelStreamingStreamer> ToBeDeleted;
+		FScopeLock Lock(&StreamersCS);
+		if (Streamers.Contains(StreamerId))
+		{
+			ToBeDeleted = Streamers[StreamerId];
+			Streamers.Remove(StreamerId);
+		}
+		return ToBeDeleted;
 	}
 
-	IInputDevice& FPixelStreamingModule::GetInputDevice()
+	rtc::scoped_refptr<webrtc::VideoTrackSourceInterface> FPixelStreamingModule::CreateExternalVideoSource()
 	{
-		return *InputDevice;
+		return ExternalVideoSourceGroup->CreateVideoSource([]() { return true; });
 	}
 
-	TSharedPtr<FInputDevice> FPixelStreamingModule::GetInputDevicePtr()
+	void FPixelStreamingModule::ReleaseExternalVideoSource(const webrtc::VideoTrackSourceInterface* InVideoSource)
 	{
-		return InputDevice;
+		ExternalVideoSourceGroup->RemoveVideoSource(InVideoSource);
 	}
 
 	void FPixelStreamingModule::AddInputComponent(UPixelStreamingInput* InInputComponent)
@@ -358,223 +261,129 @@ namespace UE::PixelStreaming
 		return InputComponents;
 	}
 
-	void FPixelStreamingModule::FreezeFrame(UTexture2D* Texture)
+	webrtc::VideoEncoderFactory* FPixelStreamingModule::CreateVideoEncoderFactory()
 	{
-		if (Texture)
-		{
-			ENQUEUE_RENDER_COMMAND(ReadSurfaceCommand)
-			([this, Texture](FRHICommandListImmediate& RHICmdList) {
-				// A frame is supplied so immediately read its data and send as a JPEG.
-				FTextureRHIRef TextureRHI = Texture->GetResource() ? Texture->GetResource()->TextureRHI : nullptr;
-				if (!TextureRHI)
-				{
-					UE_LOG(LogPixelStreaming, Error, TEXT("Attempting freeze frame with texture %s with no texture RHI"), *Texture->GetName());
-					return;
-				}
-				uint32 Width = TextureRHI->GetDesc().Extent.X;
-				uint32 Height = TextureRHI->GetDesc().Extent.Y;
-
-				FTextureRHIRef DestTexture = CreateRHITexture(Width, Height);
-
-				FGPUFenceRHIRef CopyFence = GDynamicRHI->RHICreateGPUFence(*FString::Printf(TEXT("FreezeFrameFence")));
-
-				// Copy freeze frame texture to empty texture
-				CopyTextureToRHI(TextureRHI, DestTexture, CopyFence);
-
-				TArray<FColor> Data;
-				FIntRect Rect(0, 0, Width, Height);
-				RHICmdList.ReadSurfaceData(DestTexture, Rect, Data, FReadSurfaceDataFlags());
-				SendJpeg(MoveTemp(Data), Rect);
-			});
-		}
-		else
-		{
-			// A frame is not supplied, so we need to capture the back buffer at
-			// the next opportunity, and send as a JPEG.
-			bCaptureNextBackBufferAndStream = true;
-		}
-
-		// Stop streaming.
-		bFrozen = true;
+		return new FVideoEncoderFactorySimple();
 	}
 
-	void FPixelStreamingModule::UnfreezeFrame()
+	FString FPixelStreamingModule::GetDefaultStreamerID()
 	{
-		if (!Streamer.IsValid())
-		{
-			return;
-		}
-
-		Streamer->SendUnfreezeFrame();
-
-		// Resume streaming.
-		bFrozen = false;
+		return Settings::GetDefaultStreamerID();
 	}
-	void FPixelStreamingModule::AddPlayerConfig(TSharedRef<FJsonObject>& JsonObject)
+	
+	void FPixelStreamingModule::ForEachStreamer(const TFunction<void(TSharedPtr<IPixelStreamingStreamer>)>& Func)
 	{
-		checkf(InputDevice.IsValid(), TEXT("No Input Device available when populating Player Config"));
-
-		JsonObject->SetBoolField(TEXT("FakingTouchEvents"), InputDevice->IsFakingTouchEvents());
-
-		FString PixelStreamingControlScheme;
-		if (Settings::GetControlScheme(PixelStreamingControlScheme))
+		TSet<FString> KeySet;
 		{
-			JsonObject->SetStringField(TEXT("ControlScheme"), PixelStreamingControlScheme);
+			FScopeLock Lock(&StreamersCS);
+			Streamers.GetKeys(KeySet);
 		}
-
-		float PixelStreamingFastPan;
-		if (Settings::GetFastPan(PixelStreamingFastPan))
+		for(auto&& StreamerId : KeySet)
 		{
-			JsonObject->SetNumberField(TEXT("FastPan"), PixelStreamingFastPan);
+			if(TSharedPtr<IPixelStreamingStreamer> Streamer = GetStreamer(StreamerId))
+			{
+				Func(Streamer);
+			}
 		}
 	}
+	/** 
+	 * End IPixelStreamingModule implementation
+	 */
 
-	void FPixelStreamingModule::SendResponse(const FString& Descriptor)
+	void FPixelStreamingModule::InitDefaultStreamer()
 	{
-		if (!Streamer.IsValid())
+		UE_LOG(LogPixelStreaming, Log, TEXT("PixelStreaming endpoint ID: %s"), *Settings::GetDefaultStreamerID());
+
+		// Check to see if we can use the Pixel Streaming plugin on this platform.
+		// If not then we avoid setting up our delegates to prevent access to the
+		// plugin. Note that Pixel Streaming is not currently performed in the
+		// Editor.
+		if (!IsPlatformCompatible())
 		{
 			return;
 		}
 
-		Streamer->SendPlayerMessage(Protocol::EToPlayerMsg::Response, Descriptor);
-	}
-
-	void FPixelStreamingModule::SendCommand(const FString& Descriptor)
-	{
-		if (!Streamer.IsValid())
+		if (!ensure(GEngine != nullptr))
 		{
 			return;
 		}
 
-		Streamer->SendPlayerMessage(Protocol::EToPlayerMsg::Command, Descriptor);
-	}
+		FApp::SetUnfocusedVolumeMultiplier(1.0f);
+		// Allow Pixel Streaming to broadcast to various delegates bound in the application-specific blueprint.
+		UPixelStreamingDelegates::CreateInstance();
+		verify(FModuleManager::Get().LoadModule(FName("ImageWrapper")));
 
-	void FPixelStreamingModule::SendJpeg(TArray<FColor> RawData, const FIntRect& Rect)
-	{
-		if (!Streamer.IsValid())
-		{
-			return;
-		}
 
-		IImageWrapperModule& ImageWrapperModule = FModuleManager::GetModuleChecked<IImageWrapperModule>(TEXT("ImageWrapper"));
-		TSharedPtr<IImageWrapper> ImageWrapper = ImageWrapperModule.CreateImageWrapper(EImageFormat::JPEG);
-		bool bSuccess = ImageWrapper->SetRaw(RawData.GetData(), RawData.Num() * sizeof(FColor), Rect.Width(), Rect.Height(), ERGBFormat::BGRA, 8);
-		if (bSuccess)
-		{
-			// Compress to a JPEG of the maximum possible quality.
-			int32 Quality = Settings::CVarPixelStreamingFreezeFrameQuality.GetValueOnAnyThread();
-			const TArray64<uint8>& JpegBytes = ImageWrapper->GetCompressed(Quality);
-			Streamer->SendFreezeFrame(JpegBytes);
-		}
-		else
+		FString SignallingServerURL;
+		if (!FParse::Value(FCommandLine::Get(), TEXT("PixelStreamingURL="), SignallingServerURL))
 		{
-			UE_LOG(LogPixelStreaming, Error, TEXT("JPEG image wrapper failed to accept frame data"));
+			// didnt get the startup URL for pixel streaming. Check deprecated options...
+			FString SignallingServerIP;
+			uint16 SignallingServerPort;
+			if (FParse::Value(FCommandLine::Get(), TEXT("PixelStreamingIP="), SignallingServerIP)
+				&& FParse::Value(FCommandLine::Get(), TEXT("PixelStreamingPort="), SignallingServerPort))
+			{
+				// got both old parameters. Warn about deprecation and build the proper url.
+				UE_LOG(LogPixelStreaming, Warning, TEXT("PixelStreamingIP and PixelStreamingPort are deprecated flags. Use PixelStreamingURL instead. eg. -PixelStreamingURL=ws://%s:%d"), *SignallingServerIP, SignallingServerPort);
+				SignallingServerURL = FString::Printf(TEXT("ws://%s:%d"), *SignallingServerIP, SignallingServerPort);
+			}
 		}
-	}
 
-	void FPixelStreamingModule::KickPlayer(FPixelStreamingPlayerId PlayerId)
-	{
-		if (Streamer)
+		TSharedPtr<IPixelStreamingStreamer> Streamer = CreateStreamer(Settings::GetDefaultStreamerID());
+		Streamer->SetVideoInput(MakeShared<FVideoInputBackBuffer>());
+
+		if (!SignallingServerURL.IsEmpty())
 		{
-			Streamer->KickPlayer(PlayerId);
+			// have a startup url. dont start in editor though.
+			if (GIsEditor)
+			{
+				FText TitleText = FText::FromString(TEXT("Pixel Streaming Plugin"));
+				FString ErrorString = TEXT("Pixel Streaming Plugin is not supported in editor, but it was explicitly enabled by command-line arguments. Please remove `PixelStreamingURL` or `PixelStreamingIP` and `PixelStreamingPort` args from editor command line.");
+				FText ErrorText = FText::FromString(ErrorString);
+				FMessageDialog::Open(EAppMsgType::Ok, ErrorText, &TitleText);
+				UE_LOG(LogPixelStreaming, Error, TEXT("%s"), *ErrorString);
+			}
+			else
+			{
+				Streamer->SetSignallingServerURL(SignallingServerURL);
+				Streamer->StartStreaming();
+			}
 		}
 	}
 
-	void FPixelStreamingModule::SendFileData(TArray<uint8>& ByteData, FString& MimeType, FString& FileExtension)
-	{
-		Streamer->SendFileData(ByteData, MimeType, FileExtension);
-	}
-
-	bool FPixelStreamingModule::IsTickableWhenPaused() const
-	{
-		return true;
-	}
-
-	bool FPixelStreamingModule::IsTickableInEditor() const
-	{
-		return true;
-	}
-
-	void FPixelStreamingModule::Tick(float DeltaTime)
-	{
-	}
-
-	TStatId FPixelStreamingModule::GetStatId() const
+	bool FPixelStreamingModule::IsPlatformCompatible() const
 	{
-		RETURN_QUICK_DECLARE_CYCLE_STAT(FPixelStreamingModule, STATGROUP_Tickables);
-	}
+		bool bCompatible = true;
 
-	IPixelStreamingAudioSink* FPixelStreamingModule::GetPeerAudioSink(FPixelStreamingPlayerId PlayerId)
-	{
-		if (!Streamer.IsValid())
+#if PLATFORM_WINDOWS || PLATFORM_XBOXONE
+		bool bWin8OrHigher = IsWindows8OrGreater();
+		if (!bWin8OrHigher)
 		{
-			UE_LOG(LogPixelStreaming, Error, TEXT("Cannot get audio sink when streamer does not yet exist."));
-			return nullptr;
+			FString ErrorString(TEXT("Failed to initialize Pixel Streaming plugin because minimum requirement is Windows 8"));
+			FText ErrorText = FText::FromString(ErrorString);
+			FText TitleText = FText::FromString(TEXT("Pixel Streaming Plugin"));
+			FMessageDialog::Open(EAppMsgType::Ok, ErrorText, &TitleText);
+			UE_LOG(LogPixelStreaming, Error, TEXT("%s"), *ErrorString);
+			bCompatible = false;
 		}
+#endif
 
-		return Streamer->GetPlayerSessions().ForSession<IPixelStreamingAudioSink*>(PlayerId, [](TSharedPtr<IPlayerSession> Session) { return Session->GetAudioSink(); });
-	}
-
-	IPixelStreamingAudioSink* FPixelStreamingModule::GetUnlistenedAudioSink()
-	{
-		if (!Streamer.IsValid())
+		if (Settings::CVarPixelStreamingEncoderCodec.GetValueOnAnyThread() == "H264"
+			&& !AVEncoder::FVideoEncoderFactory::Get().HasEncoderForCodec(AVEncoder::ECodecType::H264))
 		{
-			UE_LOG(LogPixelStreaming, Error, TEXT("Cannot get audio sink when streamer does not yet exist."));
-			return nullptr;
+			FText TitleText = FText::FromString(TEXT("Pixel Streaming Plugin"));
+			FString ErrorString = TEXT("No compatible GPU found, or failed to load their respective encoder libraries");
+			FText ErrorText = FText::FromString(ErrorString);
+			FMessageDialog::Open(EAppMsgType::Ok, ErrorText, &TitleText);
+			UE_LOG(LogPixelStreaming, Error, TEXT("%s"), *ErrorString);
+			bCompatible = false;
 		}
 
-		IPixelStreamingAudioSink* Result = nullptr;
-		Streamer->GetPlayerSessions().ForEachSession([&Result](TSharedPtr<IPlayerSession> Session) {
-			if (!Result && !Session->GetAudioSink()->HasAudioConsumers())
-			{
-				Result = Session->GetAudioSink();
-			}
-		});
-
-		return Result;
-	}
-
-	void FPixelStreamingModule::AddPollerTask(TFunction<void()> Task, TFunction<bool()> IsTaskFinished, TSharedRef<bool, ESPMode::ThreadSafe> bKeepRunning)
-	{
-		Poller->AddJob(IsTaskFinished, bKeepRunning, Task);
-	}
-
-	void FPixelStreamingModule::RegisterPumpable(rtc::scoped_refptr<FPixelStreamingPumpable> Pumpable)
-	{
-		FramePump->RegisterPumpable(Pumpable);
-	}
-
-	void FPixelStreamingModule::UnregisterPumpable(rtc::scoped_refptr<FPixelStreamingPumpable> Pumpable)
-	{
-		FramePump->UnregisterPumpable(Pumpable);
-	}
-
-	webrtc::VideoEncoderFactory* FPixelStreamingModule::CreateVideoEncoderFactory()
-	{
-		return new FVideoEncoderFactory();
-	}
-
-	rtc::scoped_refptr<webrtc::VideoTrackSourceInterface> FPixelStreamingModule::CreateExternalVideoSource(FName SourceType)
-	{
-		FVideoSourceP2P* VideoSource = new FVideoSourceP2P(SourceType, []() { return true; });
-		VideoSource->Initialize();
-		return rtc::scoped_refptr<webrtc::VideoTrackSourceInterface>(VideoSource);
-	}
-
-	IPixelStreamingTextureSourceFactory& FPixelStreamingModule::GetTextureSourceFactory()
-	{
-		return *TextureSourceFactory;
-	}
-
-	void FPixelStreamingModule::SetActiveTextureSourceTypes(const TArray<FName>& SourceTypes)
-	{
-		UE::PixelStreaming::Settings::SetActiveTextureSourceTypes(SourceTypes);
-	}
-
-	const TArray<FName>& FPixelStreamingModule::GetActiveTextureSourceTypes() const
-	{
-		return UE::PixelStreaming::Settings::GetActiveTextureSourceTypes();
+		return bCompatible;
 	}
+	/**
+	 * End own methods
+	 */
 } // namespace UE::PixelStreaming
 
 IMPLEMENT_MODULE(UE::PixelStreaming::FPixelStreamingModule, PixelStreaming)
diff --git a/Engine/Plugins/Media/PixelStreaming/Source/PixelStreaming/Private/PixelStreamingModule.h b/Engine/Plugins/Media/PixelStreaming/Source/PixelStreaming/Private/PixelStreamingModule.h
--- a/Engine/Plugins/Media/PixelStreaming/Source/PixelStreaming/Private/PixelStreamingModule.h
+++ b/Engine/Plugins/Media/PixelStreaming/Source/PixelStreaming/Private/PixelStreamingModule.h
@@ -6,105 +6,66 @@
 #include "RHI.h"
 #include "Tickable.h"
 #include "InputDevice.h"
-#include "Poller.h"
-#include "FixedFPSPump.h"
-#include "PixelStreamingPumpable.h"
-#include "TextureSourceFactory.h"
 
-class AController;
-class AGameModeBase;
-class APlayerController;
-class FSceneViewport;
+
 class UPixelStreamingInput;
 class SWindow;
 
 namespace UE::PixelStreaming
 {
 	class FStreamer;
+	class FVideoInputBackBuffer;
+	class FVideoSourceGroup;
 
 	/*
 	 * This plugin allows the back buffer to be sent as a compressed video across a network.
 	 */
-	class FPixelStreamingModule : public IPixelStreamingModule, public FTickableGameObject
+	class FPixelStreamingModule : public IPixelStreamingModule
 	{
 	public:
 		static IPixelStreamingModule* GetModule();
 
-		virtual bool StartStreaming(const FString& SignallingServerUrl) override;
+		/** IPixelStreamingModule implementation */
+		virtual FReadyEvent& OnReady() override;
+		virtual bool IsReady() override;
+		virtual bool StartStreaming() override;
 		virtual void StopStreaming() override;
+		virtual TSharedPtr<IPixelStreamingStreamer> CreateStreamer(const FString& StreamerId) override;
+		virtual TArray<FString> GetStreamerIds() override;
+		virtual TSharedPtr<IPixelStreamingStreamer> GetStreamer(const FString& StreamerId) override;
+		virtual TSharedPtr<IPixelStreamingStreamer> DeleteStreamer(const FString& StreamerId) override;
+		virtual FString GetDefaultStreamerID() override;
+		// These are staying on the module at the moment as theres no way of the BPs knowing which streamer they are relevant to
+		virtual void AddInputComponent(UPixelStreamingInput* InInputComponent) override;
+		virtual void RemoveInputComponent(UPixelStreamingInput* InInputComponent) override;
+		virtual const TArray<UPixelStreamingInput*> GetInputComponents() override;
+		// Don't delete, is used externally to PS
+		virtual rtc::scoped_refptr<webrtc::VideoTrackSourceInterface> CreateExternalVideoSource() override;
+		virtual void ReleaseExternalVideoSource(const webrtc::VideoTrackSourceInterface* InVideoSource) override;
+		virtual webrtc::VideoEncoderFactory* CreateVideoEncoderFactory() override;
+		virtual void ForEachStreamer(const TFunction<void(TSharedPtr<IPixelStreamingStreamer>)>& Func) override;
+		/** End IPixelStreamingModule implementation */
 
 	private:
 		/** IModuleInterface implementation */
 		void StartupModule() override;
 		void ShutdownModule() override;
+		/** End IModuleInterface implementation */
 
-		TSharedPtr<class IInputDevice> CreateInputDevice(const TSharedRef<FGenericApplicationMessageHandler>& InMessageHandler) override;
-
-		/** IPixelStreamingModule implementation */
-		FReadyEvent& OnReady() override;
-		FStreamingStartedEvent& OnStreamingStarted() override;
-		FStreamingStoppedEvent& OnStreamingStopped() override;
-		bool IsReady() override;
-		IInputDevice& GetInputDevice() override;
-		void AddPlayerConfig(TSharedRef<FJsonObject>& JsonObject) override;
-		void SendResponse(const FString& Descriptor) override;
-		void SendCommand(const FString& Descriptor) override;
-
-		/*
-		 * Returns a shared pointer to the device which handles pixel streaming
-		 * input.
-		 * @return The shared pointer to the input device.
-		 */
-		TSharedPtr<FInputDevice> GetInputDevicePtr();
-		void AddInputComponent(UPixelStreamingInput* InInputComponent) override;
-		void RemoveInputComponent(UPixelStreamingInput* InInputComponent) override;
-		const TArray<UPixelStreamingInput*> GetInputComponents() override;
-
-		void FreezeFrame(UTexture2D* Texture) override;
-		void UnfreezeFrame() override;
-		void KickPlayer(FPixelStreamingPlayerId PlayerId);
-		IPixelStreamingAudioSink* GetPeerAudioSink(FPixelStreamingPlayerId PlayerId) override;
-		IPixelStreamingAudioSink* GetUnlistenedAudioSink() override;
-		virtual rtc::scoped_refptr<webrtc::VideoTrackSourceInterface> CreateExternalVideoSource(FName SourceType) override;
-		virtual webrtc::VideoEncoderFactory* CreateVideoEncoderFactory() override;
-		virtual void RegisterPumpable(rtc::scoped_refptr<FPixelStreamingPumpable> Pumpable) override;
-		virtual void UnregisterPumpable(rtc::scoped_refptr<FPixelStreamingPumpable> Pumpable) override;
-		virtual void AddPollerTask(TFunction<void()> Task, TFunction<bool()> IsTaskFinished, TSharedRef<bool, ESPMode::ThreadSafe> bKeepRunning) override;
-		virtual IPixelStreamingTextureSourceFactory& GetTextureSourceFactory() override;
-		virtual void SetActiveTextureSourceTypes(const TArray<FName>& SourceTypes) override;
-		virtual const TArray<FName>& GetActiveTextureSourceTypes() const override;
-		/** End IPixelStreamingModule implementation */
-
-		// FTickableGameObject
-		bool IsTickableWhenPaused() const override;
-		bool IsTickableInEditor() const override;
-		void Tick(float DeltaTime) override;
-		TStatId GetStatId() const override;
-
+		// Own methods
+		void InitDefaultStreamer();
 		bool IsPlatformCompatible() const;
-		void UpdateViewport(FSceneViewport* Viewport);
-		void OnBackBufferReady_RenderThread(SWindow& SlateWindow, const FTexture2DRHIRef& BackBuffer);
-		void SendJpeg(TArray<FColor> RawData, const FIntRect& Rect);
-		void SendFileData(TArray<uint8>& ByteData, FString& MimeType, FString& FileExtension);
-
-		void InitStreamer();
-
 	private:
-		FReadyEvent ReadyEvent;
-		FStreamingStartedEvent StreamingStartedEvent;
-		FStreamingStoppedEvent StreamingStoppedEvent;
-		TUniquePtr<FStreamer> Streamer;
-		TSharedPtr<FInputDevice> InputDevice;
-		TArray<UPixelStreamingInput*> InputComponents;
-		bool bFrozen = false;
-		bool bCaptureNextBackBufferAndStream = false;
-		double LastVideoEncoderQPReportTime = 0;
+		bool bModuleReady = false;
 		static IPixelStreamingModule* PixelStreamingModule;
 
-		TUniquePtr<FFixedFPSPump> FramePump;
-		TUniquePtr<FPoller> Poller;
-		TUniquePtr<IPixelStreamingTextureSourceFactory> TextureSourceFactory;
+		FReadyEvent ReadyEvent;
+		
+		TArray<UPixelStreamingInput*> InputComponents;
+		TUniquePtr<FVideoSourceGroup> ExternalVideoSourceGroup;
+		
 
-		bool bDecoupleFrameRate = false;
+		mutable FCriticalSection StreamersCS;
+		TMap<FString, TSharedPtr<IPixelStreamingStreamer>> Streamers;
 	};
 } // namespace UE::PixelStreaming
