Early Pixel Streaming 5.1 features, bugfixes, and refactors that benefit main. (From various authors)

diff --git a/Engine/Plugins/Media/PixelStreaming/Source/PixelStreaming/Private/PixelStreamingPeerConnection.cpp b/Engine/Plugins/Media/PixelStreaming/Source/PixelStreaming/Private/PixelStreamingPeerConnection.cpp
new file mode 100644
--- /dev/null
+++ b/Engine/Plugins/Media/PixelStreaming/Source/PixelStreaming/Private/PixelStreamingPeerConnection.cpp
@@ -0,0 +1,316 @@
+// Copyright Epic Games, Inc. All Rights Reserved.
+
+#include "PixelStreamingPeerConnection.h"
+#include "VideoDecoderFactory.h"
+#include "VideoEncoderFactorySimulcast.h"
+#include "PixelStreamingSessionDescriptionObservers.h"
+#include "PixelStreamingPrivate.h"
+#include "Settings.h"
+#include "AudioCapturer.h"
+#include "PixelStreamingAudioDeviceModule.h"
+#include "ToStringExtensions.h"
+
+namespace
+{
+	inline size_t SerializeToBuffer(rtc::CopyOnWriteBuffer& Buffer, size_t Pos, const void* Data, size_t DataSize)
+	{
+		FMemory::Memcpy(&Buffer[Pos], reinterpret_cast<const uint8_t*>(Data), DataSize);
+		return Pos + DataSize;
+	}
+} // namespace
+
+TUniquePtr<rtc::Thread> FPixelStreamingPeerConnection::SignallingThread = nullptr;
+rtc::scoped_refptr<webrtc::PeerConnectionFactoryInterface> FPixelStreamingPeerConnection::PeerConnectionFactory = nullptr;
+
+TUniquePtr<FPixelStreamingPeerConnection> FPixelStreamingPeerConnection::Create(const FConfig& Config)
+{
+	if (!PeerConnectionFactory)
+	{
+		CreatePeerConnectionFactory();
+	}
+
+	TUniquePtr<FPixelStreamingPeerConnection> NewPeerConnection = TUniquePtr<FPixelStreamingPeerConnection>(new FPixelStreamingPeerConnection());
+	NewPeerConnection->PeerConnection = PeerConnectionFactory->CreatePeerConnection(Config, nullptr, nullptr, NewPeerConnection.Get());
+
+	return NewPeerConnection;
+}
+
+void FPixelStreamingPeerConnection::Shutdown()
+{
+	PeerConnectionFactory = nullptr;
+	if (SignallingThread)
+	{
+		SignallingThread->Stop();
+	}
+	SignallingThread = nullptr;
+}
+
+void FPixelStreamingPeerConnection::SetSuccessCallback(const TFunction<void(const webrtc::SessionDescriptionInterface*)>& Callback)
+{
+	SuccessCallback = Callback;
+}
+
+void FPixelStreamingPeerConnection::SetFailureCallback(const TFunction<void(const FString&)>& Callback)
+{
+	FailureCallback = Callback;
+}
+
+void FPixelStreamingPeerConnection::SetIceCandidateCallback(const TFunction<void(const webrtc::IceCandidateInterface*)>& Callback)
+{
+	IceCandidateCallback = Callback;
+}
+
+void FPixelStreamingPeerConnection::SetRemoteDescription(const FString& Sdp)
+{
+	FPixelStreamingSetSessionDescriptionObserver* SetLocalDescriptionObserver = FPixelStreamingSetSessionDescriptionObserver::Create(
+		[this]() // on success
+		{
+			if (SuccessCallback)
+			{
+				SuccessCallback(PeerConnection->local_description());
+			}
+		}, FailureCallback);
+
+	FPixelStreamingCreateSessionDescriptionObserver* CreateAnswerObserver = FPixelStreamingCreateSessionDescriptionObserver::Create(
+		[this, SetLocalDescriptionObserver](webrtc::SessionDescriptionInterface* SDP) {
+			PeerConnection->SetLocalDescription(SetLocalDescriptionObserver, SDP);
+		}, FailureCallback);
+
+	FPixelStreamingSetSessionDescriptionObserver* SetRemoteDescriptionObserver = FPixelStreamingSetSessionDescriptionObserver::Create(
+		[this, CreateAnswerObserver]() {
+			// Note: these offer to receive at superseded now we are use transceivers to setup our peer connection media
+			int offer_to_receive_video = 1;
+			int offer_to_receive_audio = 1;
+			bool voice_activity_detection = false;
+			bool ice_restart = true;
+			bool use_rtp_mux = true;
+
+			webrtc::PeerConnectionInterface::RTCOfferAnswerOptions AnswerOption{
+				offer_to_receive_video,
+				offer_to_receive_audio,
+				voice_activity_detection,
+				ice_restart,
+				use_rtp_mux
+			};
+
+			PeerConnection->CreateAnswer(CreateAnswerObserver, AnswerOption);
+		}, FailureCallback);
+
+	webrtc::SdpParseError Error;
+	std::unique_ptr<webrtc::SessionDescriptionInterface> SessionDesc = webrtc::CreateSessionDescription(webrtc::SdpType::kOffer, UE::PixelStreaming::ToString(Sdp), &Error);
+	if (!SessionDesc)
+	{
+		if (FailureCallback)
+		{
+			FailureCallback(FString::Printf(TEXT("Failed to create session description: %s"), Error.description.c_str()));
+		}
+	}
+
+	PeerConnection->SetRemoteDescription(SetRemoteDescriptionObserver, SessionDesc.release());
+}
+
+void FPixelStreamingPeerConnection::AddRemoteIceCandidate(const FString& SdpMid, int SdpMLineIndex, const FString& Sdp)
+{
+	webrtc::SdpParseError Error;
+	std::unique_ptr<webrtc::IceCandidateInterface> Candidate(webrtc::CreateIceCandidate(UE::PixelStreaming::ToString(SdpMid), SdpMLineIndex, UE::PixelStreaming::ToString(Sdp), &Error));
+	if (!Candidate)
+	{
+		UE_LOG(LogPixelStreaming, Error, TEXT("Failed to create ICE candicate: %s"), Error.description.c_str());
+		return;
+	}
+
+	PeerConnection->AddIceCandidate(std::move(Candidate), [](webrtc::RTCError error) {
+		if (!error.ok())
+		{
+			UE_LOG(LogPixelStreaming, Error, TEXT("AddIceCandidate failed: %S"), error.message());
+		}
+	});
+}
+
+void FPixelStreamingPeerConnection::CreateDataChannels(int32 SendStreamId, int32 RecvStreamId)
+{
+	// called when we manually negotiate a datachannel connection via the signalling server
+	UE_LOG(LogPixelStreaming, Log, TEXT("OnPeerDataChannels"));
+	webrtc::DataChannelInit SendConfig;
+	SendConfig.negotiated = true;
+	SendConfig.id = SendStreamId;
+	rtc::scoped_refptr<webrtc::DataChannelInterface> SendDataChannel = PeerConnection->CreateDataChannel("datachannel", &SendConfig);
+	DataChannel = SendDataChannel;
+
+	if (SendStreamId != RecvStreamId)
+	{
+		webrtc::DataChannelInit RecvConfig;
+		RecvConfig.negotiated = true;
+		RecvConfig.id = RecvStreamId;
+		rtc::scoped_refptr<webrtc::DataChannelInterface> RecvDataChannel = PeerConnection->CreateDataChannel("datachannel", &RecvConfig);
+		RecvDataChannel->RegisterObserver(this);
+	}
+	else
+	{
+		SendDataChannel->RegisterObserver(this);
+	}
+}
+
+void FPixelStreamingPeerConnection::SetVideoSink(rtc::VideoSinkInterface<webrtc::VideoFrame>* Sink)
+{
+	VideoSink = Sink;
+}
+
+bool FPixelStreamingPeerConnection::SendMessage(UE::PixelStreaming::Protocol::EToStreamerMsg Type, const FString& Descriptor) const
+{
+	if (!DataChannel)
+	{
+		return false;
+	}
+
+	const uint8 MessageType = static_cast<uint8>(Type);
+	const size_t DescriptorSize = Descriptor.Len() * sizeof(TCHAR);
+
+	rtc::CopyOnWriteBuffer MsgBuffer(sizeof(MessageType) + DescriptorSize);
+
+	size_t Pos = 0;
+	Pos = SerializeToBuffer(MsgBuffer, Pos, &MessageType, sizeof(MessageType));
+	Pos = SerializeToBuffer(MsgBuffer, Pos, *Descriptor, DescriptorSize);
+
+	return DataChannel->Send(webrtc::DataBuffer(MsgBuffer, true));
+}
+
+void FPixelStreamingPeerConnection::OnSignalingChange(webrtc::PeerConnectionInterface::SignalingState NewState)
+{
+	UE_LOG(LogPixelStreaming, Log, TEXT("OnSignalingChange (%s)"), UE::PixelStreaming::ToString(NewState));
+}
+
+void FPixelStreamingPeerConnection::OnAddStream(rtc::scoped_refptr<webrtc::MediaStreamInterface> Stream)
+{
+	UE_LOG(LogPixelStreaming, Log, TEXT("OnAddStream"));
+}
+
+void FPixelStreamingPeerConnection::OnRemoveStream(rtc::scoped_refptr<webrtc::MediaStreamInterface> Stream)
+{
+	UE_LOG(LogPixelStreaming, Log, TEXT("OnRemoveStream"));
+}
+
+void FPixelStreamingPeerConnection::OnDataChannel(rtc::scoped_refptr<webrtc::DataChannelInterface> Channel)
+{
+	UE_LOG(LogPixelStreaming, Log, TEXT("OnDataChannel"));
+
+	if (DataChannel)
+	{
+		UE_LOG(LogPixelStreaming, Warning, TEXT("Replacing datachannel"));
+	}
+
+	DataChannel = Channel;
+	DataChannel->RegisterObserver(this);
+
+	// OnDataChannelOpen.Broadcast(*this);
+}
+
+void FPixelStreamingPeerConnection::OnRenegotiationNeeded()
+{
+	UE_LOG(LogPixelStreaming, Log, TEXT("OnRenegotiationNeeded"));
+}
+
+void FPixelStreamingPeerConnection::OnIceConnectionChange(webrtc::PeerConnectionInterface::IceConnectionState NewState)
+{
+	UE_LOG(LogPixelStreaming, Log, TEXT("OnIceConnectionChange (%s)"), UE::PixelStreaming::ToString(NewState));
+
+	if (NewState == webrtc::PeerConnectionInterface::IceConnectionState::kIceConnectionConnected)
+	{
+		//State = EState::ConnectedStreamer;
+		// OnConnected.Broadcast(*this);
+	}
+	else if (NewState == webrtc::PeerConnectionInterface::IceConnectionState::kIceConnectionDisconnected)
+	{
+		//State = EState::Disconnected;
+		// OnDisconnected.Broadcast(*this);
+	}
+}
+
+void FPixelStreamingPeerConnection::OnIceGatheringChange(webrtc::PeerConnectionInterface::IceGatheringState NewState)
+{
+	UE_LOG(LogPixelStreaming, Log, TEXT("OnIceGatheringChange (%s)"), UE::PixelStreaming::ToString(NewState));
+}
+
+void FPixelStreamingPeerConnection::OnIceCandidate(const webrtc::IceCandidateInterface* Candidate)
+{
+	UE_LOG(LogPixelStreaming, Log, TEXT("OnIceCandidate"));
+	if (IceCandidateCallback)
+	{
+		IceCandidateCallback(Candidate);
+	}
+}
+
+void FPixelStreamingPeerConnection::OnIceCandidatesRemoved(const std::vector<cricket::Candidate>& candidates)
+{
+	UE_LOG(LogPixelStreaming, Log, TEXT("OnIceCandidatesRemoved"));
+}
+
+void FPixelStreamingPeerConnection::OnIceConnectionReceivingChange(bool Receiving)
+{
+	UE_LOG(LogPixelStreaming, Log, TEXT("OnIceConnectionReceivingChange"));
+}
+
+void FPixelStreamingPeerConnection::OnTrack(rtc::scoped_refptr<webrtc::RtpTransceiverInterface> Transceiver)
+{
+	UE_LOG(LogPixelStreaming, Log, TEXT("OnTrack"));
+	if (VideoSink && Transceiver->media_type() == cricket::MediaType::MEDIA_TYPE_VIDEO)
+	{
+		webrtc::VideoTrackInterface* VideoTrack = static_cast<webrtc::VideoTrackInterface*>(Transceiver->receiver()->track().get());
+		VideoTrack->AddOrUpdateSink(VideoSink, rtc::VideoSinkWants());
+	}
+	return;
+}
+
+void FPixelStreamingPeerConnection::OnRemoveTrack(rtc::scoped_refptr<webrtc::RtpReceiverInterface> receiver)
+{
+	UE_LOG(LogPixelStreaming, Log, TEXT("OnRemoveTrack"));
+}
+
+void FPixelStreamingPeerConnection::OnStateChange()
+{
+	UE_LOG(LogPixelStreaming, Log, TEXT("OnStateChange"));
+}
+
+void FPixelStreamingPeerConnection::OnMessage(const webrtc::DataBuffer& buffer)
+{
+	UE_LOG(LogPixelStreaming, Log, TEXT("OnMessage"));
+
+	const uint8 MsgType = static_cast<uint8>(buffer.data.data()[0]);
+	const size_t DescriptorSize = (buffer.data.size() - 1) / sizeof(TCHAR);
+	const TCHAR* DescPtr = reinterpret_cast<const TCHAR*>(buffer.data.data() + 1);
+	const FString Descriptor(DescriptorSize, DescPtr);
+	// OnDataMessage.Broadcast(*this, MsgType, Descriptor);
+}
+
+void FPixelStreamingPeerConnection::CreatePeerConnectionFactory()
+{
+	using namespace UE::PixelStreaming;
+
+	SignallingThread = MakeUnique<rtc::Thread>(rtc::SocketServer::CreateDefault());
+	SignallingThread->SetName("FPixelStreamingPeerConnection SignallingThread", nullptr);
+	SignallingThread->Start();
+
+	bool bUseLegacyAudioDeviceModule = Settings::CVarPixelStreamingWebRTCUseLegacyAudioDevice.GetValueOnAnyThread();
+	rtc::scoped_refptr<webrtc::AudioDeviceModule> AudioDeviceModule;
+	if (bUseLegacyAudioDeviceModule)
+	{
+		AudioDeviceModule = new rtc::RefCountedObject<FAudioCapturer>();
+	}
+	else
+	{
+		AudioDeviceModule = new rtc::RefCountedObject<FPixelStreamingAudioDeviceModule>();
+	}
+
+	PeerConnectionFactory = webrtc::CreatePeerConnectionFactory(
+		nullptr,													   // network_thread
+		nullptr,													   // worker_thread
+		SignallingThread.Get(),										   // signaling_thread
+		AudioDeviceModule,											   // default_adm
+		webrtc::CreateAudioEncoderFactory<webrtc::AudioEncoderOpus>(), // audio_encoder_factory
+		webrtc::CreateAudioDecoderFactory<webrtc::AudioDecoderOpus>(), // audio_decoder_factory
+		std::make_unique<FVideoEncoderFactorySimulcast>(),			   // video_encoder_factory
+		std::make_unique<FVideoDecoderFactory>(),					   // video_decoder_factory
+		nullptr,													   // audio_mixer
+		nullptr);													   // audio_processing
+	check(PeerConnectionFactory);
+}
diff --git a/Engine/Plugins/Media/PixelStreaming/Source/PixelStreaming/Private/PixelStreamingServers.cpp b/Engine/Plugins/Media/PixelStreaming/Source/PixelStreaming/Private/PixelStreamingServers.cpp
new file mode 100644
--- /dev/null
+++ b/Engine/Plugins/Media/PixelStreaming/Source/PixelStreaming/Private/PixelStreamingServers.cpp
@@ -0,0 +1,506 @@
+// Copyright Epic Games, Inc. All Rights Reserved.
+
+#include "PixelStreamingServers.h"
+#include "PixelStreamingPrivate.h"
+#include "HAL/PlatformFileManager.h"
+#include "Misc/Paths.h"
+#include "GenericPlatform/GenericPlatformFile.h"
+#include "GenericPlatform/GenericPlatformProcess.h"
+#include "IWebSocket.h"
+#include "HAL/ThreadSafeCounter.h"
+#include "WebSocketsModule.h"
+
+namespace UE::PixelStreaming
+{
+	namespace Servers
+	{
+		static const FString SamplesRelativePath = FString(TEXT("Samples")) / TEXT("PixelStreaming") / TEXT("WebServers");
+		static const FString CirrusName = FString(TEXT("SignallingWebServer"));
+		static const FString TransientServerPrefix = TEXT("Transient");
+		static FThreadSafeCounter TransientServerId = 0;
+
+		/*
+		 * The path to the servers.
+		 * If running in editor this is the Engine's installed directory Samples/
+		 * If running a packaged game, this is the game's packaged directory Samples/
+		 */
+		FString
+		GetServersAbsPath()
+		{
+			return FPaths::ConvertRelativePathToFull(FPaths::RootDir() / SamplesRelativePath);
+		}
+
+		TSharedPtr<FMonitoredProcess> LaunchScriptChildProcess(FString ScriptAbsPath, FString ScriptArgs)
+		{
+			// Check if the binary actually exists
+			IPlatformFile& FileManager = FPlatformFileManager::Get().GetPlatformFile();
+			if (!FileManager.FileExists(*ScriptAbsPath))
+			{
+				UE_LOG(LogPixelStreaming, Error, TEXT("Cannot start child process - the specified file did not exist. File=%s"), *ScriptAbsPath);
+				return TSharedPtr<FMonitoredProcess>();
+			}
+
+			FString ScriptRunner;
+			FString ScriptRunnerParams;
+
+// Get the executable we will use to run the scripts (e.g. cmd.exe on Windows)
+#if PLATFORM_WINDOWS
+			ScriptRunner = TEXT("cmd.exe");
+			ScriptRunnerParams = FString::Printf(TEXT("/c \"\"%s\" %s\""), *ScriptAbsPath, *ScriptArgs);
+#elif PLATFORM_LINUX
+			ScriptRunner = TEXT("/usr/bin/bash");
+			ScriptRunnerParams = FString::Printf(TEXT(" -- \"%s\" %s"), *ScriptAbsPath, *ScriptArgs);
+#else
+			UE_LOG(LogPixelStreaming, Error, TEXT("Unsupported platform for Pixel Streaming scripts."));
+			return TSharedPtr<FMonitoredProcess>();
+#endif
+
+			TSharedPtr<FMonitoredProcess> ChildProcess = MakeShared<FMonitoredProcess>(ScriptRunner, ScriptRunnerParams, true, true);
+
+			// Run the child process
+			UE_LOG(LogPixelStreaming, Log, TEXT("Launch child process - %s %s"), *ScriptRunner, *ScriptRunnerParams);
+			ChildProcess->Launch();
+			return ChildProcess;
+		}
+
+		FString AddPlatformScripts(FString ServerPath)
+		{
+			ServerPath = ServerPath / TEXT("platform_scripts");
+#if PLATFORM_WINDOWS
+			return ServerPath / TEXT("cmd") / TEXT("run_local.bat");
+#elif PLATFORM_LINUX
+			return ServerPath / TEXT("bash") / TEXT("run_local.sh");
+#else
+			UE_LOG(LogPixelStreaming, Error, TEXT("Unsupported platform for Pixel Streaming scripts."));
+			return ServerPath;
+#endif
+		}
+
+		void SelectivelyCopyServerTo(FString SourceAbsPath, FString DestinationAbsPath)
+		{
+			IPlatformFile& PlatformFile = FPlatformFileManager::Get().GetPlatformFile();
+
+			if (!FPaths::DirectoryExists(DestinationAbsPath))
+			{
+				verifyf(PlatformFile.CreateDirectory(*DestinationAbsPath), TEXT("Creating the destination directory for the servers failed"));
+			}
+
+			PlatformFile.IterateDirectory(*SourceAbsPath, [&PlatformFile, &DestinationAbsPath](const TCHAR* Path, bool bIsDir) -> bool {
+				FString ExistingAbsPath = FString(Path);
+
+				if (bIsDir)
+				{
+					FString DirName = FPaths::GetPathLeaf(ExistingAbsPath);
+					// we don't want to copy node_modules, logs, or platform_script
+					if (DirName == TEXT("node_modules") || DirName == TEXT("logs"))
+					{
+						return true;
+					}
+					// special handling for platform_scripts
+					else if (DirName == TEXT("platform_scripts"))
+					{
+						// make "platform_scripts" dir
+						FString TargetPlatformScriptsDir = DestinationAbsPath / DirName;
+						verifyf(PlatformFile.CreateDirectory(*TargetPlatformScriptsDir), TEXT("Creating the platform scripts directory failed."));
+
+						// get the cmd/bash dirs into an array so we can do a file only copy on both dirs
+						const int NSubDirs = 2;
+						FString PlatformScriptSubDirs[NSubDirs];
+						PlatformScriptSubDirs[0] = TEXT("cmd");
+						PlatformScriptSubDirs[1] = TEXT("bash");
+						for (int i = 0; i < NSubDirs; i++)
+						{
+							FString SubDir = PlatformScriptSubDirs[i];
+							FString PlatformScriptNewSubDir = TargetPlatformScriptsDir / SubDir;
+							verifyf(PlatformFile.CreateDirectory(*PlatformScriptNewSubDir), TEXT("Creating the platform scripts sub directory (e.g. /cmd or /bash) failed."));
+							// Get all script files in the platform_scripts/subdir
+							TArray<FString> ScriptFiles;
+							FString ExistingPlatformScriptsSubDir = ExistingAbsPath / SubDir;
+							PlatformFile.FindFiles(ScriptFiles, *ExistingPlatformScriptsSubDir, nullptr);
+							for (FString& ScriptFileAbsPath : ScriptFiles)
+							{
+								FString ScriptFilename = FPaths::GetPathLeaf(ScriptFileAbsPath);
+								FString ToPath = PlatformScriptNewSubDir / ScriptFilename;
+								verifyf(PlatformFile.CopyFile(*ToPath, *ScriptFileAbsPath), TEXT("Could not copy file to: %s from: %s"), *ToPath, *ScriptFileAbsPath);
+							}
+						}
+					}
+					// copy everything else
+					else
+					{
+						FString CopyDestDir = DestinationAbsPath / DirName;
+						verifyf(PlatformFile.CreateDirectory(*CopyDestDir), TEXT("Creating directory inside server failed - specifically %s"), *CopyDestDir);
+						verifyf(PlatformFile.CopyDirectoryTree(*CopyDestDir, *ExistingAbsPath, true), TEXT("Could not copy everything to: %s from: %s"), *CopyDestDir, *ExistingAbsPath);
+					}
+				}
+				// Copy any loose files
+				else
+				{
+					FString Filename = FPaths::GetPathLeaf(ExistingAbsPath);
+					FString ToPath = DestinationAbsPath / Filename;
+					verifyf(PlatformFile.CopyFile(*ToPath, *ExistingAbsPath), TEXT("Could not copy file to: %s from: %s"), *ToPath, *ExistingAbsPath);
+				}
+				return true;
+			});
+		}
+
+		bool CopyServerTo(FString SourceServerAbsPath, FString DestinationAbsPath)
+		{
+
+			if (!FPaths::DirectoryExists(SourceServerAbsPath))
+			{
+				UE_LOG(LogPixelStreaming, Error, TEXT("The path passed to copy the servers from does not exist - this path does not exist %s."), *SourceServerAbsPath);
+				return false;
+			}
+
+			if (!FPaths::DirectoryExists(DestinationAbsPath))
+			{
+				UE_LOG(LogPixelStreaming, Error, TEXT("The path passed to copy the servers to does not exist - this path does not exist %s."), *DestinationAbsPath);
+				return false;
+			}
+
+			// We selectively copy because some of these server folders can have a lot of content in logs, node_modules and can take a while if copied.
+			SelectivelyCopyServerTo(SourceServerAbsPath, DestinationAbsPath);
+			return true;
+		}
+
+		FString CreateTransientServerDirectory(FString ServerName)
+		{
+			// Transient servers are created mostly for test so we store them in the automation directory.
+			IPlatformFile& PlatformFile = FPlatformFileManager::Get().GetPlatformFile();
+
+			FString SavedPath = FPaths::ConvertRelativePathToFull(FPaths::ProjectSavedDir());
+			verifyf(PlatformFile.CreateDirectory(*SavedPath), TEXT("Creating /Saved directory failed."));
+
+			// Put the transient server under Saved/${Prefix}${ServerName}/
+			FString TransientDir = FString::Printf(TEXT("%s%s%d"), *TransientServerPrefix, *ServerName, TransientServerId.Increment());
+			FString TransientServerPathAbs = SavedPath / TransientDir;
+			UE_LOG(LogPixelStreaming, Log, TEXT("Attempting to creating directory: %s"), *TransientServerPathAbs);
+			verifyf(PlatformFile.CreateDirectory(*TransientServerPathAbs), TEXT("Creating %s failed."), *TransientServerPathAbs);
+
+			return TransientServerPathAbs;
+		}
+
+		/*
+		 * Server is ephemeral because we clone an existing server directory, then when server process is
+		 * complete/cancelled we destroy the cloned server directory.
+		 */
+		TSharedPtr<FMonitoredProcess> LaunchEphemeralServer(FString ExistingServerAbsPath, FString Args, FString& OutLaunchAbsPath)
+		{
+			FString ServerName = FPaths::GetPathLeaf(ExistingServerAbsPath);
+			FString NewServerPath = CreateTransientServerDirectory(ServerName);
+			OutLaunchAbsPath = NewServerPath;
+
+			IPlatformFile& PlatformFile = FPlatformFileManager::Get().GetPlatformFile();
+
+			UE_LOG(LogPixelStreaming, Log, TEXT("Copying %s to %s"), *ExistingServerAbsPath, *NewServerPath);
+			CopyServerTo(ExistingServerAbsPath, NewServerPath);
+
+			FString ScriptAbsPath = AddPlatformScripts(NewServerPath);
+			TSharedPtr<FMonitoredProcess> ServerProcess = LaunchScriptChildProcess(ScriptAbsPath, Args);
+
+			// Bind to output so we can capture the output in the log
+			ServerProcess->OnOutput().BindLambda([ServerName](FString Output) {
+				UE_LOG(LogPixelStreaming, Log, TEXT("%s - %s"), *ServerName, *Output);
+			});
+
+			// Bind to the cancel and delete the server folder
+			ServerProcess->OnCanceled().BindLambda([NewServerPath]() {
+				UE_LOG(LogPixelStreaming, Log, TEXT("Server was cancelled, recursively deleting %s"), *NewServerPath);
+				IPlatformFile& PlatformFile = FPlatformFileManager::Get().GetPlatformFile();
+				bool bDeletedServerDir = PlatformFile.DeleteDirectoryRecursively(*NewServerPath);
+				if (!bDeletedServerDir)
+				{
+					UE_LOG(LogPixelStreaming, Error, TEXT("Could not delete server dir %s"), *NewServerPath);
+				}
+			});
+
+			// Bind to the complete and delete the server folder
+			ServerProcess->OnCompleted().BindLambda([NewServerPath](int ExitCode) {
+				UE_LOG(LogPixelStreaming, Log, TEXT("Server process was finished, recursively deleting %s"), *NewServerPath);
+				IPlatformFile& PlatformFile = FPlatformFileManager::Get().GetPlatformFile();
+				bool bDeletedServerDir = PlatformFile.DeleteDirectoryRecursively(*NewServerPath);
+				if (!bDeletedServerDir)
+				{
+					UE_LOG(LogPixelStreaming, Error, TEXT("Could not delete server dir %s"), *NewServerPath);
+				}
+			});
+
+			return ServerProcess;
+		}
+
+		TSharedPtr<FMonitoredProcess> LaunchServer(FString ServerAbsPath, FLaunchArgs LaunchArgs, FString& OutLaunchAbsPath)
+		{
+			/* Launch the server process, either in-place, or by copying to a temp dir */
+
+			if (LaunchArgs.bEphemeral)
+			{
+				return LaunchEphemeralServer(ServerAbsPath, LaunchArgs.ProcessArgs, OutLaunchAbsPath);
+			}
+			else
+			{
+				OutLaunchAbsPath = ServerAbsPath;
+				return LaunchScriptChildProcess(AddPlatformScripts(ServerAbsPath), LaunchArgs.ProcessArgs);
+			}
+
+			/* Note this has not handled waiting for server to become ready, that is server specific and is therefore handled in say `LaunchCirrus` */
+		}
+
+		bool ExtractValueFromArgs(FString ArgsString, FString ArgKey, FString FallbackValue, FString& OutValue)
+		{
+			// Tokenize string in single whitespace " ".
+			TArray<FString> ArgTokens;
+			ArgsString.ParseIntoArray(ArgTokens, TEXT(" "), true);
+
+			for (FString& Token : ArgTokens)
+			{
+				Token.TrimStartAndEndInline();
+
+				if (!Token.StartsWith(ArgKey, ESearchCase::Type::CaseSensitive))
+				{
+					continue;
+				}
+
+				// We have a matching token for our search "key" - split on it.
+				FString RightStr;
+				if (!Token.Split(TEXT("="), nullptr, &RightStr))
+				{
+					continue;
+				}
+
+				OutValue = RightStr;
+				return true;
+			}
+			OutValue = FallbackValue;
+			return false;
+		}
+
+		int NextPort()
+		{
+			// Todo: Checking for in-use ports.
+			return (4000 + TransientServerId.Increment()) % 65535;
+		}
+
+		FString QueryOrSetServerArg(FLaunchArgs& LaunchArgs, FString ArgKey, FString FallbackArgValue)
+		{
+			FString OutValue;
+			bool bExtractedValue = ExtractValueFromArgs(LaunchArgs.ProcessArgs, ArgKey, FallbackArgValue, OutValue);
+
+			// No key was present so we will inject our own.
+			if (!bExtractedValue)
+			{
+				LaunchArgs.ProcessArgs += FString::Printf(TEXT(" %s%s"), *ArgKey, *FallbackArgValue);
+			}
+
+			return OutValue;
+		}
+
+		/* ------------------------------ FWebSocketProbe -------------------------- */
+
+		/* A utility class that tries to establish a websocket connection. */
+		class FWebSocketProbe
+		{
+		private:
+			TSharedRef<IWebSocket> WebSocket;
+			FThreadSafeBool bShouldAttemptReconnect;
+
+		public:
+			FWebSocketProbe(FString Url)
+				: WebSocket(FWebSocketsModule::Get().CreateWebSocket(Url, TEXT("")))
+				, bShouldAttemptReconnect(true)
+			{
+				WebSocket->OnConnectionError().AddLambda([&bShouldAttemptReconnect = bShouldAttemptReconnect](const FString& Error) {
+					UE_LOG(LogPixelStreaming, Log, TEXT("Waiting to connect to the signalling server.... %s"), *Error);
+					bShouldAttemptReconnect = true;
+				});
+			}
+
+			bool Probe()
+			{
+				bool bIsConnected = WebSocket->IsConnected();
+
+				if (!bIsConnected && bShouldAttemptReconnect)
+				{
+					WebSocket->Connect();
+					bShouldAttemptReconnect = false;
+				}
+
+				return bIsConnected;
+			}
+		};
+
+		/* ------------------------------ FServerBase ------------------------------ */
+
+		void FServerBase::Launch(FLaunchArgs& InLaunchArgs)
+		{
+			LaunchArgs = InLaunchArgs;
+			ServerProcess = LaunchImpl(InLaunchArgs, ServerRootAbsPath, Endpoints);
+			bHasLaunched = true;
+			PollingStartedSeconds = FPlatformTime::Seconds();
+		}
+
+		FString FServerBase::GetPathOnDisk()
+		{
+			return ServerRootAbsPath;
+		}
+
+		bool FServerBase::IsReady()
+		{
+			return bIsReady;
+		}
+
+		bool FServerBase::IsTickableWhenPaused() const
+		{
+			return true;
+		}
+
+		bool FServerBase::IsTickableInEditor() const
+		{
+			return true;
+		}
+
+		bool FServerBase::IsAllowedToTick() const
+		{
+			return bAllowedToTick;
+		}
+
+		void FServerBase::Tick(float DeltaTime)
+		{
+			// No need to start polling if we have not launched or we have already concluded the server is ready.
+			if (!bHasLaunched || bIsReady)
+			{
+				return;
+			}
+
+			float SecondsElapsedPolling = FPlatformTime::Seconds() - PollingStartedSeconds;
+			float SecondsSinceReconnect = FPlatformTime::Seconds() - LastReconnectionTimeSeconds;
+
+			if (SecondsElapsedPolling < LaunchArgs.ReconnectionTimeoutSeconds)
+			{
+				if (SecondsSinceReconnect < LaunchArgs.ReconnectionIntervalSeconds)
+				{
+					return;
+				}
+
+				if (TestConnection())
+				{
+					// No need to poll anymore we are connected
+					bAllowedToTick = false;
+					bIsReady = true;
+					UE_LOG(LogPixelStreaming, Log, TEXT("Connected to the server. Server is now ready."));
+					OnReady.Broadcast(Endpoints);
+					return;
+				}
+				else
+				{
+					UE_LOG(LogPixelStreaming, Log, TEXT("Polling again in another %.f seconds for server to become ready..."), LaunchArgs.ReconnectionIntervalSeconds);
+					LastReconnectionTimeSeconds = FPlatformTime::Seconds();
+				}
+			}
+			else
+			{
+				// No need to poll anymore we timed out
+				UE_LOG(LogPixelStreaming, Error, TEXT("Server was not ready after %.f seconds, polling timed out."), LaunchArgs.ReconnectionTimeoutSeconds);
+				bAllowedToTick = false;
+				bTimedOut = true;
+				OnFailedToReady.Broadcast();
+			}
+		}
+
+		FServerBase::~FServerBase()
+		{
+			Stop();
+		}
+
+		void FServerBase::Stop()
+		{
+			if (ServerProcess)
+			{
+				ServerProcess->Cancel();
+				ServerProcess.Reset();
+			}
+		}
+
+		bool FServerBase::HasLaunched()
+		{
+			return bHasLaunched;
+		}
+
+		bool FServerBase::IsTimedOut()
+		{
+			return bTimedOut;
+		}
+
+		/* ------------------------------ FCirrus ------------------------------ */
+
+		/* The NodeJS Cirrus signalling server wrapped in a child process */
+		class FCirrus : public FServerBase
+		{
+		public:
+			FCirrus() = default;
+			virtual ~FCirrus() = default;
+
+		private:
+			TUniquePtr<FWebSocketProbe> Probe;
+
+		protected:
+			/* Begin FServerBase interface */
+
+			TSharedPtr<FMonitoredProcess> LaunchImpl(FLaunchArgs& InLaunchArgs, FString& OutServerPath, TMap<EEndpoint, FString>& OutEndPoints) override
+			{
+				FString CirrusDir = GetServersAbsPath() / CirrusName;
+
+				// Query for ports, or set them if they don't exist
+				FString StreamerPort = QueryOrSetServerArg(InLaunchArgs, TEXT("--StreamerPort="), FString::FromInt(NextPort()));
+				FString SFUPort = QueryOrSetServerArg(InLaunchArgs, TEXT("--SFUPort="), FString::FromInt(NextPort()));
+				FString MatchmakerPort = QueryOrSetServerArg(InLaunchArgs, TEXT("--MatchmakerPort="), FString::FromInt(NextPort()));
+				FString HttpPort = QueryOrSetServerArg(InLaunchArgs, TEXT("--HttpPort="), FString::FromInt(NextPort()));
+
+				// Construct endpoint urls
+				const FString SignallingStreamerUrl = FString::Printf(TEXT("ws://localhost:%s"), *StreamerPort);
+				OutEndPoints.Add(EEndpoint::Signalling_Streamer, SignallingStreamerUrl);
+				OutEndPoints.Add(EEndpoint::Signalling_Players, FString::Printf(TEXT("ws://localhost:%s"), *HttpPort));
+				OutEndPoints.Add(EEndpoint::Signalling_SFU, FString::Printf(TEXT("ws://localhost:%s"), *SFUPort));
+				OutEndPoints.Add(EEndpoint::Signalling_Matchmaker, FString::Printf(TEXT("ws://localhost:%s"), *MatchmakerPort));
+				OutEndPoints.Add(EEndpoint::Signalling_Webserver, FString::Printf(TEXT("http://localhost:%s"), *HttpPort));
+
+				TSharedPtr<FMonitoredProcess> CirrusProcess = LaunchServer(CirrusDir, InLaunchArgs, OutServerPath);
+
+				Probe = MakeUnique<FWebSocketProbe>(SignallingStreamerUrl);
+
+				return CirrusProcess;
+			}
+
+			bool TestConnection() override
+			{
+				if (bIsReady)
+				{
+					return true;
+				}
+				else
+				{
+					bool bConnected = Probe->Probe();
+					if (bConnected)
+					{
+						// Close the websocket connection so others can use it
+						Probe.Reset();
+						return true;
+					}
+					else
+					{
+						return false;
+					}
+				}
+			}
+
+			/* End FServerBase interface */
+		};
+
+		/* -------------------- Public static utility methods -------------------- */
+		TSharedPtr<FServerBase> MakeSignallingServer()
+		{
+			return MakeShared<FCirrus>();
+		}
+
+	} // namespace Servers
+} // namespace UE::PixelStreaming
diff --git a/Engine/Plugins/Media/PixelStreaming/Source/PixelStreaming/Private/PlayerSession.cpp b/Engine/Plugins/Media/PixelStreaming/Source/PixelStreaming/Private/PlayerSession.cpp
--- a/Engine/Plugins/Media/PixelStreaming/Source/PixelStreaming/Private/PlayerSession.cpp
+++ b/Engine/Plugins/Media/PixelStreaming/Source/PixelStreaming/Private/PlayerSession.cpp
@@ -4,8 +4,8 @@
 #include "PlayerSessions.h"
 #include "PixelStreamingStatNames.h"
 #include "ToStringExtensions.h"
-#include "SessionDescriptionObservers.h"
-#include "SignallingServerConnection.h"
+#include "PixelStreamingSessionDescriptionObservers.h"
+#include "PixelStreamingSignallingConnection.h"
 #include "Async/Async.h"
 #include "IPixelStreamingStatsConsumer.h"
 #include "RTCStatsCollector.h"
@@ -26,7 +26,7 @@ namespace UE::PixelStreaming
 		// Begin IPixelStreamingStatsConsumer
 		void ConsumeStat(FPixelStreamingPlayerId PlayerId, FName StatName, float StatValue) override
 		{
-			Sessions->ForSession(PlayerId, [&](TSharedPtr<IPlayerSession> Session) {
+			Sessions->ForSession(PlayerId, [StatValue](TSharedPtr<IPlayerSession> Session) {
 				Session->SendVideoEncoderQP((int)StatValue);
 			});
 		}
@@ -35,11 +35,11 @@ namespace UE::PixelStreaming
 		FPlayerSessions* Sessions;
 	};
 
-	FPlayerSession::FPlayerSession(FPlayerSessions* InSessions, FSignallingServerConnection* InSignallingServerConnection, FPixelStreamingPlayerId InPlayerId)
+	FPlayerSession::FPlayerSession(FPlayerSessions* InSessions, FPixelStreamingSignallingConnection* InSignallingServerConnection, FPixelStreamingPlayerId InPlayerId, TSharedPtr<IPixelStreamingInputDevice> InInputDevice)
 		: PlayerSessions(InSessions)
 		, SignallingServerConnection(InSignallingServerConnection)
 		, PlayerId(InPlayerId)
-		, DataChannelObserver(InSessions, InPlayerId)
+		, DataChannelObserver(new FDataChannelObserver(InSessions, InPlayerId, EDataChannelDirection::Bidirectional, InInputDevice))
 		, WebRTCStatsCallback(new rtc::RefCountedObject<FRTCStatsCollector>(InPlayerId))
 		, QPReporter(MakeShared<QPConsumer>(InSessions))
 	{
@@ -59,7 +59,7 @@ namespace UE::PixelStreaming
 	{
 		UE_LOG(LogPixelStreaming, Log, TEXT("%s: PlayerId=%s"), TEXT("FPlayerSession::~FPlayerSession"), *PlayerId);
 
-		DataChannelObserver.Unregister();
+		DataChannelObserver->Unregister();
 		DataChannel = nullptr;
 
 		if (PeerConnection)
@@ -67,12 +67,6 @@ namespace UE::PixelStreaming
 			PeerConnection->Close();
 			PeerConnection = nullptr;
 		}
-
-		FStats* Stats = FStats::Get();
-		if (Stats)
-		{
-			Stats->RemovePeersStats(PlayerId);
-		}
 	}
 
 	void FPlayerSession::PollWebRTCStats() const
@@ -93,15 +87,15 @@ namespace UE::PixelStreaming
 	{
 	}
 
-	FDataChannelObserver* FPlayerSession::GetDataChannelObserver()
+	TSharedPtr<FDataChannelObserver> FPlayerSession::GetDataChannelObserver()
 	{
-		return &DataChannelObserver;
+		return DataChannelObserver;
 	}
 
-	webrtc::PeerConnectionInterface& FPlayerSession::GetPeerConnection()
+	rtc::scoped_refptr<webrtc::PeerConnectionInterface> FPlayerSession::GetPeerConnection()
 	{
 		check(PeerConnection);
-		return *PeerConnection;
+		return PeerConnection;
 	}
 
 	void FPlayerSession::SetPeerConnection(const rtc::scoped_refptr<webrtc::PeerConnectionInterface>& InPeerConnection)
@@ -112,6 +106,7 @@ namespace UE::PixelStreaming
 	void FPlayerSession::SetDataChannel(const rtc::scoped_refptr<webrtc::DataChannelInterface>& InDataChannel)
 	{
 		DataChannel = InDataChannel;
+		DataChannelObserver->Register(InDataChannel);
 	}
 
 	void FPlayerSession::SetVideoSource(const rtc::scoped_refptr<webrtc::VideoTrackSourceInterface>& InVideoSource)
@@ -146,15 +141,15 @@ namespace UE::PixelStreaming
 	void FPlayerSession::OnAnswer(FString Sdp)
 	{
 		webrtc::SdpParseError Error;
-		std::unique_ptr<webrtc::SessionDescriptionInterface> SessionDesc = webrtc::CreateSessionDescription(webrtc::SdpType::kAnswer, ToString(Sdp), &Error);
+		std::unique_ptr<webrtc::SessionDescriptionInterface> SessionDesc = webrtc::CreateSessionDescription(webrtc::SdpType::kAnswer, UE::PixelStreaming::ToString(Sdp), &Error);
 		if (!SessionDesc)
 		{
 			UE_LOG(LogPixelStreaming, Error, TEXT("Failed to parse answer's SDP\n%s"), *Sdp);
 			return;
 		}
 
-		FSetSessionDescriptionObserver* SetRemoteDescriptionObserver = FSetSessionDescriptionObserver::Create(
-			[this]() // on success
+		FPixelStreamingSetSessionDescriptionObserver* SetRemoteDescriptionObserver = FPixelStreamingSetSessionDescriptionObserver::Create(
+			[]() // on success
 			{
 
 			},
@@ -169,7 +164,7 @@ namespace UE::PixelStreaming
 	void FPlayerSession::OnRemoteIceCandidate(const FString& SdpMid, int SdpMLineIndex, const FString& Sdp)
 	{
 		webrtc::SdpParseError Error;
-		std::unique_ptr<webrtc::IceCandidateInterface> Candidate(webrtc::CreateIceCandidate(ToString(SdpMid), SdpMLineIndex, ToString(Sdp), &Error));
+		std::unique_ptr<webrtc::IceCandidateInterface> Candidate(webrtc::CreateIceCandidate(UE::PixelStreaming::ToString(SdpMid), SdpMLineIndex, UE::PixelStreaming::ToString(Sdp), &Error));
 		if (Candidate)
 		{
 			PeerConnection->AddIceCandidate(std::move(Candidate), [this](webrtc::RTCError error) {
@@ -203,6 +198,11 @@ namespace UE::PixelStreaming
 	}
 
 	bool FPlayerSession::SendMessage(Protocol::EToPlayerMsg InMessageType, const FString& Descriptor) const
+	{
+		return SendMessage(DataChannel, InMessageType, Descriptor);
+	}
+
+	bool FPlayerSession::SendMessage(rtc::scoped_refptr<webrtc::DataChannelInterface> DataChannel, Protocol::EToPlayerMsg InMessageType, const FString& Descriptor)
 	{
 		if (!DataChannel)
 		{
@@ -221,13 +221,18 @@ namespace UE::PixelStreaming
 		return DataChannel->Send(webrtc::DataBuffer(Buffer, true));
 	}
 
-	void FPlayerSession::SendQualityControlStatus(bool bIsQualityController) const
+	bool FPlayerSession::SendQualityControlStatus(bool bIsQualityController) const
 	{
-		FStats::Get()->StorePeerStat(PlayerId, FStatData(PixelStreamingStatNames::QualityController, bIsQualityController ? 1.0 : 0.0, 0));
-
 		if (!DataChannel)
 		{
-			return;
+			UE_LOG(LogPixelStreaming, Error, TEXT("Cannot send quality controller status when data channel is null."));
+			return false;
+		}
+
+		if (DataChannel->state() != webrtc::DataChannelInterface::DataState::kOpen)
+		{
+			UE_LOG(LogPixelStreaming, Warning, TEXT("Cannot send quality controller status when data channel not yet open."));
+			return false;
 		}
 
 		const uint8 MessageType = static_cast<uint8>(Protocol::EToPlayerMsg::QualityControlOwnership);
@@ -241,16 +246,63 @@ namespace UE::PixelStreaming
 
 		if (!DataChannel->Send(webrtc::DataBuffer(Buffer, true)))
 		{
-			UE_LOG(LogPixelStreaming, Error, TEXT("failed to send quality control status"));
+			UE_LOG(LogPixelStreaming, Error, TEXT("Failed to send quality control status."));
+			return false;
+		}
+		else
+		{
+			FStats::Get()->StorePeerStat(PlayerId, FStatData(PixelStreamingStatNames::QualityController, bIsQualityController ? 1.0 : 0.0, 0));
+			return true;
 		}
 	}
 
-	void FPlayerSession::SendFreezeFrame(const TArray64<uint8>& JpegBytes) const
+	bool FPlayerSession::SendInputControlStatus(rtc::scoped_refptr<webrtc::DataChannelInterface> DataChannel, FPixelStreamingPlayerId PlayerId, bool bIsInputController)
 	{
-		FPlayerSession::SendArbitraryData(static_cast<TArray<uint8>>(JpegBytes), static_cast<uint8>(Protocol::EToPlayerMsg::FreezeFrame));
+		if (!DataChannel)
+		{
+			UE_LOG(LogPixelStreaming, Error, TEXT("Cannot send input controller status for player=%s when data channel is null."), *PlayerId);
+			return false;
+		}
+
+		if (DataChannel->state() != webrtc::DataChannelInterface::DataState::kOpen)
+		{
+			UE_LOG(LogPixelStreaming, Warning, TEXT("Cannot send input controller status when data channel not yet open."));
+			return false;
+		}
+
+		const uint8 MessageType = static_cast<uint8>(Protocol::EToPlayerMsg::InputControlOwnership);
+		const uint8 ControlsInput = bIsInputController ? 1 : 0;
+
+		rtc::CopyOnWriteBuffer Buffer(sizeof(MessageType) + sizeof(ControlsInput));
+
+		size_t Pos = 0;
+		Pos = SerializeToBuffer(Buffer, Pos, &MessageType, sizeof(MessageType));
+		Pos = SerializeToBuffer(Buffer, Pos, &ControlsInput, sizeof(ControlsInput));
+
+		if (!DataChannel->Send(webrtc::DataBuffer(Buffer, true)))
+		{
+			UE_LOG(LogPixelStreaming, Error, TEXT("Failed to send input control status."));
+			return false;
+		}
+		else
+		{
+			UE_LOG(LogPixelStreaming, Log, TEXT("Player %s sent input controller status. Msg=%d. (Controls input = %s)."), *PlayerId, ControlsInput, bIsInputController ? TEXT("true") : TEXT("false"));
+			FStats::Get()->StorePeerStat(PlayerId, FStatData(PixelStreamingStatNames::InputController, bIsInputController ? 1.0 : 0.0, 0));
+			return true;
+		}
 	}
 
-	void FPlayerSession::SendFileData(const TArray<uint8>& ByteData, const FString& MimeType, const FString& FileExtension) const
+	bool FPlayerSession::SendInputControlStatus(bool bIsInputController) const
+	{
+		return SendInputControlStatus(DataChannel, PlayerId, bIsInputController);
+	}
+
+	bool FPlayerSession::SendFreezeFrame(const TArray64<uint8>& JpegBytes) const
+	{
+		return FPlayerSession::SendArbitraryData(static_cast<TArray<uint8>>(JpegBytes), static_cast<uint8>(Protocol::EToPlayerMsg::FreezeFrame));
+	}
+
+	bool FPlayerSession::SendFileData(const TArray<uint8>& ByteData, const FString& MimeType, const FString& FileExtension) const
 	{
 		// Send the mime type first
 		FPlayerSession::SendMessage(Protocol::EToPlayerMsg::FileMimeType, MimeType);
@@ -258,17 +310,22 @@ namespace UE::PixelStreaming
 		// Send the extension next
 		FPlayerSession::SendMessage(Protocol::EToPlayerMsg::FileExtension, FileExtension);
 
-		// Send the contents of the file
-		AsyncTask(ENamedThreads::AnyHiPriThreadHiPriTask, [this, ByteData]() {
-			FPlayerSession::SendArbitraryData(ByteData, static_cast<uint8>(Protocol::EToPlayerMsg::FileContents));
-		});
+		// Send the contents of the file. Note to callers: consider running this on its own thread, it can take a while if the file is big.
+		bool bSentData = SendArbitraryData(ByteData, static_cast<uint8>(Protocol::EToPlayerMsg::FileContents));
+		if (!bSentData)
+		{
+			UE_LOG(LogPixelStreaming, Error, TEXT("Unable to send file data over the data channel."));
+			return false;
+		}
+		return true;
 	}
 
-	void FPlayerSession::SendUnfreezeFrame() const
+	bool FPlayerSession::SendUnfreezeFrame() const
 	{
 		if (!DataChannel)
 		{
-			return;
+			UE_LOG(LogPixelStreaming, Error, TEXT("Cannot send unfreeze frame when data channel is null."));
+			return false;
 		}
 
 		const uint8 MessageType = static_cast<uint8>(Protocol::EToPlayerMsg::UnfreezeFrame);
@@ -280,15 +337,18 @@ namespace UE::PixelStreaming
 
 		if (!DataChannel->Send(webrtc::DataBuffer(Buffer, true)))
 		{
-			UE_LOG(LogPixelStreaming, Error, TEXT("failed to send unfreeze frame"));
+			UE_LOG(LogPixelStreaming, Error, TEXT("Failed to send unfreeze frame."));
+			return false;
 		}
+		return true;
 	}
 
-	void FPlayerSession::SendArbitraryData(const TArray<uint8>& DataBytes, const uint8 MessageType) const
+	bool FPlayerSession::SendArbitraryData(const TArray<uint8>& DataBytes, const uint8 MessageType) const
 	{
 		if (!DataChannel)
 		{
-			return;
+			UE_LOG(LogPixelStreaming, Error, TEXT("Cannot send arbitrary data when data channel is null."));
+			return false;
 		}
 
 		// int32 results in a maximum 4GB file (4,294,967,296 bytes)
@@ -322,19 +382,21 @@ namespace UE::PixelStreaming
 			uint64_t BufferBefore = DataChannel->buffered_amount();
 			while (BufferBefore + BytesToTransmit >= 16 * 1024 * 1024) // 16MB (WebRTC Data Channel buffer size)
 			{
-				FPlatformProcess::Sleep(0.000001f); // sleep 1 microsecond
+				// As per UE docs a Sleep of 0.0 simply lets other threads take CPU cycles while this is happening.
+				FPlatformProcess::Sleep(0.0);
 				BufferBefore = DataChannel->buffered_amount();
 			}
 
 			if (!DataChannel->Send(webrtc::DataBuffer(Buffer, true)))
 			{
 				UE_LOG(LogPixelStreaming, Error, TEXT("Failed to send data channel packet"));
-				return;
+				return false;
 			}
 
 			// Increment the number of bytes transmitted
 			BytesTransmitted += BytesToTransmit;
 		}
+		return true;
 	}
 
 	//
@@ -343,25 +405,25 @@ namespace UE::PixelStreaming
 
 	void FPlayerSession::OnSignalingChange(webrtc::PeerConnectionInterface::SignalingState NewState)
 	{
-		UE_LOG(LogPixelStreaming, Log, TEXT("%s : PlayerId=%s, NewState=%s"), TEXT("FPlayerSession::OnSignalingChange"), *PlayerId, ToString(NewState));
+		UE_LOG(LogPixelStreaming, Log, TEXT("%s : PlayerId=%s, NewState=%s"), TEXT("FPlayerSession::OnSignalingChange"), *PlayerId, UE::PixelStreaming::ToString(NewState));
 	}
 
 	// Called when a remote stream is added
 	void FPlayerSession::OnAddStream(rtc::scoped_refptr<webrtc::MediaStreamInterface> Stream)
 	{
-		UE_LOG(LogPixelStreaming, Log, TEXT("%s : PlayerId=%s, Stream=%s"), TEXT("FPlayerSession::OnAddStream"), *PlayerId, *ToString(Stream->id()));
+		UE_LOG(LogPixelStreaming, Log, TEXT("%s : PlayerId=%s, Stream=%s"), TEXT("FPlayerSession::OnAddStream"), *PlayerId, *UE::PixelStreaming::ToString(Stream->id()));
 	}
 
 	void FPlayerSession::OnRemoveStream(rtc::scoped_refptr<webrtc::MediaStreamInterface> Stream)
 	{
-		UE_LOG(LogPixelStreaming, Log, TEXT("%s : PlayerId=%s, Stream=%s"), TEXT("FPlayerSession::OnRemoveStream"), *PlayerId, *ToString(Stream->id()));
+		UE_LOG(LogPixelStreaming, Log, TEXT("%s : PlayerId=%s, Stream=%s"), TEXT("FPlayerSession::OnRemoveStream"), *PlayerId, *UE::PixelStreaming::ToString(Stream->id()));
 	}
 
 	void FPlayerSession::OnDataChannel(rtc::scoped_refptr<webrtc::DataChannelInterface> InDataChannel)
 	{
 		UE_LOG(LogPixelStreaming, Log, TEXT("%s : PlayerId=%s"), TEXT("FPlayerSession::OnDataChannel"), *PlayerId);
 		DataChannel = InDataChannel;
-		DataChannelObserver.Register(InDataChannel);
+		DataChannelObserver->Register(InDataChannel);
 	}
 
 	void FPlayerSession::OnRenegotiationNeeded()
@@ -371,17 +433,19 @@ namespace UE::PixelStreaming
 
 	void FPlayerSession::OnIceConnectionChange(webrtc::PeerConnectionInterface::IceConnectionState NewState)
 	{
-		UE_LOG(LogPixelStreaming, Log, TEXT("%s : PlayerId=%s, NewState=%s"), TEXT("FPlayerSession::OnIceConnectionChange"), *PlayerId, ToString(NewState));
+		UE_LOG(LogPixelStreaming, Log, TEXT("%s : PlayerId=%s, NewState=%s"), TEXT("FPlayerSession::OnIceConnectionChange"), *PlayerId, UE::PixelStreaming::ToString(NewState));
 	}
 
 	void FPlayerSession::OnIceGatheringChange(webrtc::PeerConnectionInterface::IceGatheringState NewState)
 	{
-		UE_LOG(LogPixelStreaming, Log, TEXT("%s : PlayerId=%s, NewState=%s"), TEXT("FPlayerSession::OnIceGatheringChange"), *PlayerId, ToString(NewState));
+		UE_LOG(LogPixelStreaming, Log, TEXT("%s : PlayerId=%s, NewState=%s"), TEXT("FPlayerSession::OnIceGatheringChange"), *PlayerId, UE::PixelStreaming::ToString(NewState));
 	}
 
 	void FPlayerSession::OnIceCandidate(const webrtc::IceCandidateInterface* Candidate)
 	{
-		UE_LOG(LogPixelStreaming, Log, TEXT("%s : PlayerId=%s"), TEXT("FPlayerSession::OnIceCandidate"), *PlayerId);
+		std::string IceCandidateStdStr;
+		Candidate->ToString(&IceCandidateStdStr);
+		UE_LOG(LogPixelStreaming, Verbose, TEXT("%s : PlayerId=%s Candidate=%s"), TEXT("FPlayerSession::OnIceCandidate"), *PlayerId, *FString(IceCandidateStdStr.c_str()));
 
 		SignallingServerConnection->SendIceCandidate(PlayerId, *Candidate);
 	}
@@ -465,11 +529,13 @@ namespace UE::PixelStreaming
 		UE_LOG(LogPixelStreaming, Log, TEXT("%s : PlayerId=%s"), TEXT("FPlayerSession::OnRemoveTrack"), *PlayerId);
 	}
 
-	void FPlayerSession::SendVideoEncoderQP(double QP) const
+	bool FPlayerSession::SendVideoEncoderQP(double QP) const
 	{
 		if (!SendMessage(Protocol::EToPlayerMsg::VideoEncoderAvgQP, FString::FromInt(QP)))
 		{
 			UE_LOG(LogPixelStreaming, Verbose, TEXT("Failed to send video encoder QP to peer."));
+			return false;
 		}
+		return true;
 	}
 } // namespace UE::PixelStreaming
diff --git a/Engine/Plugins/Media/PixelStreaming/Source/PixelStreaming/Private/PlayerSession.h b/Engine/Plugins/Media/PixelStreaming/Source/PixelStreaming/Private/PlayerSession.h
--- a/Engine/Plugins/Media/PixelStreaming/Source/PixelStreaming/Private/PlayerSession.h
+++ b/Engine/Plugins/Media/PixelStreaming/Source/PixelStreaming/Private/PlayerSession.h
@@ -9,22 +9,24 @@
 
 class IPixelStreamingSessions;
 class IPixelStreamingStatsConsumer;
+class IPixelStreamingInputDevice;
+class FPixelStreamingSignallingConnection;
 
 namespace UE::PixelStreaming
 {
 	class FPlayerSessions;
-	class FSignallingServerConnection;
 
 	class FPlayerSession : public IPlayerSession
 	{
 	public:
-		FPlayerSession(FPlayerSessions* InSessions, FSignallingServerConnection* InSignallingServerConnection, FPixelStreamingPlayerId PlayerId);
+		FPlayerSession(FPlayerSessions* InSessions, FPixelStreamingSignallingConnection* InSignallingServerConnection, FPixelStreamingPlayerId PlayerId, TSharedPtr<IPixelStreamingInputDevice> InInputDevice);
 		virtual ~FPlayerSession();
 
-		virtual webrtc::PeerConnectionInterface& GetPeerConnection() override;
+		virtual rtc::scoped_refptr<webrtc::PeerConnectionInterface> GetPeerConnection() override;
 		virtual void SetPeerConnection(const rtc::scoped_refptr<webrtc::PeerConnectionInterface>& InPeerConnection) override;
 		virtual void SetDataChannel(const rtc::scoped_refptr<webrtc::DataChannelInterface>& InDataChannel) override;
 		virtual void SetVideoSource(const rtc::scoped_refptr<webrtc::VideoTrackSourceInterface>& InVideoSource) override;
+		virtual rtc::scoped_refptr<webrtc::VideoTrackSourceInterface> GetVideoSource() const override { return VideoSource; }
 
 		virtual void OnAnswer(FString Sdp) override;
 		virtual void OnRemoteIceCandidate(const FString& SdpMid, int SdpMLineIndex, const FString& Sdp) override;
@@ -33,15 +35,16 @@ namespace UE::PixelStreaming
 		virtual FName GetSessionType() const override { return Type; }
 		virtual FPixelStreamingPlayerId GetPlayerId() const override;
 		virtual IPixelStreamingAudioSink* GetAudioSink() override;
-		virtual FDataChannelObserver* GetDataChannelObserver() override;
+		virtual TSharedPtr<FDataChannelObserver> GetDataChannelObserver() override;
 
 		virtual bool SendMessage(Protocol::EToPlayerMsg Type, const FString& Descriptor) const override;
-		virtual void SendQualityControlStatus(bool bIsQualityController) const override;
-		virtual void SendFreezeFrame(const TArray64<uint8>& JpegBytes) const override;
-		virtual void SendFileData(const TArray<uint8>& ByteData, const FString& MimeType, const FString& FileExtension) const override;
-		virtual void SendUnfreezeFrame() const override;
-		virtual void SendArbitraryData(const TArray<uint8>& DataBytes, const uint8 MessageType) const override;
-		virtual void SendVideoEncoderQP(double QP) const override;
+		virtual bool SendInputControlStatus(bool bIsInputController) const override;
+		virtual bool SendQualityControlStatus(bool bIsQualityController) const override;
+		virtual bool SendFreezeFrame(const TArray64<uint8>& JpegBytes) const override;
+		virtual bool SendFileData(const TArray<uint8>& ByteData, const FString& MimeType, const FString& FileExtension) const override;
+		virtual bool SendUnfreezeFrame() const override;
+		virtual bool SendArbitraryData(const TArray<uint8>& DataBytes, const uint8 MessageType) const override;
+		virtual bool SendVideoEncoderQP(double QP) const override;
 		virtual void PollWebRTCStats() const override;
 		virtual void OnDataChannelClosed() const override;
 
@@ -66,14 +69,16 @@ namespace UE::PixelStreaming
 
 	public:
 		inline static const FName Type = FName(TEXT("P2P"));
+		static bool SendInputControlStatus(rtc::scoped_refptr<webrtc::DataChannelInterface> DataChannel, FPixelStreamingPlayerId PlayerId, bool bIsInputController);
+		static bool SendMessage(rtc::scoped_refptr<webrtc::DataChannelInterface> DataChannel, Protocol::EToPlayerMsg Type, const FString& Descriptor);
 
 	protected:
 		FPlayerSessions* PlayerSessions;
-		FSignallingServerConnection* SignallingServerConnection;
+		FPixelStreamingSignallingConnection* SignallingServerConnection;
 		FPixelStreamingPlayerId PlayerId;
 		rtc::scoped_refptr<webrtc::PeerConnectionInterface> PeerConnection;
 		rtc::scoped_refptr<webrtc::DataChannelInterface> DataChannel;
-		FDataChannelObserver DataChannelObserver;
+		TSharedPtr<FDataChannelObserver> DataChannelObserver;
 		FThreadSafeBool bDisconnecting = false;
 		FAudioSink AudioSink;
 		rtc::scoped_refptr<webrtc::RTCStatsCollectorCallback> WebRTCStatsCallback;
diff --git a/Engine/Plugins/Media/PixelStreaming/Source/PixelStreaming/Private/PlayerSessionDataOnly.cpp b/Engine/Plugins/Media/PixelStreaming/Source/PixelStreaming/Private/PlayerSessionDataOnly.cpp
--- a/Engine/Plugins/Media/PixelStreaming/Source/PixelStreaming/Private/PlayerSessionDataOnly.cpp
+++ b/Engine/Plugins/Media/PixelStreaming/Source/PixelStreaming/Private/PlayerSessionDataOnly.cpp
@@ -4,70 +4,77 @@
 #include "PixelStreamingDelegates.h"
 #include "PixelStreamingPrivate.h"
 #include "Utils.h"
+#include "PlayerSession.h"
 
 namespace UE::PixelStreaming
 {
-	FPlayerSessionDataOnly::FPlayerSessionDataOnly(FPlayerSessions* InSessions,
-		FSignallingServerConnection* InSignallingServerConnection,
-		FPixelStreamingPlayerId InPlayerId,
-		webrtc::PeerConnectionInterface* InPeerConnection,
-		int32 SendStreamId,
-		int32 RecvStreamId)
-		: SignallingServerConnection(InSignallingServerConnection)
-		, PlayerId(InPlayerId)
-		, PeerConnection(InPeerConnection)
-		, DataChannelObserver(InSessions, InPlayerId)
+	FPlayerSessionDataOnly::FPlayerSessionDataOnly(FPlayerSessions* InSessions, FPixelStreamingPlayerId InPlayerId, TSharedPtr<IPixelStreamingInputDevice> InInputDevice, rtc::scoped_refptr<webrtc::PeerConnectionInterface> InPeerConnection, int32 SendStreamId, int32 RecvStreamId)
+		: PlayerId(InPlayerId)
+		, bUsingBidirectionalDataChannel(SendStreamId == RecvStreamId)
 	{
-		webrtc::DataChannelInit SendDataChannelConfig;
-		SendDataChannelConfig.reliable = true;
-		SendDataChannelConfig.ordered = true;
-		SendDataChannelConfig.negotiated = true;
-		SendDataChannelConfig.id = SendStreamId;
+		SetupDataChannels(InSessions, InPeerConnection, SendStreamId, RecvStreamId, InInputDevice);
+	}
 
-		SendDataChannel = PeerConnection->CreateDataChannel("datachannel", &SendDataChannelConfig);
+	FPlayerSessionDataOnly::~FPlayerSessionDataOnly()
+	{
+	}
 
-		if (SendStreamId != RecvStreamId)
+	void FPlayerSessionDataOnly::SetupDataChannels(FPlayerSessions* InSessions, rtc::scoped_refptr<webrtc::PeerConnectionInterface> InPeerConnection, int32 SendStreamId, int32 RecvStreamId, TSharedPtr<IPixelStreamingInputDevice> InInputDevice)
+	{
+		bUsingBidirectionalDataChannel = SendStreamId == RecvStreamId;
+		if (!bUsingBidirectionalDataChannel)
 		{
+			webrtc::DataChannelInit SendDataChannelConfig;
+			SendDataChannelConfig.reliable = true;
+			SendDataChannelConfig.ordered = true;
+			SendDataChannelConfig.negotiated = true;
+			SendDataChannelConfig.id = SendStreamId;
+			SendDataChannel = InPeerConnection->CreateDataChannel("send-datachannel", &SendDataChannelConfig);
+			SendDataChannelObserver = MakeShared<FDataChannelObserver>(InSessions, PlayerId, EDataChannelDirection::SendOnly, InInputDevice);
+			SendDataChannelObserver->Register(SendDataChannel);
+
 			webrtc::DataChannelInit RecvDataChannelConfig;
 			RecvDataChannelConfig.reliable = true;
 			RecvDataChannelConfig.ordered = true;
 			RecvDataChannelConfig.negotiated = true;
 			RecvDataChannelConfig.id = RecvStreamId;
-			RecvDataChannel = PeerConnection->CreateDataChannel("datachannel", &RecvDataChannelConfig);
+			RecvDataChannel = InPeerConnection->CreateDataChannel("recv-datachannel", &RecvDataChannelConfig);
+			RecvDataChannelObserver = MakeShared<FDataChannelObserver>(InSessions, PlayerId, EDataChannelDirection::RecvOnly, InInputDevice);
+			RecvDataChannelObserver->Register(RecvDataChannel);
 		}
 		else
 		{
-			RecvDataChannel = SendDataChannel;
+			webrtc::DataChannelInit DataChannelConfig;
+			DataChannelConfig.reliable = true;
+			DataChannelConfig.ordered = true;
+			DataChannelConfig.negotiated = true;
+			DataChannelConfig.id = SendStreamId;
+			BidiDataChannel = InPeerConnection->CreateDataChannel("datachannel", &DataChannelConfig);
+			BidiDataChannelObserver = MakeShared<FDataChannelObserver>(InSessions, PlayerId, EDataChannelDirection::Bidirectional, InInputDevice);
+			BidiDataChannelObserver->Register(BidiDataChannel);
 		}
-		RecvDataChannel->RegisterObserver(&DataChannelObserver);
 	}
 
-	FPlayerSessionDataOnly::~FPlayerSessionDataOnly()
+	TSharedPtr<FDataChannelObserver> FPlayerSessionDataOnly::GetDataChannelObserver()
 	{
-		UE_LOG(LogPixelStreaming, Log, TEXT("%s: PlayerId=%s"), TEXT("FPlayerSessionDataOnly::~FPlayerSessionDataOnly"), *PlayerId);
-
-		RecvDataChannel->UnregisterObserver();
-		SendDataChannel = nullptr;
-		RecvDataChannel = nullptr;
+		if (bUsingBidirectionalDataChannel)
+		{
+			return BidiDataChannelObserver;
+		}
+		else
+		{
+			return SendDataChannelObserver;
+		}
 	}
 
 	bool FPlayerSessionDataOnly::SendMessage(Protocol::EToPlayerMsg InMessageType, const FString& Descriptor) const
 	{
-		if (!SendDataChannel)
-		{
-			return false;
-		}
-
-		const uint8 MessageType = static_cast<uint8>(InMessageType);
-		const size_t DescriptorSize = Descriptor.Len() * sizeof(TCHAR);
-
-		rtc::CopyOnWriteBuffer Buffer(sizeof(MessageType) + DescriptorSize);
-
-		size_t Pos = 0;
-		Pos = SerializeToBuffer(Buffer, Pos, &MessageType, sizeof(MessageType));
-		Pos = SerializeToBuffer(Buffer, Pos, *Descriptor, DescriptorSize);
+		return FPlayerSession::SendMessage(bUsingBidirectionalDataChannel ? BidiDataChannel : SendDataChannel, InMessageType, Descriptor);
+	}
 
-		return SendDataChannel->Send(webrtc::DataBuffer(Buffer, true));
+	bool FPlayerSessionDataOnly::SendInputControlStatus(bool bIsInputController) const
+	{
+		return FPlayerSession::SendInputControlStatus(bUsingBidirectionalDataChannel ? BidiDataChannel : SendDataChannel, PlayerId, bIsInputController);
 	}
 
 	void FPlayerSessionDataOnly::OnDataChannelClosed() const
diff --git a/Engine/Plugins/Media/PixelStreaming/Source/PixelStreaming/Private/PlayerSessionDataOnly.h b/Engine/Plugins/Media/PixelStreaming/Source/PixelStreaming/Private/PlayerSessionDataOnly.h
--- a/Engine/Plugins/Media/PixelStreaming/Source/PixelStreaming/Private/PlayerSessionDataOnly.h
+++ b/Engine/Plugins/Media/PixelStreaming/Source/PixelStreaming/Private/PlayerSessionDataOnly.h
@@ -7,25 +7,25 @@
 
 namespace UE::PixelStreaming
 {
-	class FSignallingServerConnection;
 	class FPlayerSessions;
 
 	class FPlayerSessionDataOnly : public IPlayerSession
 	{
 	public:
 		FPlayerSessionDataOnly(FPlayerSessions* InSessions,
-			FSignallingServerConnection* InSignallingServerConnection,
 			FPixelStreamingPlayerId InPlayerId,
-			webrtc::PeerConnectionInterface* InPeerConnection,
+			TSharedPtr<IPixelStreamingInputDevice> InInputDevice,
+			rtc::scoped_refptr<webrtc::PeerConnectionInterface> InPeerConnection,
 			int32 SendStreamId,
 			int32 RecvStreamId);
 
 		virtual ~FPlayerSessionDataOnly();
 
-		virtual webrtc::PeerConnectionInterface& GetPeerConnection() override { return *PeerConnection; }
+		virtual rtc::scoped_refptr<webrtc::PeerConnectionInterface> GetPeerConnection() override { return nullptr; }
 		virtual void SetPeerConnection(const rtc::scoped_refptr<webrtc::PeerConnectionInterface>& InPeerConnection) override {}
 		virtual void SetDataChannel(const rtc::scoped_refptr<webrtc::DataChannelInterface>& InDataChannel) override {}
 		virtual void SetVideoSource(const rtc::scoped_refptr<webrtc::VideoTrackSourceInterface>& InVideoSource) override {}
+		virtual rtc::scoped_refptr<webrtc::VideoTrackSourceInterface> GetVideoSource() const override { return nullptr; }
 
 		virtual void OnAnswer(FString Sdp) override {}
 		virtual void OnRemoteIceCandidate(const FString& SdpMid, int SdpMLineIndex, const FString& Sdp) override {}
@@ -34,15 +34,16 @@ namespace UE::PixelStreaming
 		virtual FName GetSessionType() const override { return Type; };
 		virtual FPixelStreamingPlayerId GetPlayerId() const override { return PlayerId; }
 		virtual IPixelStreamingAudioSink* GetAudioSink() override { return nullptr; }
-		virtual FDataChannelObserver* GetDataChannelObserver() override { return &DataChannelObserver; }
+		virtual TSharedPtr<FDataChannelObserver> GetDataChannelObserver() override;
 
 		virtual bool SendMessage(Protocol::EToPlayerMsg Type, const FString& Descriptor) const override;
-		virtual void SendQualityControlStatus(bool bIsQualityController) const override {}
-		virtual void SendFreezeFrame(const TArray64<uint8>& JpegBytes) const override {}
-		virtual void SendFileData(const TArray<uint8>& ByteData, const FString& MimeType, const FString& FileExtension) const override {}
-		virtual void SendUnfreezeFrame() const override {}
-		virtual void SendArbitraryData(const TArray<uint8>& DataBytes, const uint8 MessageType) const override {}
-		virtual void SendVideoEncoderQP(double QP) const override {}
+		virtual bool SendInputControlStatus(bool bIsInputController) const override;
+		virtual bool SendQualityControlStatus(bool bIsQualityController) const override { return false; }
+		virtual bool SendFreezeFrame(const TArray64<uint8>& JpegBytes) const override { return false; }
+		virtual bool SendFileData(const TArray<uint8>& ByteData, const FString& MimeType, const FString& FileExtension) const override { return false; }
+		virtual bool SendUnfreezeFrame() const override { return false; }
+		virtual bool SendArbitraryData(const TArray<uint8>& DataBytes, const uint8 MessageType) const override { return false; }
+		virtual bool SendVideoEncoderQP(double QP) const override { return false; }
 		virtual void PollWebRTCStats() const override {}
 		virtual void OnDataChannelClosed() const override;
 
@@ -63,15 +64,19 @@ namespace UE::PixelStreaming
 		virtual void OnTrack(rtc::scoped_refptr<webrtc::RtpTransceiverInterface> transceiver) override {}
 		virtual void OnRemoveTrack(rtc::scoped_refptr<webrtc::RtpReceiverInterface> receiver) override {}
 
+		virtual void SetupDataChannels(FPlayerSessions* InSessions, rtc::scoped_refptr<webrtc::PeerConnectionInterface> InPeerConnection, int32 SendStreamId, int32 RecvStreamId, TSharedPtr<IPixelStreamingInputDevice> InInputDevice);
+
 	public:
 		inline static const FName Type = FName(TEXT("DataOnly"));
 
 	private:
-		FSignallingServerConnection* SignallingServerConnection;
 		FPixelStreamingPlayerId PlayerId;
-		webrtc::PeerConnectionInterface* PeerConnection;
+		bool bUsingBidirectionalDataChannel;
 		rtc::scoped_refptr<webrtc::DataChannelInterface> SendDataChannel;
 		rtc::scoped_refptr<webrtc::DataChannelInterface> RecvDataChannel;
-		FDataChannelObserver DataChannelObserver;
+		rtc::scoped_refptr<webrtc::DataChannelInterface> BidiDataChannel;
+		TSharedPtr<FDataChannelObserver> SendDataChannelObserver;
+		TSharedPtr<FDataChannelObserver> RecvDataChannelObserver;
+		TSharedPtr<FDataChannelObserver> BidiDataChannelObserver;
 	};
 } // namespace UE::PixelStreaming
diff --git a/Engine/Plugins/Media/PixelStreaming/Source/PixelStreaming/Private/PlayerSessionSFU.cpp b/Engine/Plugins/Media/PixelStreaming/Source/PixelStreaming/Private/PlayerSessionSFU.cpp
--- a/Engine/Plugins/Media/PixelStreaming/Source/PixelStreaming/Private/PlayerSessionSFU.cpp
+++ b/Engine/Plugins/Media/PixelStreaming/Source/PixelStreaming/Private/PlayerSessionSFU.cpp
@@ -6,8 +6,8 @@
 
 namespace UE::PixelStreaming
 {
-	FPlayerSessionSFU::FPlayerSessionSFU(FPlayerSessions* InSessions, FSignallingServerConnection* InSignallingServerConnection, FPixelStreamingPlayerId InPlayerId)
-		: FPlayerSession(InSessions, InSignallingServerConnection, InPlayerId)
+	FPlayerSessionSFU::FPlayerSessionSFU(FPlayerSessions* InSessions, FPixelStreamingSignallingConnection* InSignallingServerConnection, FPixelStreamingPlayerId InPlayerId, TSharedPtr<IPixelStreamingInputDevice> InInputDevice)
+		: FPlayerSession(InSessions, InSignallingServerConnection, InPlayerId, InInputDevice)
 	{
 	}
 
diff --git a/Engine/Plugins/Media/PixelStreaming/Source/PixelStreaming/Private/PlayerSessionSFU.h b/Engine/Plugins/Media/PixelStreaming/Source/PixelStreaming/Private/PlayerSessionSFU.h
--- a/Engine/Plugins/Media/PixelStreaming/Source/PixelStreaming/Private/PlayerSessionSFU.h
+++ b/Engine/Plugins/Media/PixelStreaming/Source/PixelStreaming/Private/PlayerSessionSFU.h
@@ -10,7 +10,7 @@ namespace UE::PixelStreaming
 	class FPlayerSessionSFU : public FPlayerSession
 	{
 	public:
-		FPlayerSessionSFU(FPlayerSessions* InSessions, FSignallingServerConnection* InSignallingServerConnection, FPixelStreamingPlayerId PlayerId);
+		FPlayerSessionSFU(FPlayerSessions* InSessions, FPixelStreamingSignallingConnection* InSignallingServerConnection, FPixelStreamingPlayerId PlayerId, TSharedPtr<IPixelStreamingInputDevice> InInputDevice);
 		virtual ~FPlayerSessionSFU();
 
 		void AddChildSession(TSharedPtr<FPlayerSessionDataOnly> ChildSession);
diff --git a/Engine/Plugins/Media/PixelStreaming/Source/PixelStreaming/Private/PlayerSessions.cpp b/Engine/Plugins/Media/PixelStreaming/Source/PixelStreaming/Private/PlayerSessions.cpp
--- a/Engine/Plugins/Media/PixelStreaming/Source/PixelStreaming/Private/PlayerSessions.cpp
+++ b/Engine/Plugins/Media/PixelStreaming/Source/PixelStreaming/Private/PlayerSessions.cpp
@@ -7,20 +7,20 @@
 #include "PlayerSessionDataOnly.h"
 #include "PlayerSessionSFU.h"
 #include "Settings.h"
+#include "Stats.h"
 
 namespace UE::PixelStreaming
 {
 	TSharedPtr<IPlayerSession> FPlayerSessions::CreatePlayerSession(
 		FPixelStreamingPlayerId PlayerId,
+		TSharedPtr<IPixelStreamingInputDevice> InputDevice,
 		rtc::scoped_refptr<webrtc::PeerConnectionFactoryInterface> PeerConnectionFactory,
 		webrtc::PeerConnectionInterface::RTCConfiguration PeerConnectionConfig,
-		FSignallingServerConnection* SignallingServerConnection, int Flags)
+		FPixelStreamingSignallingConnection* SignallingServerConnection, int Flags)
 	{
 		check(PeerConnectionFactory);
 
-		// With unified plan, we get several calls to OnOffer, which in turn calls
-		// this several times.
-		// Therefore, we only try to create the player if not created already
+		// We don't want to create players we already have.
 		if (Players.Contains(PlayerId))
 		{
 			return nullptr;
@@ -28,45 +28,80 @@ namespace UE::PixelStreaming
 
 		UE_LOG(LogPixelStreaming, Log, TEXT("Creating player session for PlayerId=%s"), *PlayerId);
 
-		// this is called from WebRTC signalling thread, the only thread where
-		// `Players` map is modified, so no need to lock it
 		bool bIsSFU = (Flags & Protocol::EPlayerFlags::PSPFlag_IsSFU) != 0;
 
 		TSharedPtr<IPlayerSession> Session;
 		if (bIsSFU)
 		{
-			Session = MakeShared<FPlayerSessionSFU>(this, SignallingServerConnection, PlayerId);
+			Session = MakeShared<FPlayerSessionSFU>(this, SignallingServerConnection, PlayerId, InputDevice);
 		}
 		else
 		{
-			Session = MakeShared<FPlayerSession>(this, SignallingServerConnection, PlayerId);
+			Session = MakeShared<FPlayerSession>(this, SignallingServerConnection, PlayerId, InputDevice);
 		}
 
+		rtc::scoped_refptr<webrtc::PeerConnectionInterface> PeerConnection = CreatePeerConnection(PeerConnectionFactory, PeerConnectionConfig, Session);
+		Session->SetPeerConnection(PeerConnection);
+
+		// Create a datachannel, if needed.
+		const bool bSupportsDataChannel = (Flags & Protocol::EPlayerFlags::PSPFlag_SupportsDataChannel) != 0;
+		if (bSupportsDataChannel)
+		{
+			Session->SetDataChannel(CreateDataChannel(PeerConnection, bIsSFU));
+		}
+
+		AddSession(PlayerId, Session);
+		return Session;
+	}
+
+	rtc::scoped_refptr<webrtc::DataChannelInterface> FPlayerSessions::CreateDataChannel(
+		rtc::scoped_refptr<webrtc::PeerConnectionInterface> PeerConnection,
+		bool bIsSFU) const
+	{
+		webrtc::DataChannelInit DataChannelConfig;
+		DataChannelConfig.reliable = true;
+		DataChannelConfig.ordered = true;
+		if (bIsSFU)
+		{
+			DataChannelConfig.negotiated = true;
+			DataChannelConfig.id = 1023;
+		}
+
+		return PeerConnection->CreateDataChannel("datachannel", &DataChannelConfig);
+	}
+
+	void FPlayerSessions::AddSession(FPixelStreamingPlayerId PlayerId, TSharedPtr<IPlayerSession> Session)
+	{
+		bool bIsSFU = Session->GetSessionType() == FPlayerSessionSFU::Type;
+
 		// Note: this flag is very specifically calculated BEFORE the new session is added.
-		// Peer only becomes quality controller if there is no other P2P peers already connected.
+		// P2P peer only becomes quality controller if there is no other P2P peers already connected.
 		bool bMakeQualityController = !bIsSFU && NumP2PPeers() == 0 && Session->GetSessionType() == FPlayerSession::Type;
 
-		Session->SetPeerConnection(CreatePeerConnection(PeerConnectionFactory, PeerConnectionConfig, Session));
-		AddSession(PlayerId, Session);
+		// Only time we automatically make a new peer the input controlling host is if they are the first peer (and not the SFU).
+		bool bHostControlsInput = UE::PixelStreaming::Settings::GetInputControllerMode() == UE::PixelStreaming::Settings::EInputControllerMode::Host;
+		bool bMakeInputController = bHostControlsInput && !bIsSFU && NumNonSFUPeers() == 0;
+
+		{
+			FScopeLock Lock(&PlayersCS);
+			Players.Add(PlayerId, Session);
+		}
 
 		if (bMakeQualityController)
 		{
 			SetQualityController(PlayerId);
 		}
 
+		if (bMakeInputController)
+		{
+			SetInputController(PlayerId);
+		}
+
 		if (UPixelStreamingDelegates* Delegates = UPixelStreamingDelegates::GetPixelStreamingDelegates())
 		{
 			Delegates->OnNewConnection.Broadcast(PlayerId, bMakeQualityController);
 			Delegates->OnNewConnectionNative.Broadcast(PlayerId, bMakeQualityController);
 		}
-
-		return Session;
-	}
-
-	void FPlayerSessions::AddSession(FPixelStreamingPlayerId PlayerId, TSharedPtr<IPlayerSession> Session)
-	{
-		FScopeLock Lock(&PlayersCS);
-		Players.Add(PlayerId, Session);
 	}
 
 	rtc::scoped_refptr<webrtc::PeerConnectionInterface> FPlayerSessions::CreatePeerConnection(
@@ -93,21 +128,27 @@ namespace UE::PixelStreaming
 		return PeerConnection;
 	}
 
-	void FPlayerSessions::CreateNewDataChannel(
-		FPixelStreamingPlayerId SFUId, FPixelStreamingPlayerId PlayerId,
-		int32 SendStreamId, int32 RecvStreamId,
-		FSignallingServerConnection* SignallingServerConnection)
+	void FPlayerSessions::CreateSFUDataOnlyPeer(
+		FPixelStreamingPlayerId SFUId,
+		FPixelStreamingPlayerId PlayerId,
+		TSharedPtr<IPixelStreamingInputDevice> InputDevice,
+		int32 SendStreamId,
+		int32 RecvStreamId,
+		FPixelStreamingSignallingConnection* SignallingServerConnection)
 	{
+		UE_LOG(LogPixelStreaming, Log, TEXT("PlayerId=%s - Creating new SFU dataonly peer with stream ids (sendId=%d recvId=%d)"), *PlayerId, SendStreamId, RecvStreamId);
+
 		if (auto ParentSession = StaticCastSharedPtr<FPlayerSessionSFU>(GetPlayerSession(SFUId)))
 		{
-			auto PeerConnection = &ParentSession->GetPeerConnection();
-			auto NewSession = MakeShared<FPlayerSessionDataOnly>(this, SignallingServerConnection, PlayerId, PeerConnection, SendStreamId, RecvStreamId);
+			TSharedPtr<FPlayerSessionDataOnly> NewSession = MakeShared<FPlayerSessionDataOnly>(this, PlayerId, InputDevice, ParentSession->GetPeerConnection(), SendStreamId, RecvStreamId);
 			ParentSession->AddChildSession(NewSession);
 
-			{
-				FScopeLock Lock(&PlayersCS);
-				Players.Add(PlayerId, NewSession);
-			}
+			// This calls input/quality controller setters and any delegates that are required.
+			AddSession(PlayerId, NewSession);
+
+			// In FPlayerSessionDataOnly sessions the datachannel starts "open" and therefore we need to manually call state changed to
+			// fire the relevant calls when open happens, such as sending input/quality control messages.
+			NewSession->GetDataChannelObserver()->OnStateChange();
 		}
 	}
 
@@ -121,15 +162,56 @@ namespace UE::PixelStreaming
 		}
 
 		bool bWasQualityController = IsQualityController(PlayerId);
+		bool bWasInputController = IsInputController(PlayerId);
+		bool bInputRestrictingMode = UE::PixelStreaming::Settings::GetInputControllerMode() == UE::PixelStreaming::Settings::EInputControllerMode::Host;
 
 		// The actual modification to the players map
 		int RemainingCount = 0;
 		{
-			FScopeLock Lock(&PlayersCS);
-			Players.Remove(PlayerId);
-			RemainingCount = Players.Num();
+			// webrtc tries to call a method on the signalling thread when the peer connection
+			// is destroyed in the session destructor, but if the signalling thread is waiting
+			// for PlayersCS then this can cause a deadlock.
+			// This temporary is just to delay the actual deletion until AFTER we release the CS
+			TSharedPtr<IPlayerSession> ToBeDeleted;
+			{
+				FScopeLock Lock(&PlayersCS);
+				ToBeDeleted = Players[PlayerId];
+				Players.Remove(PlayerId);
+				RemainingCount = Players.Num();
+			}
+		}
+
+		if (bWasQualityController && RemainingCount > 0)
+		{
+			// Quality Controller session has been just removed, set quality control to
+			// any of remaining player sessions (not SFU!)
+			FPixelStreamingPlayerId OutPlayerId = INVALID_PLAYER_ID;
+			bool bHasPlayer = GetFirstP2PPeer(OutPlayerId);
+			if (bHasPlayer)
+			{
+				SetQualityController(OutPlayerId);
+			}
+		}
+
+		if (bInputRestrictingMode && bWasInputController && RemainingCount > 0)
+		{
+			// Input controlling session has been just removed, set the input controller to
+			// any of remaining player sessions (not SFU!)
+			FPixelStreamingPlayerId OutPlayerId = INVALID_PLAYER_ID;
+			bool bHasPlayer = GetFirstNonSFUPeer(OutPlayerId);
+			if (bHasPlayer)
+			{
+				SetInputController(OutPlayerId);
+			}
+		}
+
+		FStats* Stats = FStats::Get();
+		if (Stats)
+		{
+			Stats->RemovePeersStats(PlayerId);
 		}
 
+		// Calling of delegates should be last to avoid side effects such as setting stats for removed peers after the delegate is fired
 		UPixelStreamingDelegates* Delegates = UPixelStreamingDelegates::GetPixelStreamingDelegates();
 		if (Delegates && FModuleManager::Get().IsModuleLoaded("PixelStreaming"))
 		{
@@ -145,22 +227,20 @@ namespace UE::PixelStreaming
 			}
 		}
 
-		if (bWasQualityController && RemainingCount > 0)
-		{
-			// Quality Controller session has been just removed, set quality control to
-			// any of remaining player sessions (not SFU!)
-			FPixelStreamingPlayerId OutPlayerId = INVALID_PLAYER_ID;
-			bool bHasPlayer = GetFirstP2PPeer(OutPlayerId);
-			if (bHasPlayer)
-			{
-				SetQualityController(OutPlayerId);
-			}
-		}
-
 		return RemainingCount;
 	}
 
 	int FPlayerSessions::NumP2PPeers() const
+	{
+		return GetNumPeersMatching([](FName InSessionType) { return InSessionType == FPlayerSession::Type; });
+	}
+
+	int FPlayerSessions::NumNonSFUPeers() const
+	{
+		return GetNumPeersMatching([](FName InSessionType) { return InSessionType != FPlayerSessionSFU::Type; });
+	}
+
+	int FPlayerSessions::GetNumPeersMatching(TFunction<bool(FName)> SessionTypeMatchFunc) const
 	{
 		FScopeLock Lock(&PlayersCS);
 		int Count = 0;
@@ -168,7 +248,7 @@ namespace UE::PixelStreaming
 		{
 			FPixelStreamingPlayerId Id = Entry.Key;
 			TSharedPtr<IPlayerSession> Session = Entry.Value;
-			if (Session.IsValid() && Session->GetSessionType() == FPlayerSession::Type)
+			if (Session.IsValid() && SessionTypeMatchFunc(Session->GetSessionType()))
 			{
 				Count++;
 			}
@@ -176,7 +256,7 @@ namespace UE::PixelStreaming
 		return Count;
 	}
 
-	bool FPlayerSessions::GetFirstP2PPeer(FPixelStreamingPlayerId& OutPlayerId) const
+	bool FPlayerSessions::GetFirstPeerMatching(TFunction<bool(FName)> SessionTypeMatchFunc, FPixelStreamingPlayerId& OutPlayerId) const
 	{
 		FScopeLock Lock(&PlayersCS);
 
@@ -184,7 +264,7 @@ namespace UE::PixelStreaming
 		{
 			FPixelStreamingPlayerId Id = Entry.Key;
 			TSharedPtr<IPlayerSession> Session = Entry.Value;
-			if (Session.IsValid() && Session->GetSessionType() == FPlayerSession::Type)
+			if (Session.IsValid() && SessionTypeMatchFunc(Session->GetSessionType()))
 			{
 				OutPlayerId = Id;
 				return true;
@@ -193,6 +273,16 @@ namespace UE::PixelStreaming
 		return false;
 	}
 
+	bool FPlayerSessions::GetFirstP2PPeer(FPixelStreamingPlayerId& OutPlayerId) const
+	{
+		return GetFirstPeerMatching([](FName InSessionType) { return InSessionType == FPlayerSession::Type; }, OutPlayerId);
+	}
+
+	bool FPlayerSessions::GetFirstNonSFUPeer(FPixelStreamingPlayerId& OutPlayerId) const
+	{
+		return GetFirstPeerMatching([](FName InSessionType) { return InSessionType != FPlayerSessionSFU::Type; }, OutPlayerId);
+	}
+
 	void FPlayerSessions::DeleteAllPlayerSessions(bool bEnginePreExit)
 	{
 		const FPixelStreamingPlayerId OldQualityController = QualityControllingPlayer;
@@ -238,7 +328,6 @@ namespace UE::PixelStreaming
 
 	void FPlayerSessions::SetQualityController(FPixelStreamingPlayerId PlayerId)
 	{
-
 		TSharedPtr<IPlayerSession> PlayerSession = GetPlayerSession(PlayerId);
 
 		if (!PlayerSession.IsValid())
@@ -255,6 +344,18 @@ namespace UE::PixelStreaming
 		{
 			FScopeLock Lock(&QualityControllerCS);
 			QualityControllingPlayer = PlayerId;
+
+			UE_LOG(LogPixelStreaming, Log, TEXT("Quality controller is now PlayerId=%s."), *PlayerId);
+
+			// Update quality controller status on the browser side too
+			ForEachSession([&QualityControllingPlayer = QualityControllingPlayer](TSharedPtr<IPlayerSession> Session) {
+				bool bIsQualityController = Session->GetPlayerId() == QualityControllingPlayer;
+				if (Session->GetDataChannelObserver() && Session->GetDataChannelObserver()->IsDataChannelOpen())
+				{
+					// Only send message if datachannel is already open (in dc observer we also send this message when channel opens)
+					Session->SendQualityControlStatus(bIsQualityController);
+				}
+			});
 		}
 
 		// Let any listeners know the quality controller has changed
@@ -262,26 +363,87 @@ namespace UE::PixelStreaming
 		{
 			Delegates->OnQualityControllerChangedNative.Broadcast(PlayerId);
 		}
+	}
+
+	bool FPlayerSessions::IsInputController(FPixelStreamingPlayerId PlayerId) const
+	{
+		UE::PixelStreaming::Settings::EInputControllerMode InputControllerMode = PixelStreaming::Settings::GetInputControllerMode();
+
+		if (InputControllerMode == UE::PixelStreaming::Settings::EInputControllerMode::Host)
+		{
+			// Only in "Host" mode do we actually have to check if the peer can control input.
+			FScopeLock Lock(&InputControllerCS);
+			return InputControllingPlayer == PlayerId;
+		}
+		else
+		{
+			return true;
+		}
+	}
+
+	void FPlayerSessions::SetInputController(FPixelStreamingPlayerId InInputControllingPlayerId)
+	{
+		UE::PixelStreaming::Settings::EInputControllerMode InputControllerMode = PixelStreaming::Settings::GetInputControllerMode();
+
+		// No point setting input controller when running in "Any" mode where any peer can control input.
+		if (InputControllerMode == UE::PixelStreaming::Settings::EInputControllerMode::Any)
+		{
+			return;
+		}
+
+		TSharedPtr<IPlayerSession> PlayerSession = GetPlayerSession(InInputControllingPlayerId);
+
+		if (!PlayerSession.IsValid())
+		{
+			return;
+		}
+
+		// SFU cannot be the input controller
+		if (PlayerSession->GetSessionType() == FPlayerSessionSFU::Type)
+		{
+			return;
+		}
 
-		UE_LOG(LogPixelStreaming, Log, TEXT("Quality controller is now PlayerId=%s."), *PlayerId);
+		// The actual assignment of the input controlling peer
+		{
+			FScopeLock Lock(&InputControllerCS);
+			InputControllingPlayer = InInputControllingPlayerId;
+
+			UE_LOG(LogPixelStreaming, Log, TEXT("Input controller is now PlayerId=%s."), *InInputControllingPlayerId);
+
+
+			// Update quality controller status on the browser side too
+			ForEachSession([&InputControllingPlayer = InputControllingPlayer](TSharedPtr<IPlayerSession> Session) {
+				bool bIsInputController = Session->GetPlayerId() == InputControllingPlayer;
+				if (Session->GetDataChannelObserver() && Session->GetDataChannelObserver()->IsDataChannelOpen())
+				{
+					// Only send message if datachannel is already open (in dc observer we also send this message when channel opens)
+					Session->SendInputControlStatus(bIsInputController);
+				}
+			});
+		}
 
-		// Update quality controller status on the browser side too
-		ForEachSession([&](TSharedPtr<IPlayerSession> Session) {
-			bool bIsQualityController = Session->GetPlayerId() == QualityControllingPlayer;
-			Session->SendQualityControlStatus(bIsQualityController); });
+		// Todo (Luke): Call delegates to inform them about the new "host"
+		// if (UPixelStreamingDelegates* Delegates = UPixelStreamingDelegates::GetPixelStreamingDelegates())
+		// {
+		// 	Delegates->OnInputControllerChangedNative.Broadcast(PlayerId);
+		// }
 	}
 
 	void FPlayerSessions::ForEachSession(const TFunction<void(TSharedPtr<IPlayerSession>)>& Func)
 	{
 		// We have to be careful here
-		// We dont lock here because we can end up in a situation where Func
+		// We dont lock when calling `Func` because we can end up in a situation where Func
 		// dispatches to another thread and that thread is already waiting to acquire
 		// this lock and bam we have a deadlock. Instead we grab the current list of
 		// keys and then iterate over our copy of those. We DO lock while checking and
 		// pulling the session just in case Players is being modified between Contains
 		// and when we actually pull it. luckily GetPlayerSession handles that for us.
 		TSet<FPixelStreamingPlayerId> KeySet;
-		Players.GetKeys(KeySet);
+		{
+			FScopeLock Lock(&PlayersCS);
+			Players.GetKeys(KeySet);
+		}
 		for (auto&& PlayerId : KeySet)
 		{
 			if (TSharedPtr<IPlayerSession> Session = GetPlayerSession(PlayerId))
diff --git a/Engine/Plugins/Media/PixelStreaming/Source/PixelStreaming/Private/PlayerSessions.h b/Engine/Plugins/Media/PixelStreaming/Source/PixelStreaming/Private/PlayerSessions.h
--- a/Engine/Plugins/Media/PixelStreaming/Source/PixelStreaming/Private/PlayerSessions.h
+++ b/Engine/Plugins/Media/PixelStreaming/Source/PixelStreaming/Private/PlayerSessions.h
@@ -2,6 +2,9 @@
 #pragma once
 
 #include "IPlayerSession.h"
+#include "InputDevice.h"
+
+class FPixelStreamingSignallingConnection;
 
 namespace UE::PixelStreaming
 {
@@ -16,16 +19,18 @@ namespace UE::PixelStreaming
 		FPlayerSessions& operator=(const FPlayerSessions&) = delete;
 
 		TSharedPtr<IPlayerSession> CreatePlayerSession(FPixelStreamingPlayerId PlayerId,
+			TSharedPtr<IPixelStreamingInputDevice> InputDevice,
 			rtc::scoped_refptr<webrtc::PeerConnectionFactoryInterface> PeerConnectionFactory,
 			webrtc::PeerConnectionInterface::RTCConfiguration PeerConnectionConfig,
-			FSignallingServerConnection* SignallingServerConnection,
+			FPixelStreamingSignallingConnection* SignallingServerConnection,
 			int Flags);
 
-		void CreateNewDataChannel(FPixelStreamingPlayerId SFUId,
+		void CreateSFUDataOnlyPeer(FPixelStreamingPlayerId SFUId,
 			FPixelStreamingPlayerId PlayerId,
+			TSharedPtr<IPixelStreamingInputDevice> InputDevice,
 			int32 SendStreamId,
 			int32 RecvStreamId,
-			FSignallingServerConnection* SignallingServerConnection);
+			FPixelStreamingSignallingConnection* SignallingServerConnection);
 
 		int DeletePlayerSession(FPixelStreamingPlayerId PlayerId);
 		void DeleteAllPlayerSessions(bool bEnginePreExit = false);
@@ -34,6 +39,8 @@ namespace UE::PixelStreaming
 
 		bool IsQualityController(FPixelStreamingPlayerId PlayerId) const;
 		void SetQualityController(FPixelStreamingPlayerId PlayerId);
+		bool IsInputController(FPixelStreamingPlayerId PlayerId) const;
+		void SetInputController(FPixelStreamingPlayerId PlayerId);
 
 		// do something with a single session
 		void ForSession(FPixelStreamingPlayerId PlayerId, const TFunction<void(TSharedPtr<IPlayerSession>)>& Func) const
@@ -66,8 +73,16 @@ namespace UE::PixelStreaming
 			webrtc::PeerConnectionInterface::RTCConfiguration PeerConnectionConfig,
 			TSharedPtr<IPlayerSession> Session) const;
 
+		rtc::scoped_refptr<webrtc::DataChannelInterface> CreateDataChannel(
+			rtc::scoped_refptr<webrtc::PeerConnectionInterface> PeerConnection,
+			bool bIsSFU) const;
+
+		int GetNumPeersMatching(TFunction<bool(FName)> SessionTypeMatchFunc) const;
+		int NumNonSFUPeers() const;
 		int NumP2PPeers() const;
+		bool GetFirstPeerMatching(TFunction<bool(FName)> SessionTypeMatchFunc, FPixelStreamingPlayerId& OutPlayerId) const;
 		bool GetFirstP2PPeer(FPixelStreamingPlayerId& OutPlayerId) const;
+		bool GetFirstNonSFUPeer(FPixelStreamingPlayerId& OutPlayerId) const;
 		TSharedPtr<IPlayerSession> GetPlayerSession(FPixelStreamingPlayerId PlayerId) const;
 
 		mutable FCriticalSection PlayersCS;
@@ -75,5 +90,8 @@ namespace UE::PixelStreaming
 
 		mutable FCriticalSection QualityControllerCS;
 		FPixelStreamingPlayerId QualityControllingPlayer = INVALID_PLAYER_ID;
+
+		mutable FCriticalSection InputControllerCS;
+		FPixelStreamingPlayerId InputControllingPlayer = INVALID_PLAYER_ID;
 	};
 } // namespace UE::PixelStreaming
diff --git a/Engine/Plugins/Media/PixelStreaming/Source/PixelStreaming/Private/Settings.cpp b/Engine/Plugins/Media/PixelStreaming/Source/PixelStreaming/Private/Settings.cpp
--- a/Engine/Plugins/Media/PixelStreaming/Source/PixelStreaming/Private/Settings.cpp
+++ b/Engine/Plugins/Media/PixelStreaming/Source/PixelStreaming/Private/Settings.cpp
@@ -226,12 +226,6 @@ namespace UE::PixelStreaming::Settings
 		TEXT("If true disables latency tester being triggerable."),
 		ECVF_Default);
 
-	TAutoConsoleVariable<bool> CVarPixelStreamingDecoupleFrameRate(
-		TEXT("PixelStreaming.DecoupleFrameRate"),
-		true,
-		TEXT("If true, Pixel Streaming will send 60fps to WebRTC regardless of actual application framerate"),
-		ECVF_Default);
-
 	TAutoConsoleVariable<FString> CVarPixelStreamingKeyFilter(
 		TEXT("PixelStreaming.KeyFilter"),
 		"",
@@ -256,8 +250,16 @@ namespace UE::PixelStreaming::Settings
 		TEXT("Comma separated list of video track types to create when a peer joins. Default: Backbuffer"),
 		ECVF_Default);
 
+	TAutoConsoleVariable<FString> CVarPixelStreamingInputController(
+		TEXT("PixelStreaming.InputController"),
+		TEXT("Any"),
+		TEXT("Various modes of input control supported by Pixel Streaming, currently: \"Any\"  or \"Host\". Default: Any"),
+		ECVF_Default);
+
 	TArray<FKey> FilteredKeys;
 	TArray<FName> ActiveTextureSourceTypes;
+	bool bDecoupleFrameRate = true;
+	FString DefaultStreamerID = TEXT("DefaultStreamer");
 
 	void OnActiveTextureSourcesChanged(IConsoleVariable* Var)
 	{
@@ -287,6 +289,11 @@ namespace UE::PixelStreaming::Settings
 	{
 		bool bHudStatsEnabled = Var->GetBool();
 
+		if (!GEngine)
+		{
+			return;
+		}
+
 		for (const FWorldContext& WorldContext : GEngine->GetWorldContexts())
 		{
 			if (WorldContext.WorldType == EWorldType::Game || WorldContext.WorldType == EWorldType::PIE)
@@ -465,12 +472,37 @@ namespace UE::PixelStreaming::Settings
 		}
 	}
 
+	bool DecoupleFrameRate()
+	{
+		return bDecoupleFrameRate;
+	}
+
 	bool IsCodecVPX()
 	{
 		ECodec SelectedCodec = GetSelectedCodec();
 		return SelectedCodec == ECodec::VP8 || SelectedCodec == ECodec::VP9;
 	}
 
+	EInputControllerMode GetInputControllerMode()
+	{
+		// Convert the current value to all lowercase and remove any whitespace.
+		const FString InputControllerMode = CVarPixelStreamingInputController.GetValueOnAnyThread().ToLower().TrimStartAndEnd();
+
+		if (InputControllerMode == TEXT("host"))
+		{
+			return EInputControllerMode::Host;
+		}
+		else
+		{
+			return EInputControllerMode::Any;
+		}
+	}
+
+	FString GetDefaultStreamerID()
+	{
+		return DefaultStreamerID;
+	}
+
 	/*
 	 * Stats logger - as turned on/off by CVarPixelStreamingLogStats
 	 */
@@ -516,6 +548,11 @@ namespace UE::PixelStreaming::Settings
 	{
 		UE_LOG(LogPixelStreaming, Log, TEXT("Initialising Pixel Streaming settings."));
 
+		FString StringOptions;
+		bDecoupleFrameRate = !FParse::Value(FCommandLine::Get(), TEXT("PixelStreamingDisableDecoupleFrameRate"), StringOptions, false);
+		
+		FParse::Value(FCommandLine::Get(), TEXT("PixelStreamingID="), DefaultStreamerID, false);
+
 		CVarPixelStreamingOnScreenStats.AsVariable()->SetOnChangedCallback(FConsoleVariableDelegate::CreateStatic(&OnHudStatsToggled));
 		CVarPixelStreamingKeyFilter.AsVariable()->SetOnChangedCallback(FConsoleVariableDelegate::CreateStatic(&OnFilteredKeysChanged));
 		CVarPixelStreamingVideoTracks.AsVariable()->SetOnChangedCallback(FConsoleVariableDelegate::CreateStatic(&OnActiveTextureSourcesChanged));
@@ -545,12 +582,12 @@ namespace UE::PixelStreaming::Settings
 		CommandLineParseValue(TEXT("PixelStreamingKeyFilter="), CVarPixelStreamingKeyFilter);
 		CommandLineParseValue(TEXT("PixelStreamingVideoTracks="), CVarPixelStreamingVideoTracks);
 		CommandLineParseValue(TEXT("PixelStreamingFrameScale="), CVarPixelStreamingFrameScale);
+		CommandLineParseValue(TEXT("PixelStreamingInputController="), CVarPixelStreamingInputController);
 
 		// Options parse (if these exist they are set to true)
 		CommandLineParseOption(TEXT("AllowPixelStreamingCommands"), CVarPixelStreamingAllowConsoleCommands);
 		CommandLineParseOption(TEXT("PixelStreamingOnScreenStats"), CVarPixelStreamingOnScreenStats);
 		CommandLineParseOption(TEXT("PixelStreamingHudStats"), CVarPixelStreamingOnScreenStats);
-
 		CommandLineParseOption(TEXT("PixelStreamingDebugDumpFrame"), CVarPixelStreamingDebugDumpFrame);
 		CommandLineParseOption(TEXT("PixelStreamingEnableFillerData"), CVarPixelStreamingEnableFillerData);
 		CommandLineParseOption(TEXT("PixelStreamingWebRTCDisableStats"), CVarPixelStreamingWebRTCDisableStats);
@@ -560,7 +597,6 @@ namespace UE::PixelStreaming::Settings
 		CommandLineParseOption(TEXT("PixelStreamingSendPlayerIdAsInteger"), CVarSendPlayerIdAsInteger);
 		CommandLineParseOption(TEXT("PixelStreamingWebRTCUseLegacyAudioDevice"), CVarPixelStreamingWebRTCUseLegacyAudioDevice);
 		CommandLineParseOption(TEXT("PixelStreamingDisableLatencyTester"), CVarPixelStreamingDisableLatencyTester);
-		CommandLineParseOption(TEXT("PixelStreamingDecoupleFrameRate"), CVarPixelStreamingDecoupleFrameRate);
 		CommandLineParseOption(TEXT("PixelStreamingVPXUseCompute"), CVarPixelStreamingVPXUseCompute);
 
 		ReadSimulcastParameters();
@@ -570,9 +606,9 @@ namespace UE::PixelStreaming::Settings
 
 		IConsoleManager::Get().RegisterConsoleCommand(
 			TEXT("PixelStreaming.StartStreaming"),
-			TEXT("<SignallingServer URL> - Connect to a signalling server and begin a streaming session."),
-			FConsoleCommandWithArgsDelegate::CreateLambda([](const TArray<FString>& Args) {
-				IPixelStreamingModule::Get().StartStreaming(Args[0].TrimQuotes());
+			TEXT("Start all streaming sessions"),
+			FConsoleCommandDelegate::CreateLambda([]() {
+				IPixelStreamingModule::Get().StartStreaming();
 			}));
 
 		IConsoleManager::Get().RegisterConsoleCommand(
diff --git a/Engine/Plugins/Media/PixelStreaming/Source/PixelStreaming/Private/Settings.h b/Engine/Plugins/Media/PixelStreaming/Source/PixelStreaming/Private/Settings.h
--- a/Engine/Plugins/Media/PixelStreaming/Source/PixelStreaming/Private/Settings.h
+++ b/Engine/Plugins/Media/PixelStreaming/Source/PixelStreaming/Private/Settings.h
@@ -51,8 +51,8 @@ namespace UE::PixelStreaming::Settings
 	extern TAutoConsoleVariable<int32> CVarPixelStreamingFreezeFrameQuality;
 	extern TAutoConsoleVariable<bool> CVarSendPlayerIdAsInteger;
 	extern TAutoConsoleVariable<bool> CVarPixelStreamingDisableLatencyTester;
-	extern TAutoConsoleVariable<bool> CVarPixelStreamingDecoupleFrameRate;
 	extern TAutoConsoleVariable<bool> CVarPixelStreamingVPXUseCompute;
+	extern TAutoConsoleVariable<FString> CVarPixelStreamingInputController;
 
 	extern TArray<FKey> FilteredKeys;
 	// Ends Pixel Streaming Plugin CVars
@@ -68,7 +68,17 @@ namespace UE::PixelStreaming::Settings
 		VP9
 	};
 
+	/* Pixel Streaming can limit who can send input (keyboard, mouse, etc). */
+	enum EInputControllerMode
+	{
+		/* Any - Any peer can control input. */
+		Any,
+		/* Host - Only the "host" peer can control input. */
+		Host
+	};
+
 	// Begin utility functions etc.
+	bool DecoupleFrameRate();
 	bool IsCodecVPX();
 	ECodec GetSelectedCodec();
 	AVEncoder::FVideoEncoder::RateControlMode GetRateControlCVar();
@@ -77,6 +87,9 @@ namespace UE::PixelStreaming::Settings
 	AVEncoder::FVideoEncoder::H264Profile GetH264Profile();
 	const TArray<FName>& GetActiveTextureSourceTypes();
 	void SetActiveTextureSourceTypes(TArray<FName> InActiveTextureSourceTypes);
+	EInputControllerMode GetInputControllerMode();
+	FString GetDefaultStreamerID();
+
 	// End utility functions etc.
 
 	struct FSimulcastParameters
