Early Pixel Streaming 5.1 features, bugfixes, and refactors that benefit main. (From various authors)

diff --git a/Engine/Plugins/Media/PixelStreaming/Source/PixelStreaming/Private/Stats.cpp b/Engine/Plugins/Media/PixelStreaming/Source/PixelStreaming/Private/Stats.cpp
--- a/Engine/Plugins/Media/PixelStreaming/Source/PixelStreaming/Private/Stats.cpp
+++ b/Engine/Plugins/Media/PixelStreaming/Source/PixelStreaming/Private/Stats.cpp
@@ -9,6 +9,7 @@
 #include "IPixelStreamingStatsConsumer.h"
 #include "Engine/Console.h"
 #include "ConsoleSettings.h"
+#include "PixelStreamingFrameMetadata.h"
 
 namespace UE::PixelStreaming
 {
@@ -16,16 +17,26 @@ namespace UE::PixelStreaming
 
 	FStats* FStats::Get()
 	{
-		return FStats::Instance;
+		if (Instance == nullptr)
+		{
+			Instance = new FStats();
+		}
+		return Instance;
+	}
+
+	void FStats::AddSessions(FPlayerSessions* InSessions)
+	{
+		SessionsList.Add(InSessions);
 	}
 
-	FStats::FStats(FPlayerSessions* InSessions)
-		: Sessions(InSessions)
+	void FStats::RemoveSessions(FPlayerSessions* InSessions)
 	{
-		checkf(FStats::Instance == nullptr, TEXT("There should only ever been one PixelStreaming stats object."));
-		checkf(Sessions, TEXT("To make stats object the sessions object must not be nullptr."));
-		FStats::Instance = this;
+		SessionsList.Remove(InSessions);
+	}
 
+	FStats::FStats()
+	{
+		checkf(Instance == nullptr, TEXT("There should only ever been one PixelStreaming stats object."));
 		UConsole::RegisterConsoleAutoCompleteEntries.AddRaw(this, &FStats::UpdateConsoleAutoComplete_GameThread);
 	}
 
@@ -187,6 +198,19 @@ namespace UE::PixelStreaming
 		return PeerStats[PlayerId].StoreStat_GameThread(Stat);
 	}
 
+	double CalcMA(double PrevAvg, int NumSamples, double Value)
+	{
+		const double Result = NumSamples * PrevAvg + Value;
+		return Result / (PrevAvg + 1.0);
+	}
+
+	double CalcEMA(double PrevAvg, int NumSamples, double Value)
+	{
+		const double Mult = 2.0 / (NumSamples + 1.0);
+		const double Result = (Value - PrevAvg) * Mult + PrevAvg;
+		return Result;
+	}
+
 	bool FStats::StoreApplicationStat_GameThread(FStatData Stat)
 	{
 		checkf(IsInGameThread(), TEXT("This method must be called from the game thread"));
@@ -204,17 +228,18 @@ namespace UE::PixelStreaming
 
 			if (Stat.bSmooth && RenderableStat->Stat.StatValue != 0)
 			{
-				double CurValue = RenderableStat->Stat.StatValue;
-				double PercentageDrift = FMath::Abs(CurValue - Stat.StatValue) / CurValue;
-				if (PercentageDrift > FStats::SmoothingFactor)
+				const int MaxSamples = 60;
+				RenderableStat->Stat.NumSamples = FGenericPlatformMath::Min(MaxSamples, RenderableStat->Stat.NumSamples + 1);
+				if (RenderableStat->Stat.NumSamples < MaxSamples)
 				{
-					RenderableStat->Stat.StatValue = Stat.StatValue;
-					bUpdated = true;
+					RenderableStat->Stat.LastEMA = CalcMA(RenderableStat->Stat.LastEMA, RenderableStat->Stat.NumSamples - 1, Stat.StatValue);
 				}
 				else
 				{
-					return false;
+					RenderableStat->Stat.LastEMA = CalcEMA(RenderableStat->Stat.LastEMA, RenderableStat->Stat.NumSamples - 1, Stat.StatValue);
 				}
+				RenderableStat->Stat.StatValue = RenderableStat->Stat.LastEMA;
+				bUpdated = true;
 			}
 			else
 			{
@@ -279,12 +304,16 @@ namespace UE::PixelStreaming
 		const UConsoleSettings* ConsoleSettings = GetDefault<UConsoleSettings>();
 
 		AutoCompleteList.AddDefaulted();
-
 		FAutoCompleteCommand& AutoCompleteCommand = AutoCompleteList.Last();
-
 		AutoCompleteCommand.Command = TEXT("Stat PixelStreaming");
 		AutoCompleteCommand.Desc = TEXT("Displays stats about Pixel Streaming on screen.");
 		AutoCompleteCommand.Color = ConsoleSettings->AutoCompleteCommandColor;
+
+		AutoCompleteList.AddDefaulted();
+		FAutoCompleteCommand& AutoCompleteGraphCommand = AutoCompleteList.Last();
+		AutoCompleteGraphCommand.Command = TEXT("Stat PixelStreamingGraphs");
+		AutoCompleteGraphCommand.Desc = TEXT("Displays graphs about Pixel Streaming on screen.");
+		AutoCompleteGraphCommand.Color = ConsoleSettings->AutoCompleteCommandColor;
 	}
 
 	int32 FStats::OnRenderStats(UWorld* World, FViewport* Viewport, FCanvas* Canvas, int32 X, int32 Y, const FVector* ViewLocation, const FRotator* ViewRotation)
@@ -350,6 +379,45 @@ namespace UE::PixelStreaming
 		return true;
 	}
 
+	bool FStats::OnToggleGraphs(UWorld* World, FCommonViewportClient* ViewportClient, const TCHAR* Stream)
+	{
+		return true;
+	}
+
+	int32 FStats::OnRenderGraphs(UWorld* World, FViewport* Viewport, FCanvas* Canvas, int32 X, int32 Y, const FVector* ViewLocation, const FRotator* ViewRotation)
+	{
+		FVector2D GraphPos{ 0, 0 };
+		FVector2D GraphSize{ 200, 200 };
+		float GraphSpacing = 5;
+
+		for (auto& [GraphName, Graph] : Graphs)
+		{
+			Graph.Draw(Canvas, GraphPos, GraphSize);
+			GraphPos.X += GraphSize.X + GraphSpacing;
+			if (GraphPos.X > Canvas->GetParentCanvasSize().X)
+			{
+				GraphPos.Y += GraphSize.Y + GraphSpacing;
+				GraphPos.X = 0;
+			}
+		}
+
+		for (auto& [TileName, Tile] : Tiles)
+		{
+			Tile.Position.X = GraphPos.X;
+			Tile.Position.Y = GraphPos.Y;
+			Tile.Size = GraphSize;
+			Tile.Draw(Canvas);
+			GraphPos.X += GraphSize.X + GraphSpacing;
+			if (GraphPos.X > Canvas->GetParentCanvasSize().X)
+			{
+				GraphPos.Y += GraphSize.Y + GraphSpacing;
+				GraphPos.X = 0;
+			}
+		}
+
+		return Y;
+	}
+
 	void FStats::PollPixelStreamingSettings()
 	{
 		double DeltaSeconds = FPlatformTime::ToSeconds64(FPlatformTime::Cycles64() - LastTimeSettingsPolledCycles);
@@ -387,9 +455,12 @@ namespace UE::PixelStreaming
 		// webrtc::metrics::Histogram* Hist1 = webrtc::metrics::HistogramFactoryGetCounts("WebRTC.Video.NacksSent", 0, 100000, 100);
 		// Will require calling webrtc::metrics::Enable();
 
-		Sessions->ForEachSession([](TSharedPtr<IPlayerSession> Session) {
-			Session->PollWebRTCStats();
-		});
+		for (auto& Sessions : SessionsList)
+		{
+			Sessions->ForEachSession([](TSharedPtr<IPlayerSession> Session) {
+				Session->PollWebRTCStats();
+			});
+		}
 
 		PollPixelStreamingSettings();
 
@@ -400,12 +471,42 @@ namespace UE::PixelStreaming
 
 		if (!bRegisterEngineStats)
 		{
-			GAreScreenMessagesEnabled = true;
-			UEngine::FEngineStatRender RenderFunc = UEngine::FEngineStatRender::CreateRaw(this, &FStats::OnRenderStats);
-			UEngine::FEngineStatToggle ToggleFunc = UEngine::FEngineStatToggle::CreateRaw(this, &FStats::OnToggleStats);
-			GEngine->AddEngineStat(PixelStreamingStatName, PixelStreamingStatCategory, PixelStreamingStatDescription, RenderFunc, ToggleFunc, false);
-			bRegisterEngineStats = true;
+			RegisterEngineHooks();
+		}
+	}
+
+	void FStats::RegisterEngineHooks()
+	{
+		GAreScreenMessagesEnabled = true;
+
+		const FName StatName("STAT_PixelStreaming");
+		const FName StatCategory("STATCAT_PixelStreaming");
+		const FText StatDescription(FText::FromString("Pixel Streaming stats for all connected peers."));
+		UEngine::FEngineStatRender RenderStatFunc = UEngine::FEngineStatRender::CreateRaw(this, &FStats::OnRenderStats);
+		UEngine::FEngineStatToggle ToggleStatFunc = UEngine::FEngineStatToggle::CreateRaw(this, &FStats::OnToggleStats);
+		GEngine->AddEngineStat(StatName, StatCategory, StatDescription, RenderStatFunc, ToggleStatFunc, false);
+
+		const FName GraphName("STAT_PixelStreamingGraphs");
+		const FText GraphDescription(FText::FromString("Pixel Streaming graphs showing frame pipeline timings."));
+		UEngine::FEngineStatRender RenderGraphFunc = UEngine::FEngineStatRender::CreateRaw(this, &FStats::OnRenderGraphs);
+		UEngine::FEngineStatToggle ToggleGraphFunc = UEngine::FEngineStatToggle::CreateRaw(this, &FStats::OnToggleGraphs);
+		GEngine->AddEngineStat(GraphName, StatCategory, GraphDescription, RenderGraphFunc, ToggleGraphFunc, false);
+
+		bool StatsEnabled = Settings::CVarPixelStreamingOnScreenStats.GetValueOnAnyThread();
+		if (StatsEnabled)
+		{
+			for (const FWorldContext& WorldContext : GEngine->GetWorldContexts())
+			{
+				if (WorldContext.WorldType == EWorldType::Game || WorldContext.WorldType == EWorldType::PIE)
+				{
+					UWorld* World = WorldContext.World();
+					UGameViewportClient* ViewportClient = World->GetGameViewport();
+					GEngine->SetEngineStat(World, ViewportClient, TEXT("PixelStreaming"), StatsEnabled);
+				}
+			}
 		}
+
+		bRegisterEngineStats = true;
 	}
 
 	//
@@ -468,17 +569,14 @@ namespace UE::PixelStreaming
 
 			if (RenderableStat->Stat.bSmooth && RenderableStat->Stat.StatValue != 0)
 			{
-				double CurValue = RenderableStat->Stat.StatValue;
-				double PercentageDrift = FMath::Abs(CurValue - StatToStore.StatValue) / CurValue;
-				if (PercentageDrift > FStats::SmoothingFactor)
-				{
-					RenderableStat->Stat.StatValue = StatToStore.StatValue;
-					bUpdated = true;
-				}
+				const int MaxSamples = 60;
+				RenderableStat->Stat.NumSamples = FGenericPlatformMath::Min(MaxSamples, RenderableStat->Stat.NumSamples + 1);
+				if (RenderableStat->Stat.NumSamples < MaxSamples)
+					RenderableStat->Stat.LastEMA = CalcMA(RenderableStat->Stat.LastEMA, RenderableStat->Stat.NumSamples - 1, StatToStore.StatValue);
 				else
-				{
-					return false;
-				}
+					RenderableStat->Stat.LastEMA = CalcEMA(RenderableStat->Stat.LastEMA, RenderableStat->Stat.NumSamples - 1, StatToStore.StatValue);
+				RenderableStat->Stat.StatValue = RenderableStat->Stat.LastEMA;
+				bUpdated = true;
 			}
 			else
 			{
@@ -519,4 +617,53 @@ namespace UE::PixelStreaming
 			return bUpdated;
 		}
 	}
+
+	void FStats::GraphValue(FName InName, float Value, int InSamples, float InMinRange, float InMaxRange, float InRefValue)
+	{
+		if (!Graphs.Contains(InName))
+		{
+			auto& Graph = Graphs.Add(InName, FDebugGraph(InName, InSamples, InMinRange, InMaxRange, InRefValue));
+			Graph.AddValue(Value);
+		}
+		else
+		{
+			Graphs[InName].AddValue(Value);
+		}
+	}
+
+	double FStats::AddTimeDeltaStat(uint64 Cycles1, uint64 Cycles2, const FString& Label)
+	{
+		const uint64 MaxCycles = FGenericPlatformMath::Max(Cycles1, Cycles2);
+		const uint64 MinCycles = FGenericPlatformMath::Min(Cycles1, Cycles2);
+		const double DeltaMs = FPlatformTime::ToMilliseconds64(MaxCycles - MinCycles) * ((Cycles1 > Cycles2) ? 1.0 : -1.0);
+		const FStatData TimeData{ FName(*Label), DeltaMs, 2, true };
+		StoreApplicationStat(TimeData);
+		return DeltaMs;
+	}
+
+	void FStats::AddFrameTimingStats(const FPixelStreamingFrameMetadata& FrameMetadata)
+	{
+		const double TimePending = AddTimeDeltaStat(FrameMetadata.AdaptProcessStartTime, FrameMetadata.AdaptCallTime, FString::Printf(TEXT("%s Layer %d Frame Adapt Pending Time"), *FrameMetadata.ProcessName, FrameMetadata.Layer));
+		const double TimeGPU = AddTimeDeltaStat(FrameMetadata.AdaptProcessFinalizeTime, FrameMetadata.AdaptProcessStartTime, FString::Printf(TEXT("%s Layer %d Frame Adapt GPU Time"), *FrameMetadata.ProcessName, FrameMetadata.Layer));
+		const double TimeCPU = AddTimeDeltaStat(FrameMetadata.AdaptProcessEndTime, FrameMetadata.AdaptProcessFinalizeTime, FString::Printf(TEXT("%s Layer %d Frame Adapt CPU Time"), *FrameMetadata.ProcessName, FrameMetadata.Layer));
+		const double TimeWait = AddTimeDeltaStat(FrameMetadata.FirstEncodeStartTime, FrameMetadata.AdaptProcessEndTime, FString::Printf(TEXT("%s Layer %d Frame Wait Time"), *FrameMetadata.ProcessName, FrameMetadata.Layer));
+		const double TimeEncode = AddTimeDeltaStat(FrameMetadata.LastEncodeEndTime, FrameMetadata.LastEncodeStartTime, FString::Printf(TEXT("%s Layer %d Frame Encode Time"), *FrameMetadata.ProcessName, FrameMetadata.Layer));
+
+		const FStatData UseData{ FName(*FString::Printf(TEXT("%s Layer %d Frame Uses"), *FrameMetadata.ProcessName, FrameMetadata.Layer)), static_cast<double>(FrameMetadata.UseCount), 0, false };
+		StoreApplicationStat(UseData);
+
+		const int Samples = 100;
+		GraphValue(*FString::Printf(TEXT("%d Frame Lifetime"), FrameMetadata.Layer), StaticCast<float>(TimePending + TimeGPU + TimeCPU + TimeEncode), Samples, 0.0f, 30.0f, 16.66f);
+		GraphValue(*FString::Printf(TEXT("%d Pending Time"), FrameMetadata.Layer), StaticCast<float>(TimePending), Samples, 0.0f, 30.0f);
+		GraphValue(*FString::Printf(TEXT("%d GPU Time"), FrameMetadata.Layer), StaticCast<float>(TimeGPU), Samples, 0.0f, 6.0f);
+		GraphValue(*FString::Printf(TEXT("%d CPU Time"), FrameMetadata.Layer), StaticCast<float>(TimeCPU), Samples, 0.0f, 6.0f);
+		GraphValue(*FString::Printf(TEXT("%d Wait Time"), FrameMetadata.Layer), StaticCast<float>(TimeWait), Samples, 0.0f, 30.0f);
+		GraphValue(*FString::Printf(TEXT("%d Encode Time"), FrameMetadata.Layer), StaticCast<float>(TimeEncode), Samples, 0.0f, 10.0f);
+		GraphValue(*FString::Printf(TEXT("%d Frame Uses"), FrameMetadata.Layer), StaticCast<float>(FrameMetadata.UseCount), Samples, 0.0f, 10.0f);
+	}
+
+	void FStats::AddCanvasTile(FName Name, const FCanvasTileItem& Tile)
+	{
+		Tiles.FindOrAdd(Name, Tile);
+	}
 } // namespace UE::PixelStreaming
diff --git a/Engine/Plugins/Media/PixelStreaming/Source/PixelStreaming/Private/Stats.h b/Engine/Plugins/Media/PixelStreaming/Source/PixelStreaming/Private/Stats.h
--- a/Engine/Plugins/Media/PixelStreaming/Source/PixelStreaming/Private/Stats.h
+++ b/Engine/Plugins/Media/PixelStreaming/Source/PixelStreaming/Private/Stats.h
@@ -7,8 +7,10 @@
 #include "CanvasItem.h"
 #include "UnrealEngine.h"
 #include "ConsoleSettings.h"
+#include "DebugGraph.h"
 
 class IPixelStreamingStatsConsumer;
+struct FPixelStreamingFrameMetadata;
 
 namespace UE::PixelStreaming
 {
@@ -22,7 +24,9 @@ namespace UE::PixelStreaming
 			: StatName(InStatName)
 			, StatValue(InStatValue)
 			, NDecimalPlacesToPrint(InNDecimalPlacesToPrint)
-			, bSmooth(bInSmooth) {}
+			, bSmooth(bInSmooth)
+			 {
+			 }
 
 		bool operator==(const FStatData& Other) const
 		{
@@ -38,6 +42,8 @@ namespace UE::PixelStreaming
 		double StatValue;
 		int NDecimalPlacesToPrint;
 		bool bSmooth;
+		double LastEMA = 0;
+		int NumSamples = 0;
 	};
 
 	FORCEINLINE uint32 GetTypeHash(const FStatData& Obj)
@@ -91,7 +97,9 @@ namespace UE::PixelStreaming
 		static constexpr double SmoothingFactor = 10.0 / 100.0;
 		static FStats* Get();
 
-		FStats(FPlayerSessions* Sessions);
+		void AddSessions(FPlayerSessions* InSessions);
+		void RemoveSessions(FPlayerSessions* InSessions);
+
 		FStats(const FStats&) = delete;
 		void QueryPeerStat(FPixelStreamingPlayerId PlayerId, FName StatToQuery, TFunction<void(bool, double)> QueryCallback) const;
 		bool QueryPeerStat_GameThread(FPixelStreamingPlayerId PlayerId, FName StatToQuery, double& OutStatValue) const;
@@ -100,12 +108,25 @@ namespace UE::PixelStreaming
 		void StoreApplicationStat(FStatData PeerStat);
 		void Tick(float DeltaTime);
 		void AddOnPeerStatChangedCallback(FPixelStreamingPlayerId PlayerId, FName StatToListenOn, TWeakPtr<IPixelStreamingStatsConsumer> Callback);
-		int32 OnRenderStats(UWorld* World, FViewport* Viewport, FCanvas* Canvas, int32 X, int32 Y, const FVector* ViewLocation, const FRotator* ViewRotation);
+
 		bool OnToggleStats(UWorld* World, FCommonViewportClient* ViewportClient, const TCHAR* Stream);
+		int32 OnRenderStats(UWorld* World, FViewport* Viewport, FCanvas* Canvas, int32 X, int32 Y, const FVector* ViewLocation, const FRotator* ViewRotation);
+
+		bool OnToggleGraphs(UWorld* World, FCommonViewportClient* ViewportClient, const TCHAR* Stream);
+		int32 OnRenderGraphs(UWorld* World, FViewport* Viewport, FCanvas* Canvas, int32 X, int32 Y, const FVector* ViewLocation, const FRotator* ViewRotation);
 
 		FORCEINLINE TStatId GetStatId() const { RETURN_QUICK_DECLARE_CYCLE_STAT(PixelStreamingStats, STATGROUP_Tickables); }
 
+		void GraphValue(FName InName, float Value, int InSamples, float InMinRange, float InMaxRange, float InRefValue = 0.0f);
+		
+		double AddTimeDeltaStat(uint64 Cycles1, uint64 Cycles2, const FString& Label);
+		void AddFrameTimingStats(const FPixelStreamingFrameMetadata& FrameMetadata);
+
+		void AddCanvasTile(FName Name, const FCanvasTileItem& Tile);
+
 	private:
+		FStats();
+		void RegisterEngineHooks();
 		void PollPixelStreamingSettings();
 		void RemovePeerStat_GameThread(FPixelStreamingPlayerId PlayerId);
 		bool StorePeerStat_GameThread(FPixelStreamingPlayerId PlayerId, FStatData Stat);
@@ -117,16 +138,16 @@ namespace UE::PixelStreaming
 	private:
 		static FStats* Instance;
 
-		FPlayerSessions* Sessions;
+		TArray<FPlayerSessions*> SessionsList;
 
 		bool bRegisterEngineStats = false;
-		FName PixelStreamingStatName = FName(TEXT("STAT_PixelStreaming"));
-		FName PixelStreamingStatCategory = FName(TEXT("STATCAT_PixelStreaming"));
-		FText PixelStreamingStatDescription = FText::FromString(FString(TEXT("Pixel Streaming stats for all connected peers.")));
 
 		TMap<FPixelStreamingPlayerId, FPeerStats> PeerStats;
 		TMap<FName, FRenderableStat> ApplicationStats;
 
 		int64 LastTimeSettingsPolledCycles = 0;
+
+		TMap<FName, FDebugGraph> Graphs;
+		TMap<FName, FCanvasTileItem> Tiles;
 	};
 } // namespace UE::PixelStreaming
diff --git a/Engine/Plugins/Media/PixelStreaming/Source/PixelStreaming/Private/Streamer.cpp b/Engine/Plugins/Media/PixelStreaming/Source/PixelStreaming/Private/Streamer.cpp
--- a/Engine/Plugins/Media/PixelStreaming/Source/PixelStreaming/Private/Streamer.cpp
+++ b/Engine/Plugins/Media/PixelStreaming/Source/PixelStreaming/Private/Streamer.cpp
@@ -5,234 +5,277 @@
 #include "VideoEncoderFactory.h"
 #include "ToStringExtensions.h"
 #include "AudioCapturer.h"
-#include "EncoderFactory.h"
+#include "VideoEncoderFactorySimple.h"
+#include "VideoEncoderFactorySimulcast.h"
 #include "Settings.h"
 #include "WebRTCLogging.h"
 #include "WebSocketsModule.h"
-#include "AudioDeviceModule.h"
-#include "SessionDescriptionObservers.h"
-#include "VideoSource.h"
+#include "PixelStreamingAudioDeviceModule.h"
+#include "VideoSourceP2P.h"
+#include "VideoSourceSFU.h"
 #include "PixelStreamingDelegates.h"
 #include "DataChannelObserver.h"
 #include "IPixelStreamingModule.h"
+#include "VideoSourceGroup.h"
+#include "InputDevice.h"
+#include "Misc/CoreDelegates.h"
+#include "UtilsRender.h"
+#include "IImageWrapper.h"
+#include "IImageWrapperModule.h"
+#include "PixelStreamingSignallingConnection.h"
+#include "PixelStreamingPrivate.h"
+#include "Async/Async.h"
+#include "PixelStreamingSourceFrame.h"
 
 namespace UE::PixelStreaming
 {
-	bool FStreamer::IsPlatformCompatible()
-	{
-		return AVEncoder::FVideoEncoderFactory::Get().HasEncoderForCodec(AVEncoder::ECodecType::H264);
-	}
-
 	FStreamer::FStreamer(const FString& InStreamerId)
 		: StreamerId(InStreamerId)
-		, WebRtcSignallingThread(MakeUnique<rtc::Thread>(rtc::SocketServer::CreateDefault()))
-		, Stats(&PlayerSessions)
 	{
-		RedirectWebRtcLogsToUnreal(rtc::LoggingSeverity::LS_VERBOSE);
-
-		FModuleManager::LoadModuleChecked<IModuleInterface>(TEXT("AVEncoder"));
-
-		// required for communication with Signalling Server and must be called in the game thread, while it's used in signalling thread
-		FModuleManager::LoadModuleChecked<FWebSocketsModule>("WebSockets");
-
-		FSignallingServerConnection::FWebSocketFactory WebSocketFactory = [](const FString& Url) { return FWebSocketsModule::Get().CreateWebSocket(Url, TEXT("")); };
-		SignallingServerConnection = MakeUnique<FSignallingServerConnection>(WebSocketFactory, *this, InStreamerId);
+		FStats::Get()->AddSessions(&PlayerSessions);
+		VideoSourceGroup = MakeUnique<FVideoSourceGroup>();
+		FPixelStreamingSignallingConnection::FWebSocketFactory WebSocketFactory = [](const FString& Url) { return FWebSocketsModule::Get().CreateWebSocket(Url, TEXT("")); };
+		SignallingServerConnection = MakeUnique<FPixelStreamingSignallingConnection>(WebSocketFactory, *this, InStreamerId);
 		StartWebRtcSignallingThread();
-		FCoreDelegates::OnPreExit.AddLambda([&]()
+		FCoreDelegates::OnPreExit.AddLambda([&PlayerSessions = PlayerSessions]()
 		{ 
 			PlayerSessions.DeleteAllPlayerSessions(true);
 		});
+		StreamerInputDevices = MakeShared<FStreamerInputDevices>(FSlateApplication::Get().GetPlatformApplication()->GetMessageHandler());
+		SetInputDevice(StreamerInputDevices->CreateInputDevice());
+		IModularFeatures::Get().RegisterModularFeature(GetModularFeatureName(), this);
 	}
 
 	FStreamer::~FStreamer()
 	{
+		IModularFeatures::Get().UnregisterModularFeature(GetModularFeatureName(), this);
+		FStats::Get()->RemoveSessions(&PlayerSessions);
 		StopStreaming();
 		P2PPeerConnectionFactory = nullptr;
 		SFUPeerConnectionFactory = nullptr;
-		WebRtcSignallingThread->Stop();
-		rtc::CleanupSSL();
+		StopWebRtcSignallingThread();
 	}
 
-	bool FStreamer::StartStreaming(const FString& SignallingServerUrl)
+	/** 
+	 * IPixelStreamingStreamer implementation
+	 */
+	void FStreamer::SetStreamFPS(int32 InFramesPerSecond)
 	{
-		StopStreaming();
-		ConnectToSignallingServer(SignallingServerUrl);
-		return true;
+		VideoSourceGroup->SetFPS(InFramesPerSecond);
 	}
 
+	void FStreamer::SetVideoInput(TSharedPtr<FPixelStreamingVideoInput> Input)
+	{
+		VideoSourceGroup->SetVideoInput(Input);
+		Input->OnFrame.AddLambda([this](const FPixelStreamingSourceFrame& SourceFrame)
+		{
+			if(bCaptureNextBackBufferAndStream)
+			{
+				bCaptureNextBackBufferAndStream = false;
+
+				// Read the data out of the back buffer and send as a JPEG.
+				FRHICommandListImmediate& RHICmdList = FRHICommandListExecutor::GetImmediateCommandList();
+				FIntRect Rect(0, 0, SourceFrame.FrameTexture->GetDesc().Extent.X, SourceFrame.FrameTexture->GetDesc().Extent.Y);
+				TArray<FColor> Data;
+
+				RHICmdList.ReadSurfaceData(SourceFrame.FrameTexture, Rect, Data, FReadSurfaceDataFlags());
+				SendFreezeFrame(MoveTemp(Data), Rect);
+			}
+		});
+	}
+
+	void FStreamer::SetTargetViewport(TSharedPtr<FSceneViewport> InTargetViewport)
+	{
+		InputDevice->SetTargetViewport(InTargetViewport);
+	}
+
+	void FStreamer::SetSignallingServerURL(const FString& InSignallingServerURL)
+	{
+		CurrentSignallingServerURL = InSignallingServerURL;
+	}
+
+	void FStreamer::StartStreaming()
+	{
+		if(CurrentSignallingServerURL.IsEmpty())
+		{
+			UE_LOG(LogPixelStreaming, Error, TEXT("Attempted to start streamer (%s) but no signalling server URL has been set. Use Streamer->SetSignallingServerURL(URL)"), *StreamerId);
+			return;
+		}
+		StopStreaming();
+		VideoSourceGroup->Start();
+		SignallingServerConnection->Connect(CurrentSignallingServerURL);
+	}
+
 	void FStreamer::StopStreaming()
 	{
-		bStreamingStarted = false;
-		DisconnectFromSignallingServer();
-		DeleteAllPlayerSessions();
-
-		IPixelStreamingModule::Get().OnStreamingStopped().Broadcast(IPixelStreamingModule::Get());
-	}
-
-	rtc::scoped_refptr<webrtc::AudioDeviceModule> FStreamer::CreateAudioDeviceModule() const
-	{
-		bool bUseLegacyAudioDeviceModule = Settings::CVarPixelStreamingWebRTCUseLegacyAudioDevice.GetValueOnAnyThread();
-		rtc::scoped_refptr<webrtc::AudioDeviceModule> AudioDeviceModule;
-		if (bUseLegacyAudioDeviceModule)
-		{
-			AudioDeviceModule = new rtc::RefCountedObject<FAudioCapturer>();
-		}
-		else
-		{
-			AudioDeviceModule = new rtc::RefCountedObject<FAudioDeviceModule>();
-		}
-		return AudioDeviceModule;
-	}
-
-	void FStreamer::StartWebRtcSignallingThread()
-	{
-		// Initialisation of WebRTC. Note: That interfacing with WebRTC should generally happen in WebRTC signalling thread.
-
-		// Create our own WebRTC thread for signalling
-		WebRtcSignallingThread->SetName("WebRtcSignallingThread", nullptr);
-		WebRtcSignallingThread->Start();
-
-		rtc::InitializeSSL();
-
-		PeerConnectionConfig = {};
-
-		/* ---------- P2P Peer Connection Factory ---------- */
-
-		std::unique_ptr<FVideoEncoderFactory> P2PPeerConnectionFactoryPtr = std::make_unique<FVideoEncoderFactory>();
-		P2PVideoEncoderFactory = P2PPeerConnectionFactoryPtr.get();
-
-		P2PPeerConnectionFactory = webrtc::CreatePeerConnectionFactory(
-			nullptr,					  // network_thread
-			nullptr,					  // worker_thread
-			WebRtcSignallingThread.Get(), // signal_thread
-			CreateAudioDeviceModule(),	  // audio device module
-			webrtc::CreateAudioEncoderFactory<webrtc::AudioEncoderOpus>(),
-			webrtc::CreateAudioDecoderFactory<webrtc::AudioDecoderOpus>(),
-			std::move(P2PPeerConnectionFactoryPtr),
-			std::make_unique<webrtc::InternalDecoderFactory>(),
-			nullptr,					   // audio_mixer
-			SetupAudioProcessingModule()); // audio_processing
-
-		check(P2PPeerConnectionFactory.get() != nullptr);
-
-		/* ---------- SFU Peer Connection Factory ---------- */
-
-		std::unique_ptr<FSimulcastEncoderFactory> SFUPeerConnectionFactoryPtr = std::make_unique<FSimulcastEncoderFactory>();
-		SFUVideoEncoderFactory = SFUPeerConnectionFactoryPtr.get();
-
-		SFUPeerConnectionFactory = webrtc::CreatePeerConnectionFactory(
-			nullptr,					  // network_thread
-			nullptr,					  // worker_thread
-			WebRtcSignallingThread.Get(), // signal_thread
-			CreateAudioDeviceModule(),	  // audio device module
-			webrtc::CreateAudioEncoderFactory<webrtc::AudioEncoderOpus>(),
-			webrtc::CreateAudioDecoderFactory<webrtc::AudioDecoderOpus>(),
-			std::move(SFUPeerConnectionFactoryPtr),
-			std::make_unique<webrtc::InternalDecoderFactory>(),
-			nullptr,					   // audio_mixer
-			SetupAudioProcessingModule()); // audio_processing
-
-		check(SFUPeerConnectionFactory.get() != nullptr);
-	}
-
-	webrtc::AudioProcessing* FStreamer::SetupAudioProcessingModule()
-	{
-		webrtc::AudioProcessing* AudioProcessingModule = webrtc::AudioProcessingBuilder().Create();
-		webrtc::AudioProcessing::Config Config;
-		// Enabled multi channel audio capture/render
-		Config.pipeline.multi_channel_capture = true;
-		Config.pipeline.multi_channel_render = true;
-		Config.pipeline.maximum_internal_processing_rate = 48000;
-		// Turn off all other audio processing effects in UE's WebRTC. We want to stream audio from UE as pure as possible.
-		Config.pre_amplifier.enabled = false;
-		Config.high_pass_filter.enabled = false;
-		Config.echo_canceller.enabled = false;
-		Config.noise_suppression.enabled = false;
-		Config.transient_suppression.enabled = false;
-		Config.voice_detection.enabled = false;
-		Config.gain_controller1.enabled = false;
-		Config.gain_controller2.enabled = false;
-		Config.residual_echo_detector.enabled = false;
-		Config.level_estimation.enabled = false;
-		// Apply the config.
-		AudioProcessingModule->ApplyConfig(Config);
-
-		return AudioProcessingModule;
-	}
-
-	void FStreamer::ConnectToSignallingServer(const FString& Url)
-	{
-		PreviousSignallingServerURL = Url;
-		SignallingServerConnection->Connect(Url);
-	}
-
-	void FStreamer::DisconnectFromSignallingServer()
-	{
 		SignallingServerConnection->Disconnect();
+		DeleteAllPlayerSessions();
+
+		VideoSourceGroup->Stop();
+
+		if (bStreamingStarted)
+		{
+			OnStreamingStopped().Broadcast(this);
+		}
+		bStreamingStarted = false;
 	}
 
-	void FStreamer::OnConfig(const webrtc::PeerConnectionInterface::RTCConfiguration& Config)
+	void FStreamer::SendPlayerMessage(Protocol::EToPlayerMsg Type, const FString& Descriptor)
 	{
+		PlayerSessions.ForEachSession([Type, Descriptor](TSharedPtr<IPlayerSession> PlayerSession) { PlayerSession->SendMessage(Type, Descriptor); });
+	}
+
+	void FStreamer::FreezeStream(UTexture2D* Texture)
+	{
+		if (Texture)
+		{
+			ENQUEUE_RENDER_COMMAND(ReadSurfaceCommand)
+			([this, Texture](FRHICommandListImmediate& RHICmdList) {
+				// A frame is supplied so immediately read its data and send as a JPEG.
+				FTextureRHIRef TextureRHI = Texture->GetResource() ? Texture->GetResource()->TextureRHI : nullptr;
+				if (!TextureRHI)
+				{
+					UE_LOG(LogPixelStreaming, Error, TEXT("Attempting freeze frame with texture %s with no texture RHI"), *Texture->GetName());
+					return;
+				}
+				uint32 Width = TextureRHI->GetDesc().Extent.X;
+				uint32 Height = TextureRHI->GetDesc().Extent.Y;
+
+				FTextureRHIRef DestTexture = CreateRHITexture(Width, Height);
+
+				FGPUFenceRHIRef CopyFence = GDynamicRHI->RHICreateGPUFence(*FString::Printf(TEXT("FreezeFrameFence")));
+
+				// Copy freeze frame texture to empty texture
+				CopyTexture(RHICmdList, TextureRHI, DestTexture, CopyFence);
+
+				TArray<FColor> Data;
+				FIntRect Rect(0, 0, Width, Height);
+				RHICmdList.ReadSurfaceData(DestTexture, Rect, Data, FReadSurfaceDataFlags());
+				SendFreezeFrame(MoveTemp(Data), Rect);
+			});
+		}
+		else
+		{
+			// A frame is not supplied, so we need to capture the back buffer at
+			// the next opportunity, and send as a JPEG.
+			bCaptureNextBackBufferAndStream = true;
+		}
+	}
+
+	void FStreamer::UnfreezeStream()
+	{
+		// Force a keyframe so when stream unfreezes if player has never received a h.264 frame before they can still connect.
+		ForceKeyFrame();
+
+		PlayerSessions.ForEachSession([](TSharedPtr<IPlayerSession> PlayerSession) { PlayerSession->SendUnfreezeFrame(); });
+
+		CachedJpegBytes.Empty();
+	}
+
+	void FStreamer::SendFileData(TArray<uint8>& ByteData, FString& MimeType, FString& FileExtension)
+	{
+		PlayerSessions.ForEachSession([&ByteData, &MimeType, &FileExtension](TSharedPtr<IPlayerSession> PlayerSession) {
+			AsyncTask(ENamedThreads::AnyHiPriThreadHiPriTask, [&ByteData, &MimeType, &FileExtension, PlayerSession]() {
+				bool bSentFileData = PlayerSession->SendFileData(ByteData, MimeType, FileExtension);
+				if (!bSentFileData)
+				{
+					UE_LOG(LogPixelStreaming, Log, TEXT("Could not send file over data channel for some reason."));
+				}
+			});
+		});
+	}
+
+	void FStreamer::KickPlayer(FPixelStreamingPlayerId PlayerId)
+	{
+		SignallingServerConnection->SendDisconnectPlayer(PlayerId, TEXT("Player was kicked"));
+	}
+
+	IPixelStreamingAudioSink* FStreamer::GetPeerAudioSink(FPixelStreamingPlayerId PlayerId)
+	{
+		return PlayerSessions.ForSession<IPixelStreamingAudioSink*>(PlayerId, [](TSharedPtr<IPlayerSession> Session) { return Session->GetAudioSink(); });
+	}
+
+	IPixelStreamingAudioSink* FStreamer::GetUnlistenedAudioSink()
+	{
+		IPixelStreamingAudioSink* Result = nullptr;
+		PlayerSessions.ForEachSession([&Result](TSharedPtr<IPlayerSession> Session) {
+			if (!Result && Session->GetAudioSink() && !Session->GetAudioSink()->HasAudioConsumers())
+			{
+				Result = Session->GetAudioSink();
+			}
+		});
+		return Result;
+	}
+
+	IPixelStreamingStreamer::FStreamingStartedEvent& FStreamer::OnStreamingStarted()
+	{
+		return StreamingStartedEvent;
+	}
+
+	IPixelStreamingStreamer::FStreamingStoppedEvent& FStreamer::OnStreamingStopped()
+	{
+		return StreamingStoppedEvent;
+	}
+
+	void FStreamer::RegisterCreateInputDevice(IPixelStreamingInputDevice::FCreateInputDeviceFunc& InCreateInputDevice)
+	{
+		checkf(StreamerInputDevices, TEXT("StreamerInputDevices does not exist yet"));
+		StreamerInputDevices->OverrideInputDevice(InCreateInputDevice);
+	}
+	/** 
+	 * End IPixelStreamingStreamer implementation
+	 */
+
+	/** 
+	 * Own methods
+	 */
+	void FStreamer::ForceKeyFrame()
+	{
+		if (P2PVideoEncoderFactory)
+		{
+			P2PVideoEncoderFactory->ForceKeyFrame();
+		}
+		else
+		{
+			UE_LOG(LogPixelStreaming, Log, TEXT("Cannot force a key frame - video encoder factory is nullptr."));
+		}
+	}
+
+	void FStreamer::PushFrame()
+	{
+		VideoSourceGroup->Tick();
+	}
+
+	void FStreamer::AddPlayerConfig(TSharedRef<FJsonObject>& JsonObject)
+	{
+		checkf(InputDevice.IsValid(), TEXT("No Input Device available when populating Player Config"));
+		JsonObject->SetBoolField(TEXT("FakingTouchEvents"), InputDevice->IsFakingTouchEvents());
+		FString PixelStreamingControlScheme;
+		if (Settings::GetControlScheme(PixelStreamingControlScheme))
+		{
+			JsonObject->SetStringField(TEXT("ControlScheme"), PixelStreamingControlScheme);
+		}
+		float PixelStreamingFastPan;
+		if (Settings::GetFastPan(PixelStreamingFastPan))
+		{
+			JsonObject->SetNumberField(TEXT("FastPan"), PixelStreamingFastPan);
+		}
+	}
+	/** 
+	 * End own methods
+	 */
+
+
+
+	/** 
+	 * ISignallingServerConnectionObserver implementation
+	 */
+	void FStreamer::OnSignallingConfig(const webrtc::PeerConnectionInterface::RTCConfiguration& Config)
+	{
 		PeerConnectionConfig = Config;
 	}
 
-	void FStreamer::OnDataChannelOpen(FPixelStreamingPlayerId PlayerId, webrtc::DataChannelInterface* DataChannel)
+	void FStreamer::OnSignallingSessionDescription(FPixelStreamingPlayerId PlayerId, webrtc::SdpType Type, const FString& Sdp)
 	{
-		if (DataChannel)
-		{
-			// When data channel is open, try to send cached freeze frame (if we have one)
-			SendCachedFreezeFrameTo(PlayerId);
-		}
-	}
-
-	TSharedPtr<IPlayerSession> FStreamer::CreateSession(FPixelStreamingPlayerId PlayerId, int Flags)
-	{
-		PeerConnectionConfig.enable_simulcast_stats = true;
-
-		bool bIsSFU = (Flags & Protocol::EPlayerFlags::PSPFlag_IsSFU) != 0;
-
-		rtc::scoped_refptr<webrtc::PeerConnectionFactoryInterface> PCFactory = bIsSFU ? SFUPeerConnectionFactory : P2PPeerConnectionFactory;
-
-		TSharedPtr<IPlayerSession> NewSession = PlayerSessions.CreatePlayerSession(
-			PlayerId,
-			PCFactory,
-			PeerConnectionConfig,
-			SignallingServerConnection.Get(),
-			Flags);
-		if (NewSession)
-		{
-			webrtc::PeerConnectionInterface& PeerConnection = NewSession->GetPeerConnection();
-			const bool SupportsDataChannel = (Flags & Protocol::EPlayerFlags::PSPFlag_SupportsDataChannel) != 0;
-
-			if (SupportsDataChannel)
-			{
-				webrtc::DataChannelInit DataChannelConfig;
-				DataChannelConfig.reliable = true;
-				DataChannelConfig.ordered = true;
-				if (bIsSFU)
-				{
-					DataChannelConfig.negotiated = true;
-					DataChannelConfig.id = 1023;
-				}
-
-				rtc::scoped_refptr<webrtc::DataChannelInterface> DataChannel = PeerConnection.CreateDataChannel("datachannel", &DataChannelConfig);
-				NewSession->SetDataChannel(DataChannel);
-
-				// Bind to the on data channel open delegate (we use this as the earliest time peer is ready for freeze frame)
-				FDataChannelObserver* DataChannelObserver = NewSession->GetDataChannelObserver();
-				check(DataChannelObserver);
-				if (UPixelStreamingDelegates* Delegates = UPixelStreamingDelegates::GetPixelStreamingDelegates())
-				{
-					Delegates->OnDataChannelOpenNative.AddRaw(this, &FStreamer::OnDataChannelOpen);
-				}
-				DataChannelObserver->Register(DataChannel);
-			}
-		}
-
-		return NewSession;
-	}
-
-	void FStreamer::OnSessionDescription(FPixelStreamingPlayerId PlayerId, webrtc::SdpType Type, const FString& Sdp)
-	{
 		switch (Type)
 		{
@@ -242,7 +285,7 @@
 			case webrtc::SdpType::kAnswer:
 			case webrtc::SdpType::kPrAnswer:
 			{
-				PlayerSessions.ForSession(PlayerId, [&](TSharedPtr<IPlayerSession> Session) { Session->OnAnswer(Sdp); });
+				PlayerSessions.ForSession(PlayerId, [&Sdp](TSharedPtr<IPlayerSession> Session) { Session->OnAnswer(Sdp); });
 				bStreamingStarted = true;
 				break;
 			}
@@ -252,13 +295,155 @@
 		}
 	}
 
+	void FStreamer::OnSignallingRemoteIceCandidate(FPixelStreamingPlayerId PlayerId, const FString& SdpMid, int SdpMLineIndex, const FString& Sdp)
+	{
+		PlayerSessions.ForSession(PlayerId, [&SdpMid, SdpMLineIndex, &Sdp](TSharedPtr<IPlayerSession> Session) { Session->OnRemoteIceCandidate(SdpMid, SdpMLineIndex, Sdp); });
+	}
+
+	void FStreamer::OnSignallingPlayerConnected(FPixelStreamingPlayerId PlayerId, int Flags)
+	{
+		// create peer connection
+		if (TSharedPtr<IPlayerSession> NewSession = CreateSession(PlayerId, Flags))
+		{
+			AddStreams(NewSession, Flags);
+
+			rtc::scoped_refptr<webrtc::PeerConnectionInterface> PeerConnection = NewSession->GetPeerConnection();
+			ModifyTransceivers(PeerConnection->GetTransceivers(), Flags);
+
+			// observer for creating offer
+			FPixelStreamingCreateSessionDescriptionObserver* CreateOfferObserver = FPixelStreamingCreateSessionDescriptionObserver::Create(
+				[this, PlayerId, PeerConnection](webrtc::SessionDescriptionInterface* SDP) // on SDP create success
+				{
+					FPixelStreamingSetSessionDescriptionObserver* SetLocalDescriptionObserver = FPixelStreamingSetSessionDescriptionObserver::Create(
+						[this, PlayerId, SDP]() // on SDP set success
+						{
+							SignallingServerConnection->SendOffer(PlayerId, *SDP);
+						},
+						[](const FString& Error) // on SDP set failure
+						{
+							UE_LOG(LogPixelStreaming, Error, TEXT("Failed to set local description: %s"), *Error);
+						});
+
+					MungeLocalSDP(SDP->description());
+					SetLocalDescription(PeerConnection, SetLocalDescriptionObserver, SDP);
+				},
+				[](const FString& Error) // on SDP create failure
+				{
+					UE_LOG(LogPixelStreaming, Error, TEXT("Failed to create offer: %s"), *Error);
+				});
+
+			PeerConnection->CreateOffer(CreateOfferObserver, {});
+		}
+	}
+
+	void FStreamer::OnSignallingPlayerDisconnected(FPixelStreamingPlayerId PlayerId)
+	{
+		UE_LOG(LogPixelStreaming, Log, TEXT("player %s disconnected"), *PlayerId);
+		DeletePlayerSession(PlayerId);
+	}
+
+	void FStreamer::OnSignallingSFUPeerDataChannels(FPixelStreamingPlayerId SFUId, FPixelStreamingPlayerId PlayerId, int32 SendStreamId, int32 RecvStreamId)
+	{
+		PlayerSessions.CreateSFUDataOnlyPeer(SFUId, PlayerId, InputDevice, SendStreamId, RecvStreamId, SignallingServerConnection.Get());
+	}
+
+	void FStreamer::OnSignallingConnected()
+	{
+		OnStreamingStarted().Broadcast(this);
+	}
+
+	void FStreamer::OnSignallingDisconnected(int32 StatusCode, const FString& Reason, bool bWasClean)
+	{
+		DeleteAllPlayerSessions();
+
+		// Call disconnect here makes sure any other websocket callbacks etc are cleared
+		SignallingServerConnection->Disconnect();
+		SignallingServerConnection->Connect(CurrentSignallingServerURL);
+	}
+
+	void FStreamer::OnSignallingError(const FString& ErrorMsg)
+	{
+		DeleteAllPlayerSessions();
+
+		// Call disconnect here makes sure any other websocket callbacks etc are cleared
+		SignallingServerConnection->Disconnect();
+		SignallingServerConnection->Connect(CurrentSignallingServerURL);
+	}
+	/** 
+	 * End ISignallingServerConnectionObserver implementation
+	 */
+
+
+	/** 
+	 * IInputDevice implementation
+	 */	
+	TSharedPtr<IInputDevice> FStreamer::CreateInputDevice(const TSharedRef<FGenericApplicationMessageHandler>& InMessageHandler)
+	{
+		return StreamerInputDevices;
+	}
+	/** 
+	 * End IInputDevice implementation
+	 */
+
+
+	/** 
+	 * Own methods
+	 */
+	void FStreamer::SendFreezeFrame(TArray<FColor> RawData, const FIntRect& Rect)
+	{
+		IImageWrapperModule& ImageWrapperModule = FModuleManager::GetModuleChecked<IImageWrapperModule>(TEXT("ImageWrapper"));
+		TSharedPtr<IImageWrapper> ImageWrapper = ImageWrapperModule.CreateImageWrapper(EImageFormat::JPEG);
+		bool bSuccess = ImageWrapper->SetRaw(RawData.GetData(), RawData.Num() * sizeof(FColor), Rect.Width(), Rect.Height(), ERGBFormat::BGRA, 8);
+		if (bSuccess)
+		{
+			// Compress to a JPEG of the maximum possible quality.
+			int32 Quality = Settings::CVarPixelStreamingFreezeFrameQuality.GetValueOnAnyThread();
+			const TArray64<uint8>& JpegBytes = ImageWrapper->GetCompressed(Quality);
+			PlayerSessions.ForEachSession([&JpegBytes](TSharedPtr<IPlayerSession> PlayerSession) { PlayerSession->SendFreezeFrame(JpegBytes); });
+			CachedJpegBytes = JpegBytes;
+		}
+		else
+		{
+			UE_LOG(LogPixelStreaming, Error, TEXT("JPEG image wrapper failed to accept frame data"));
+		}
+	}
+
+	void FStreamer::SendCachedFreezeFrameTo(FPixelStreamingPlayerId PlayerId) const
+	{
+		if (CachedJpegBytes.Num() > 0)
+		{
+			PlayerSessions.ForSession(PlayerId, [&CachedJpegBytes = CachedJpegBytes](TSharedPtr<IPlayerSession> Session) { Session->SendFreezeFrame(CachedJpegBytes); });
+		}
+	}
+
+	TSharedPtr<IPlayerSession> FStreamer::CreateSession(FPixelStreamingPlayerId PlayerId, int Flags)
+	{
+#if WEBRTC_VERSION == 84
+		PeerConnectionConfig.enable_simulcast_stats = true;
+#endif
+
+		bool bIsSFU = (Flags & Protocol::EPlayerFlags::PSPFlag_IsSFU) != 0;
+
+		rtc::scoped_refptr<webrtc::PeerConnectionFactoryInterface> PCFactory = bIsSFU ? SFUPeerConnectionFactory : P2PPeerConnectionFactory;
+
+		TSharedPtr<IPlayerSession> NewSession = PlayerSessions.CreatePlayerSession(
+			PlayerId,
+			InputDevice,
+			PCFactory,
+			PeerConnectionConfig,
+			SignallingServerConnection.Get(),
+			Flags);
+
+		return NewSession;
+	}
+
 	void FStreamer::OnOffer(FPixelStreamingPlayerId PlayerId, const FString& Sdp)
 	{
 		TSharedPtr<IPlayerSession> NewSession = CreateSession(PlayerId, Protocol::EPlayerFlags::PSPFlag_SupportsDataChannel);
 		if (NewSession)
 		{
 			webrtc::SdpParseError Error;
-			std::unique_ptr<webrtc::SessionDescriptionInterface> SessionDesc = webrtc::CreateSessionDescription(webrtc::SdpType::kOffer, ToString(Sdp), &Error);
+			std::unique_ptr<webrtc::SessionDescriptionInterface> SessionDesc = webrtc::CreateSessionDescription(webrtc::SdpType::kOffer, UE::PixelStreaming::ToString(Sdp), &Error);
 			if (!SessionDesc)
 			{
 				UE_LOG(LogPixelStreaming, Error, TEXT("Failed to parse offer's SDP: %s\n%s"), Error.description.c_str(), *Sdp);
@@ -271,105 +456,53 @@
 			UE_LOG(LogPixelStreaming, Error, TEXT("Failed to create player session, peer connection was nullptr."));
 	}
 
-	void FStreamer::SetLocalDescription(webrtc::PeerConnectionInterface* PeerConnection, FSetSessionDescriptionObserver* Observer, webrtc::SessionDescriptionInterface* SDP)
+	void FStreamer::ModifyTransceivers(std::vector<rtc::scoped_refptr<webrtc::RtpTransceiverInterface>> Transceivers, int Flags)
 	{
-		// Note from Luke about WebRTC: the sink of video capturer will be added as a direct result
-		// of `PeerConnection->SetLocalDescription()` call but video encoder will be created later on
-		// when the first frame is pushed into the WebRTC pipeline (by the capturer calling `OnFrame`).
+		bool bTransmitUEAudio = !Settings::CVarPixelStreamingWebRTCDisableTransmitAudio.GetValueOnAnyThread();
+		bool bIsSFU = (Flags & Protocol::EPlayerFlags::PSPFlag_IsSFU) != Protocol::EPlayerFlags::PSPFlag_None;
+		bool bReceiveBrowserAudio = !bIsSFU && !Settings::CVarPixelStreamingWebRTCDisableReceiveAudio.GetValueOnAnyThread();
 
-		PeerConnection->SetLocalDescription(Observer, SDP);
+		for (auto& Transceiver : Transceivers)
+		{
+			rtc::scoped_refptr<webrtc::RtpSenderInterface> Sender = Transceiver->sender();
 
-		// Once local description has been set we can start setting some encoding information for the video stream rtp sender
-		for (rtc::scoped_refptr<webrtc::RtpSenderInterface> Sender : PeerConnection->GetSenders())
-		{
-			cricket::MediaType MediaType = Sender->media_type();
-			if (MediaType == cricket::MediaType::MEDIA_TYPE_VIDEO)
+			if (Transceiver->media_type() == cricket::MediaType::MEDIA_TYPE_AUDIO)
 			{
-				webrtc::RtpParameters ExistingParams = Sender->GetParameters();
 
-				// Set the degradation preference based on our CVar for it.
-				ExistingParams.degradation_preference = Settings::GetDegradationPreference();
+				// Determine the direction of the transceiver
+				webrtc::RtpTransceiverDirection AudioTransceiverDirection;
+				if (bTransmitUEAudio && bReceiveBrowserAudio)
+				{
+					AudioTransceiverDirection = webrtc::RtpTransceiverDirection::kSendRecv;
+				}
+				else if (bTransmitUEAudio)
+				{
+					AudioTransceiverDirection = webrtc::RtpTransceiverDirection::kSendOnly;
+				}
+				else if (bReceiveBrowserAudio)
+				{
+					AudioTransceiverDirection = webrtc::RtpTransceiverDirection::kRecvOnly;
+				}
+				else
+				{
+					AudioTransceiverDirection = webrtc::RtpTransceiverDirection::kInactive;
+				}
 
-				webrtc::RTCError Err = Sender->SetParameters(ExistingParams);
-				if (!Err.ok())
-				{
-					const char* ErrMsg = Err.message();
-					FString ErrorStr(ErrMsg);
-					UE_LOG(LogPixelStreaming, Error, TEXT("Failed to set RTP Sender params: %s"), *ErrorStr);
-				}
+				Transceiver->SetDirection(AudioTransceiverDirection);
+				Sender->SetStreams({ TCHAR_TO_UTF8(*GetAudioStreamID()) });
+			}
+			else if (Transceiver->media_type() == cricket::MediaType::MEDIA_TYPE_VIDEO)
+			{
+				Transceiver->SetDirection(webrtc::RtpTransceiverDirection::kSendOnly);
+				Sender->SetStreams({ TCHAR_TO_UTF8(*GetVideoStreamID()) });
+
+				webrtc::RtpParameters RtpParams = Sender->GetParameters();
+				RtpParams.encodings = CreateRTPEncodingParams(Flags);
+				Sender->SetParameters(RtpParams);
 			}
 		}
 	}
 
-	void FStreamer::SendAnswer(TSharedPtr<IPlayerSession> Session, TUniquePtr<webrtc::SessionDescriptionInterface> Sdp)
-	{
-		// below is async execution (with error handling) of:
-		//		PeerConnection.SetRemoteDescription(SDP);
-		//		Answer = PeerConnection.CreateAnswer();
-		//		PeerConnection.SetLocalDescription(Answer);
-		//		SignallingServerConnection.SendAnswer(Answer);
-
-		FSetSessionDescriptionObserver* SetLocalDescriptionObserver = FSetSessionDescriptionObserver::Create(
-			[this, Session]() // on success
-			{
-				SignallingServerConnection->SendAnswer(Session->GetPlayerId(), *Session->GetPeerConnection().local_description());
-				bStreamingStarted = true;
-			},
-			[this, Session](const FString& Error) // on failure
-			{
-				UE_LOG(LogPixelStreaming, Error, TEXT("Failed to set local description: %s"), *Error);
-				Session->DisconnectPlayer(Error);
-				PlayerSessions.DeletePlayerSession(Session->GetPlayerId());
-			});
-
-		auto OnCreateAnswerSuccess = [this, Session, SetLocalDescriptionObserver](webrtc::SessionDescriptionInterface* SDP) {
-			SetLocalDescription(&Session->GetPeerConnection(), SetLocalDescriptionObserver, SDP);
-		};
-
-		FCreateSessionDescriptionObserver* CreateAnswerObserver = FCreateSessionDescriptionObserver::Create(
-			MoveTemp(OnCreateAnswerSuccess),
-			[this, Session](const FString& Error) // on failure
-			{
-				UE_LOG(LogPixelStreaming, Error, TEXT("Failed to create answer: %s"), *Error);
-				Session->DisconnectPlayer(Error);
-				PlayerSessions.DeletePlayerSession(Session->GetPlayerId());
-			});
-
-		auto OnSetRemoteDescriptionSuccess = [this, Session, CreateAnswerObserver]() {
-			// Note: these offer to receive are superseded now we are use transceivers to setup our peer connection media
-			int offer_to_receive_video = webrtc::PeerConnectionInterface::RTCOfferAnswerOptions::kUndefined;
-			int offer_to_receive_audio = webrtc::PeerConnectionInterface::RTCOfferAnswerOptions::kUndefined;
-			bool voice_activity_detection = false;
-			bool ice_restart = true;
-			bool use_rtp_mux = true;
-
-			webrtc::PeerConnectionInterface::RTCOfferAnswerOptions AnswerOption{
-				offer_to_receive_video,
-				offer_to_receive_audio,
-				voice_activity_detection,
-				ice_restart,
-				use_rtp_mux
-			};
-
-			AddStreams(Session, Protocol::EPlayerFlags::PSPFlag_SupportsDataChannel);
-			ModifyTransceivers(Session->GetPeerConnection().GetTransceivers(), Protocol::EPlayerFlags::PSPFlag_SupportsDataChannel);
-			Session->GetPeerConnection().CreateAnswer(CreateAnswerObserver, AnswerOption);
-		};
-		FSetSessionDescriptionObserver* SetRemoteDescriptionObserver = FSetSessionDescriptionObserver::Create(
-			MoveTemp(OnSetRemoteDescriptionSuccess),
-			[this, Session](const FString& Error) // on failure
-			{
-				UE_LOG(LogPixelStreaming, Error, TEXT("Failed to set remote description: %s"), *Error);
-				Session->DisconnectPlayer(Error);
-				PlayerSessions.DeletePlayerSession(Session->GetPlayerId());
-			});
-
-		cricket::SessionDescription* RemoteDescription = Sdp->description();
-		MungeRemoteSDP(RemoteDescription);
-
-		Session->GetPeerConnection().SetRemoteDescription(SetRemoteDescriptionObserver, Sdp.Release());
-	}
-
 	void FStreamer::MungeRemoteSDP(cricket::SessionDescription* RemoteDescription)
 	{
 		// Munge SDP of remote description to inject min, max, start bitrates
@@ -393,91 +526,48 @@
 		}
 	}
 
-	void FStreamer::OnRemoteIceCandidate(FPixelStreamingPlayerId PlayerId, const FString& SdpMid, int SdpMLineIndex, const FString& Sdp)
+	void FStreamer::MungeLocalSDP(cricket::SessionDescription* SessionDescription)
 	{
-		PlayerSessions.ForSession(PlayerId, [&](TSharedPtr<IPlayerSession> Session) { Session->OnRemoteIceCandidate(SdpMid, SdpMLineIndex, Sdp); });
-	}
-
-	void FStreamer::OnPlayerConnected(FPixelStreamingPlayerId PlayerId, int Flags)
-	{
-		// create peer connection
-		if (TSharedPtr<IPlayerSession> NewSession = CreateSession(PlayerId, Flags))
+		std::vector<cricket::ContentInfo>& ContentInfos = SessionDescription->contents();
+		for (cricket::ContentInfo& ContentInfo : ContentInfos)
 		{
-			AddStreams(NewSession, Flags);
-
-			webrtc::PeerConnectionInterface& PeerConnection = NewSession->GetPeerConnection();
-			ModifyTransceivers(PeerConnection.GetTransceivers(), Flags);
-
-			// observer for creating offer
-			FCreateSessionDescriptionObserver* CreateOfferObserver = FCreateSessionDescriptionObserver::Create(
-				[this, PlayerId, PeerConnectionPtr = &PeerConnection](webrtc::SessionDescriptionInterface* SDP) // on SDP create success
+			cricket::MediaContentDescription* MediaDescription = ContentInfo.media_description();
+			cricket::MediaType MediaType = MediaDescription->type();
+			if (MediaType != cricket::MediaType::MEDIA_TYPE_AUDIO)
+			{
+				continue;
+			}
+			cricket::AudioContentDescription* AudioDescription = MediaDescription->as_audio();
+			if (AudioDescription == nullptr)
+			{
+				continue;
+			}
+			std::vector<cricket::AudioCodec> CodecsCopy = AudioDescription->codecs();
+			for (cricket::AudioCodec& Codec : CodecsCopy)
+			{
+				if (Codec.name == "opus")
 				{
-					FSetSessionDescriptionObserver* SetLocalDescriptionObserver = FSetSessionDescriptionObserver::Create(
-						[this, PlayerId, SDP]() // on SDP set success
-						{
-							SignallingServerConnection->SendOffer(PlayerId, *SDP);
-						},
-						[](const FString& Error) // on SDP set failure
-						{
-							UE_LOG(LogPixelStreaming, Error, TEXT("Failed to set local description: %s"), *Error);
-						});
-
-					MungeLocalSDP(SDP->description());
-					SetLocalDescription(PeerConnectionPtr, SetLocalDescriptionObserver, SDP);
-				},
-				[](const FString& Error) // on SDP create failure
-				{
-					UE_LOG(LogPixelStreaming, Error, TEXT("Failed to create offer: %s"), *Error);
-				});
-
-			PeerConnection.CreateOffer(CreateOfferObserver, {});
+					Codec.SetParam(cricket::kCodecParamPTime, "20");
+					Codec.SetParam(cricket::kCodecParamMaxPTime, "120");
+					Codec.SetParam(cricket::kCodecParamMinPTime, "3");
+					Codec.SetParam(cricket::kCodecParamSPropStereo, "1");
+					Codec.SetParam(cricket::kCodecParamStereo, "1");
+					Codec.SetParam(cricket::kCodecParamUseInbandFec, "1");
+					Codec.SetParam(cricket::kCodecParamUseDtx, "0");
+					Codec.SetParam(cricket::kCodecParamMaxAverageBitrate, "510000");
+					Codec.SetParam(cricket::kCodecParamMaxPlaybackRate, "48000");
+				}
+			}
+			AudioDescription->set_codecs(CodecsCopy);
 		}
 	}
 
-	void FStreamer::OnPlayerDisconnected(FPixelStreamingPlayerId PlayerId)
-	{
-		UE_LOG(LogPixelStreaming, Log, TEXT("player %s disconnected"), *PlayerId);
-		DeletePlayerSession(PlayerId);
-	}
-
-	void FStreamer::OnStreamerDataChannels(FPixelStreamingPlayerId SFUId, FPixelStreamingPlayerId PlayerId, int32 SendStreamId, int32 RecvStreamId)
-	{
-		PlayerSessions.CreateNewDataChannel(SFUId, PlayerId, SendStreamId, RecvStreamId, SignallingServerConnection.Get());
-	}
-
-	void FStreamer::OnSignallingServerConnected()
-	{
-		IPixelStreamingModule::Get().OnStreamingStarted().Broadcast(IPixelStreamingModule::Get());
-	}
-
-	void FStreamer::OnSignallingServerDisconnected()
-	{
-		DeleteAllPlayerSessions();
-
-		// Call disconnect here makes sure any other websocket callbacks etc are cleared
-		SignallingServerConnection->Disconnect();
-		ConnectToSignallingServer(PreviousSignallingServerURL); // maybe leave this up to the user?
-	}
-	void FStreamer::ForceKeyFrame()
-	{
-		if (P2PVideoEncoderFactory)
-		{
-			P2PVideoEncoderFactory->ForceKeyFrame();
-		}
-		else
-		{
-			UE_LOG(LogPixelStreaming, Log, TEXT("Cannot force a key frame - video encoder factory is nullptr."));
-		}
-	}
-
-	void FStreamer::DeleteAllPlayerSessions()
-	{
-		PlayerSessions.DeleteAllPlayerSessions();
-		bStreamingStarted = false;
-	}
-
 	void FStreamer::DeletePlayerSession(FPixelStreamingPlayerId PlayerId)
 	{
+		PlayerSessions.ForSession(PlayerId, [this](TSharedPtr<IPlayerSession> Session) {
+			VideoSourceGroup->RemoveVideoSource(Session->GetVideoSource().get());
+		});
+
 		int NumRemainingPlayers = PlayerSessions.DeletePlayerSession(PlayerId);
 
 		if (NumRemainingPlayers == 0)
@@ -486,16 +576,14 @@
 		}
 	}
 
-	FString FStreamer::GetAudioStreamID() const
+	void FStreamer::DeleteAllPlayerSessions()
 	{
-		bool bSyncVideoAndAudio = !Settings::CVarPixelStreamingWebRTCDisableAudioSync.GetValueOnAnyThread();
-		return bSyncVideoAndAudio ? TEXT("pixelstreaming_av_stream_id") : TEXT("pixelstreaming_audio_stream_id");
+		PlayerSessions.ForEachSession([this](TSharedPtr<IPlayerSession> Session) {
+			VideoSourceGroup->RemoveVideoSource(Session->GetVideoSource().get());
+		});
+
+		PlayerSessions.DeleteAllPlayerSessions();
+		bStreamingStarted = false;
 	}
 
-	FString FStreamer::GetVideoStreamID() const
-	{
-		bool bSyncVideoAndAudio = !Settings::CVarPixelStreamingWebRTCDisableAudioSync.GetValueOnAnyThread();
-		return bSyncVideoAndAudio ? TEXT("pixelstreaming_av_stream_id") : TEXT("pixelstreaming_video_stream_id");
-	}
-
 	void FStreamer::AddStreams(TSharedPtr<IPlayerSession> Session, int Flags)
@@ -505,82 +593,22 @@
 		SetupAudioTrack(Session, GetAudioStreamID(), AudioTrackLabel, Flags);
 	}
 
-	std::vector<webrtc::RtpEncodingParameters> FStreamer::CreateRTPEncodingParams(int Flags)
-	{
-		bool bIsSFU = (Flags & Protocol::EPlayerFlags::PSPFlag_IsSFU) != Protocol::EPlayerFlags::PSPFlag_None;
-
-		std::vector<webrtc::RtpEncodingParameters> EncodingParams;
-
-		if (bIsSFU && Settings::SimulcastParameters.Layers.Num() > 0)
-		{
-			using FLayer = Settings::FSimulcastParameters::FLayer;
-
-			// encodings should be lowest res to highest
-			TArray<FLayer*> SortedLayers;
-			for (FLayer& Layer : Settings::SimulcastParameters.Layers)
-			{
-				SortedLayers.Add(&Layer);
-			}
-
-			SortedLayers.Sort([](const FLayer& LayerA, const FLayer& LayerB) { return LayerA.Scaling > LayerB.Scaling; });
-
-			const int LayerCount = SortedLayers.Num();
-			for (int i = 0; i < LayerCount; ++i)
-			{
-				const FLayer* SimulcastLayer = SortedLayers[i];
-				webrtc::RtpEncodingParameters LayerEncoding{};
-				LayerEncoding.rid = TCHAR_TO_UTF8(*(FString("simulcast") + FString::FromInt(LayerCount - i)));
-				LayerEncoding.min_bitrate_bps = SimulcastLayer->MinBitrate;
-				LayerEncoding.max_bitrate_bps = SimulcastLayer->MaxBitrate;
-				LayerEncoding.scale_resolution_down_by = SimulcastLayer->Scaling;
-				
-				// In M84 this will crash with "Attempted to set an unimplemented parameter of RtpParameters".
-				// Try re-enabling this when we upgrade WebRTC versions.
-				//LayerEncoding.network_priority = webrtc::Priority::kHigh;
-				
-				LayerEncoding.max_framerate = FMath::Max(60, Settings::CVarPixelStreamingWebRTCFps.GetValueOnAnyThread());
-				EncodingParams.push_back(LayerEncoding);
-			}
-		}
-		else
-		{
-			webrtc::RtpEncodingParameters Encoding{};
-			Encoding.rid = "base";
-			Encoding.max_bitrate_bps = Settings::CVarPixelStreamingWebRTCMaxBitrate.GetValueOnAnyThread();
-			Encoding.min_bitrate_bps = Settings::CVarPixelStreamingWebRTCMinBitrate.GetValueOnAnyThread();
-			Encoding.max_framerate = FMath::Max(60, Settings::CVarPixelStreamingWebRTCFps.GetValueOnAnyThread());
-			Encoding.scale_resolution_down_by.reset();
-			Encoding.network_priority = webrtc::Priority::kHigh;
-			EncodingParams.push_back(Encoding);
-		}
-
-		return EncodingParams;
-	}
-
 	void FStreamer::SetupVideoTrack(TSharedPtr<IPlayerSession> Session, const FString InVideoStreamId, const FString InVideoTrackLabel, int Flags)
 	{
-
 		bool bIsSFU = (Flags & Protocol::EPlayerFlags::PSPFlag_IsSFU) != Protocol::EPlayerFlags::PSPFlag_None;
 
-		const TArray<FName>& ActiveTextureSourceTypes = IPixelStreamingModule::Get().GetActiveTextureSourceTypes();
-
-		// Todo (Luke): Add multitrack support here.
-		verifyf(ActiveTextureSourceTypes.Num() == 1, TEXT("Only a single track is supported for streaming currently."));
-
-		FName SourceType = ActiveTextureSourceTypes[0];
-
 		rtc::scoped_refptr<webrtc::PeerConnectionFactoryInterface> PCFactory = bIsSFU ? SFUPeerConnectionFactory : P2PPeerConnectionFactory;
 
 		// Create a video source for this player
-		rtc::scoped_refptr<FVideoSourceBase> VideoSource;
+		rtc::scoped_refptr<webrtc::VideoTrackSourceInterface> VideoSource;
 		if (bIsSFU)
 		{
-			VideoSource = new FVideoSourceSFU(SourceType);
+			VideoSource = VideoSourceGroup->CreateSFUVideoSource();
 		}
 		else
 		{
 			FPixelStreamingPlayerId PlayerId = Session->GetPlayerId();
-			VideoSource = new FVideoSourceP2P(SourceType, [PlayerId, this]() { return PlayerSessions.IsQualityController(PlayerId); });
+			VideoSource = VideoSourceGroup->CreateVideoSource([PlayerId, this]() { return PlayerSessions.IsQualityController(PlayerId); });
 		}
 
 		Session->SetVideoSource(VideoSource);
@@ -588,7 +616,6 @@
 		// Create video track
 		rtc::scoped_refptr<webrtc::VideoTrackInterface> VideoTrack = PCFactory->CreateVideoTrack(TCHAR_TO_UTF8(*InVideoTrackLabel), VideoSource.get());
 		VideoTrack->set_enabled(true);
-		VideoSource->Initialize();
 
 		// Set some content hints based on degradation prefs, WebRTC uses these internally.
 		webrtc::DegradationPreference DegradationPref = Settings::GetDegradationPreference();
@@ -606,8 +633,8 @@
 
 		bool bHasVideoTransceiver = false;
 
-		webrtc::PeerConnectionInterface& PeerConnection = Session->GetPeerConnection();
-		for (auto& Transceiver : PeerConnection.GetTransceivers())
+		rtc::scoped_refptr<webrtc::PeerConnectionInterface> PeerConnection = Session->GetPeerConnection();
+		for (auto& Transceiver : PeerConnection->GetTransceivers())
 		{
 			rtc::scoped_refptr<webrtc::RtpSenderInterface> Sender = Transceiver->sender();
 			if (Transceiver->media_type() == cricket::MediaType::MEDIA_TYPE_VIDEO)
@@ -625,48 +652,11 @@
 			TransceiverOptions.direction = webrtc::RtpTransceiverDirection::kSendOnly;
 			TransceiverOptions.send_encodings = CreateRTPEncodingParams(Flags);
 
-			webrtc::RTCErrorOr<rtc::scoped_refptr<webrtc::RtpTransceiverInterface>> Result = PeerConnection.AddTransceiver(VideoTrack, TransceiverOptions);
+			webrtc::RTCErrorOr<rtc::scoped_refptr<webrtc::RtpTransceiverInterface>> Result = PeerConnection->AddTransceiver(VideoTrack, TransceiverOptions);
 			checkf(Result.ok(), TEXT("Failed to add Video transceiver to PeerConnection. Msg=%s"), *FString(Result.error().message()));
 		}
 	}
 
-	void FStreamer::MungeLocalSDP(cricket::SessionDescription* SessionDescription)
-	{
-
-		std::vector<cricket::ContentInfo>& ContentInfos = SessionDescription->contents();
-		for (cricket::ContentInfo& ContentInfo : ContentInfos)
-		{
-			cricket::MediaContentDescription* MediaDescription = ContentInfo.media_description();
-			cricket::MediaType MediaType = MediaDescription->type();
-			if (MediaType != cricket::MediaType::MEDIA_TYPE_AUDIO)
-			{
-				continue;
-			}
-			cricket::AudioContentDescription* AudioDescription = MediaDescription->as_audio();
-			if (AudioDescription == nullptr)
-			{
-				continue;
-			}
-			std::vector<cricket::AudioCodec> CodecsCopy = AudioDescription->codecs();
-			for (cricket::AudioCodec& Codec : CodecsCopy)
-			{
-				if (Codec.name == "opus")
-				{
-					Codec.SetParam(cricket::kCodecParamPTime, "20");
-					Codec.SetParam(cricket::kCodecParamMaxPTime, "120");
-					Codec.SetParam(cricket::kCodecParamMinPTime, "3");
-					Codec.SetParam(cricket::kCodecParamSPropStereo, "1");
-					Codec.SetParam(cricket::kCodecParamStereo, "1");
-					Codec.SetParam(cricket::kCodecParamUseInbandFec, "1");
-					Codec.SetParam(cricket::kCodecParamUseDtx, "0");
-					Codec.SetParam(cricket::kCodecParamMaxAverageBitrate, "510000");
-					Codec.SetParam(cricket::kCodecParamMaxPlaybackRate, "48000");
-				}
-			}
-			AudioDescription->set_codecs(CodecsCopy);
-		}
-	}
-
 	void FStreamer::SetupAudioTrack(TSharedPtr<IPlayerSession> Session, FString const InAudioStreamId, FString const InAudioTrackLabel, int Flags)
 	{
 		bool bIsSFU = (Flags & Protocol::EPlayerFlags::PSPFlag_IsSFU) != Protocol::EPlayerFlags::PSPFlag_None;
@@ -707,8 +697,8 @@
 		rtc::scoped_refptr<webrtc::AudioTrackInterface> AudioTrack = PCFactory->CreateAudioTrack(TCHAR_TO_UTF8(*InAudioTrackLabel), AudioSource);
 
 		bool bHasAudioTransceiver = false;
-		webrtc::PeerConnectionInterface& PeerConnection = Session->GetPeerConnection();
-		for (auto& Transceiver : PeerConnection.GetTransceivers())
+		rtc::scoped_refptr<webrtc::PeerConnectionInterface> PeerConnection = Session->GetPeerConnection();
+		for (auto& Transceiver : PeerConnection->GetTransceivers())
 		{
 			rtc::scoped_refptr<webrtc::RtpSenderInterface> Sender = Transceiver->sender();
 			if (Transceiver->media_type() == cricket::MediaType::MEDIA_TYPE_AUDIO)
@@ -721,7 +711,7 @@
 		if (!bHasAudioTransceiver)
 		{
 			// Add the track
-			webrtc::RTCErrorOr<rtc::scoped_refptr<webrtc::RtpSenderInterface>> Result = PeerConnection.AddTrack(AudioTrack, { TCHAR_TO_UTF8(*InAudioStreamId) });
+			webrtc::RTCErrorOr<rtc::scoped_refptr<webrtc::RtpSenderInterface>> Result = PeerConnection->AddTrack(AudioTrack, { TCHAR_TO_UTF8(*InAudioStreamId) });
 
 			if (!Result.ok())
 			{
@@ -730,95 +720,282 @@
 		}
 	}
 
-	void FStreamer::ModifyTransceivers(std::vector<rtc::scoped_refptr<webrtc::RtpTransceiverInterface>> Transceivers, int Flags)
+	std::vector<webrtc::RtpEncodingParameters> FStreamer::CreateRTPEncodingParams(int Flags)
 	{
-		bool bTransmitUEAudio = !Settings::CVarPixelStreamingWebRTCDisableTransmitAudio.GetValueOnAnyThread();
 		bool bIsSFU = (Flags & Protocol::EPlayerFlags::PSPFlag_IsSFU) != Protocol::EPlayerFlags::PSPFlag_None;
-		bool bReceiveBrowserAudio = !bIsSFU && !Settings::CVarPixelStreamingWebRTCDisableReceiveAudio.GetValueOnAnyThread();
 
-		for (auto& Transceiver : Transceivers)
+		std::vector<webrtc::RtpEncodingParameters> EncodingParams;
+
+		if (bIsSFU && Settings::SimulcastParameters.Layers.Num() > 0)
 		{
-			rtc::scoped_refptr<webrtc::RtpSenderInterface> Sender = Transceiver->sender();
+			using FLayer = Settings::FSimulcastParameters::FLayer;
 
-			if (Transceiver->media_type() == cricket::MediaType::MEDIA_TYPE_AUDIO)
+			// encodings should be lowest res to highest
+			TArray<FLayer*> SortedLayers;
+			for (FLayer& Layer : Settings::SimulcastParameters.Layers)
 			{
+				SortedLayers.Add(&Layer);
+			}
 
-				// Determine the direction of the transceiver
-				webrtc::RtpTransceiverDirection AudioTransceiverDirection;
-				if (bTransmitUEAudio && bReceiveBrowserAudio)
+			SortedLayers.Sort([](const FLayer& LayerA, const FLayer& LayerB) { return LayerA.Scaling > LayerB.Scaling; });
+
+			const int LayerCount = SortedLayers.Num();
+			for (int i = 0; i < LayerCount; ++i)
+			{
+				const FLayer* SimulcastLayer = SortedLayers[i];
+				webrtc::RtpEncodingParameters LayerEncoding{};
+				LayerEncoding.rid = TCHAR_TO_UTF8(*(FString("simulcast") + FString::FromInt(LayerCount - i)));
+				LayerEncoding.min_bitrate_bps = SimulcastLayer->MinBitrate;
+				LayerEncoding.max_bitrate_bps = SimulcastLayer->MaxBitrate;
+				LayerEncoding.scale_resolution_down_by = SimulcastLayer->Scaling;
+
+				// In M84 this will crash with "Attempted to set an unimplemented parameter of RtpParameters".
+				// Try re-enabling this when we upgrade WebRTC versions.
+				// LayerEncoding.network_priority = webrtc::Priority::kHigh;
+
+				LayerEncoding.max_framerate = FMath::Max(60, Settings::CVarPixelStreamingWebRTCFps.GetValueOnAnyThread());
+				EncodingParams.push_back(LayerEncoding);
+			}
+		}
+		else
+		{
+			webrtc::RtpEncodingParameters Encoding{};
+			Encoding.rid = "base";
+			Encoding.max_bitrate_bps = Settings::CVarPixelStreamingWebRTCMaxBitrate.GetValueOnAnyThread();
+			Encoding.min_bitrate_bps = Settings::CVarPixelStreamingWebRTCMinBitrate.GetValueOnAnyThread();
+			Encoding.max_framerate = FMath::Max(60, Settings::CVarPixelStreamingWebRTCFps.GetValueOnAnyThread());
+			Encoding.scale_resolution_down_by.reset();
+			Encoding.network_priority = webrtc::Priority::kHigh;
+			EncodingParams.push_back(Encoding);
+		}
+
+		return EncodingParams;
+	}
+
+	void FStreamer::SendAnswer(TSharedPtr<IPlayerSession> Session, TUniquePtr<webrtc::SessionDescriptionInterface> Sdp)
+	{
+		// below is async execution (with error handling) of:
+		//		PeerConnection.SetRemoteDescription(SDP);
+		//		Answer = PeerConnection.CreateAnswer();
+		//		PeerConnection.SetLocalDescription(Answer);
+		//		SignallingServerConnection.SendAnswer(Answer);
+
+		FPixelStreamingSetSessionDescriptionObserver* SetLocalDescriptionObserver = FPixelStreamingSetSessionDescriptionObserver::Create(
+			[this, Session]() // on success
+			{
+				SignallingServerConnection->SendAnswer(Session->GetPlayerId(), *Session->GetPeerConnection()->local_description());
+				bStreamingStarted = true;
+			},
+			[this, Session](const FString& Error) // on failure
+			{
+				UE_LOG(LogPixelStreaming, Error, TEXT("Failed to set local description: %s"), *Error);
+				Session->DisconnectPlayer(Error);
+				PlayerSessions.DeletePlayerSession(Session->GetPlayerId());
+			});
+
+		auto OnCreateAnswerSuccess = [this, Session, SetLocalDescriptionObserver](webrtc::SessionDescriptionInterface* SDP) {
+			SetLocalDescription(Session->GetPeerConnection(), SetLocalDescriptionObserver, SDP);
+		};
+
+		FPixelStreamingCreateSessionDescriptionObserver* CreateAnswerObserver = FPixelStreamingCreateSessionDescriptionObserver::Create(
+			MoveTemp(OnCreateAnswerSuccess),
+			[this, Session](const FString& Error) // on failure
+			{
+				UE_LOG(LogPixelStreaming, Error, TEXT("Failed to create answer: %s"), *Error);
+				Session->DisconnectPlayer(Error);
+				PlayerSessions.DeletePlayerSession(Session->GetPlayerId());
+			});
+
+		auto OnSetRemoteDescriptionSuccess = [this, Session, CreateAnswerObserver]() {
+			// Note: these offer to receive are superseded now we are use transceivers to setup our peer connection media
+			int offer_to_receive_video = webrtc::PeerConnectionInterface::RTCOfferAnswerOptions::kUndefined;
+			int offer_to_receive_audio = webrtc::PeerConnectionInterface::RTCOfferAnswerOptions::kUndefined;
+			bool voice_activity_detection = false;
+			bool ice_restart = true;
+			bool use_rtp_mux = true;
+
+			webrtc::PeerConnectionInterface::RTCOfferAnswerOptions AnswerOption{
+				offer_to_receive_video,
+				offer_to_receive_audio,
+				voice_activity_detection,
+				ice_restart,
+				use_rtp_mux
+			};
+
+			AddStreams(Session, Protocol::EPlayerFlags::PSPFlag_SupportsDataChannel);
+			ModifyTransceivers(Session->GetPeerConnection()->GetTransceivers(), Protocol::EPlayerFlags::PSPFlag_SupportsDataChannel);
+			Session->GetPeerConnection()->CreateAnswer(CreateAnswerObserver, AnswerOption);
+		};
+		FPixelStreamingSetSessionDescriptionObserver* SetRemoteDescriptionObserver = FPixelStreamingSetSessionDescriptionObserver::Create(
+			MoveTemp(OnSetRemoteDescriptionSuccess),
+			[this, Session](const FString& Error) // on failure
+			{
+				UE_LOG(LogPixelStreaming, Error, TEXT("Failed to set remote description: %s"), *Error);
+				Session->DisconnectPlayer(Error);
+				PlayerSessions.DeletePlayerSession(Session->GetPlayerId());
+			});
+
+		cricket::SessionDescription* RemoteDescription = Sdp->description();
+		MungeRemoteSDP(RemoteDescription);
+
+		Session->GetPeerConnection()->SetRemoteDescription(SetRemoteDescriptionObserver, Sdp.Release());
+	}
+
+	void FStreamer::OnDataChannelOpen(FPixelStreamingPlayerId PlayerId, webrtc::DataChannelInterface* DataChannel)
+	{
+		if (DataChannel)
+		{
+			// When data channel is open, try to send cached freeze frame (if we have one)
+			SendCachedFreezeFrameTo(PlayerId);
+		}
+	}
+
+	void FStreamer::SetLocalDescription(webrtc::PeerConnectionInterface* PeerConnection, FPixelStreamingSetSessionDescriptionObserver* Observer, webrtc::SessionDescriptionInterface* SDP)
+	{
+		// Note from Luke about WebRTC: the sink of video capturer will be added as a direct result
+		// of `PeerConnection->SetLocalDescription()` call but video encoder will be created later on
+		// when the first frame is pushed into the WebRTC pipeline (by the capturer calling `OnFrame`).
+
+		PeerConnection->SetLocalDescription(Observer, SDP);
+
+		// Once local description has been set we can start setting some encoding information for the video stream rtp sender
+		for (rtc::scoped_refptr<webrtc::RtpSenderInterface> Sender : PeerConnection->GetSenders())
+		{
+			cricket::MediaType MediaType = Sender->media_type();
+			if (MediaType == cricket::MediaType::MEDIA_TYPE_VIDEO)
+			{
+				webrtc::RtpParameters ExistingParams = Sender->GetParameters();
+
+				// Set the degradation preference based on our CVar for it.
+				ExistingParams.degradation_preference = Settings::GetDegradationPreference();
+
+				webrtc::RTCError Err = Sender->SetParameters(ExistingParams);
+				if (!Err.ok())
 				{
-					AudioTransceiverDirection = webrtc::RtpTransceiverDirection::kSendRecv;
+					const char* ErrMsg = Err.message();
+					FString ErrorStr(ErrMsg);
+					UE_LOG(LogPixelStreaming, Error, TEXT("Failed to set RTP Sender params: %s"), *ErrorStr);
 				}
-				else if (bTransmitUEAudio)
-				{
-					AudioTransceiverDirection = webrtc::RtpTransceiverDirection::kSendOnly;
-				}
-				else if (bReceiveBrowserAudio)
-				{
-					AudioTransceiverDirection = webrtc::RtpTransceiverDirection::kRecvOnly;
-				}
-				else
-				{
-					AudioTransceiverDirection = webrtc::RtpTransceiverDirection::kInactive;
-				}
-
-				Transceiver->SetDirection(AudioTransceiverDirection);
-				Sender->SetStreams({ TCHAR_TO_UTF8(*GetAudioStreamID()) });
-			}
-			else if (Transceiver->media_type() == cricket::MediaType::MEDIA_TYPE_VIDEO)
-			{
-				Transceiver->SetDirection(webrtc::RtpTransceiverDirection::kSendOnly);
-				Sender->SetStreams({ TCHAR_TO_UTF8(*GetVideoStreamID()) });
-
-				webrtc::RtpParameters RtpParams = Sender->GetParameters();
-				RtpParams.encodings = CreateRTPEncodingParams(Flags);
-				Sender->SetParameters(RtpParams);
 			}
 		}
 	}
 
-	void FStreamer::SendPlayerMessage(Protocol::EToPlayerMsg Type, const FString& Descriptor)
+	FString FStreamer::GetAudioStreamID() const
 	{
-		PlayerSessions.ForEachSession([&](TSharedPtr<IPlayerSession> PlayerSession) { PlayerSession->SendMessage(Type, Descriptor); });
+		bool bSyncVideoAndAudio = !Settings::CVarPixelStreamingWebRTCDisableAudioSync.GetValueOnAnyThread();
+		return bSyncVideoAndAudio ? TEXT("pixelstreaming_av_stream_id") : TEXT("pixelstreaming_audio_stream_id");
 	}
 
-	void FStreamer::SendFreezeFrame(const TArray64<uint8>& JpegBytes)
+	FString FStreamer::GetVideoStreamID() const
 	{
-		if (!bStreamingStarted)
-		{
-			return;
-		}
-
-		PlayerSessions.ForEachSession([&JpegBytes](TSharedPtr<IPlayerSession> PlayerSession) { PlayerSession->SendFreezeFrame(JpegBytes); });
-
-		CachedJpegBytes = JpegBytes;
+		bool bSyncVideoAndAudio = !Settings::CVarPixelStreamingWebRTCDisableAudioSync.GetValueOnAnyThread();
+		return bSyncVideoAndAudio ? TEXT("pixelstreaming_av_stream_id") : TEXT("pixelstreaming_video_stream_id");
 	}
 
-	void FStreamer::SendCachedFreezeFrameTo(FPixelStreamingPlayerId PlayerId) const
+	rtc::scoped_refptr<webrtc::AudioDeviceModule> FStreamer::CreateAudioDeviceModule() const
 	{
-		if (CachedJpegBytes.Num() > 0)
+		bool bUseLegacyAudioDeviceModule = Settings::CVarPixelStreamingWebRTCUseLegacyAudioDevice.GetValueOnAnyThread();
+		rtc::scoped_refptr<webrtc::AudioDeviceModule> AudioDeviceModule;
+		if (bUseLegacyAudioDeviceModule)
 		{
-			PlayerSessions.ForSession(PlayerId, [&](TSharedPtr<IPlayerSession> Session) { Session->SendFreezeFrame(CachedJpegBytes); });
+			AudioDeviceModule = new rtc::RefCountedObject<FAudioCapturer>();
+		}
+		else
+		{
+			AudioDeviceModule = new rtc::RefCountedObject<FPixelStreamingAudioDeviceModule>();
+		}
+		return AudioDeviceModule;
+	}
+
+	rtc::scoped_refptr<webrtc::AudioProcessing> FStreamer::SetupAudioProcessingModule()
+	{
+#if WEBRTC_VERSION == 84
+		webrtc::AudioProcessing* AudioProcessingModule = webrtc::AudioProcessingBuilder().Create();
+#elif WEBRTC_VERSION == 96
+		rtc::scoped_refptr<webrtc::AudioProcessing> AudioProcessingModule = webrtc::AudioProcessingBuilder().Create();
+#endif 
+		webrtc::AudioProcessing::Config Config;
+		// Enabled multi channel audio capture/render
+		Config.pipeline.multi_channel_capture = true;
+		Config.pipeline.multi_channel_render = true;
+		Config.pipeline.maximum_internal_processing_rate = 48000;
+		// Turn off all other audio processing effects in UE's WebRTC. We want to stream audio from UE as pure as possible.
+		Config.pre_amplifier.enabled = false;
+		Config.high_pass_filter.enabled = false;
+		Config.echo_canceller.enabled = false;
+		Config.noise_suppression.enabled = false;
+		Config.transient_suppression.enabled = false;
+		Config.voice_detection.enabled = false;
+		Config.gain_controller1.enabled = false;
+		Config.gain_controller2.enabled = false;
+		Config.residual_echo_detector.enabled = false;
+		Config.level_estimation.enabled = false;
+		// Apply the config.
+		AudioProcessingModule->ApplyConfig(Config);
+
+		return AudioProcessingModule;
+	}
+
+	void FStreamer::StartWebRtcSignallingThread()
+	{
+		// Initialisation of WebRTC. Note: That interfacing with WebRTC should generally happen in WebRTC signalling thread.
+
+		// Create our own WebRTC thread for signalling
+		checkf(!WebRtcSignallingThread, TEXT("Signalling thread already exists."));
+		WebRtcSignallingThread = MakeUnique<rtc::Thread>(std::make_unique<rtc::PhysicalSocketServer>());
+		WebRtcSignallingThread->SetName("WebRtcSignallingThread", nullptr);
+		WebRtcSignallingThread->Start();
+
+		PeerConnectionConfig = {};
+
+		/* ---------- P2P Peer Connection Factory ---------- */
+
+		std::unique_ptr<FVideoEncoderFactorySimple> P2PPeerConnectionFactoryPtr = std::make_unique<FVideoEncoderFactorySimple>();
+		P2PVideoEncoderFactory = P2PPeerConnectionFactoryPtr.get();
+
+		P2PPeerConnectionFactory = webrtc::CreatePeerConnectionFactory(
+			nullptr,					  // network_thread
+			nullptr,					  // worker_thread
+			WebRtcSignallingThread.Get(), // signal_thread
+			CreateAudioDeviceModule(),	  // audio device module
+			webrtc::CreateAudioEncoderFactory<webrtc::AudioEncoderOpus>(),
+			webrtc::CreateAudioDecoderFactory<webrtc::AudioDecoderOpus>(),
+			std::move(P2PPeerConnectionFactoryPtr),
+			std::make_unique<webrtc::InternalDecoderFactory>(),
+			nullptr,					   // audio_mixer
+			SetupAudioProcessingModule()); // audio_processing
+
+		check(P2PPeerConnectionFactory.get() != nullptr);
+
+		/* ---------- SFU Peer Connection Factory ---------- */
+
+		std::unique_ptr<FVideoEncoderFactorySimulcast> SFUPeerConnectionFactoryPtr = std::make_unique<FVideoEncoderFactorySimulcast>();
+		SFUVideoEncoderFactory = SFUPeerConnectionFactoryPtr.get();
+
+		SFUPeerConnectionFactory = webrtc::CreatePeerConnectionFactory(
+			nullptr,					  // network_thread
+			nullptr,					  // worker_thread
+			WebRtcSignallingThread.Get(), // signal_thread
+			CreateAudioDeviceModule(),	  // audio device module
+			webrtc::CreateAudioEncoderFactory<webrtc::AudioEncoderOpus>(),
+			webrtc::CreateAudioDecoderFactory<webrtc::AudioDecoderOpus>(),
+			std::move(SFUPeerConnectionFactoryPtr),
+			std::make_unique<webrtc::InternalDecoderFactory>(),
+			nullptr,					   // audio_mixer
+			SetupAudioProcessingModule()); // audio_processing
+
+		check(SFUPeerConnectionFactory.get() != nullptr);
+	}
+
+	void FStreamer::StopWebRtcSignallingThread()
+	{
+		if (WebRtcSignallingThread)
+		{
+			WebRtcSignallingThread->Stop();
+			WebRtcSignallingThread = nullptr;
 		}
 	}
-
-	void FStreamer::SendUnfreezeFrame()
-	{
-		// Force a keyframe so when stream unfreezes if player has never received a h.264 frame before they can still connect.
-		ForceKeyFrame();
-
-		PlayerSessions.ForEachSession([](TSharedPtr<IPlayerSession> PlayerSession) { PlayerSession->SendUnfreezeFrame(); });
-
-		CachedJpegBytes.Empty();
-	}
-	void FStreamer::SendFileData(TArray<uint8>& ByteData, FString& MimeType, FString& FileExtension)
-	{
-		PlayerSessions.ForEachSession([&](TSharedPtr<IPlayerSession> PlayerSession) { PlayerSession->SendFileData(ByteData, MimeType, FileExtension); });
-	}
-
-	void FStreamer::KickPlayer(FPixelStreamingPlayerId PlayerId)
-	{
-		SignallingServerConnection->SendDisconnectPlayer(PlayerId, TEXT("Player was kicked"));
-	}
-
+	/** 
+	 * End own methods
+	 */	
 } // namespace UE::PixelStreaming
diff --git a/Engine/Plugins/Media/PixelStreaming/Source/PixelStreaming/Private/Streamer.h b/Engine/Plugins/Media/PixelStreaming/Source/PixelStreaming/Private/Streamer.h
--- a/Engine/Plugins/Media/PixelStreaming/Source/PixelStreaming/Private/Streamer.h
+++ b/Engine/Plugins/Media/PixelStreaming/Source/PixelStreaming/Private/Streamer.h
@@ -2,66 +2,82 @@
 
 #pragma once
 
-#include "SignallingServerConnection.h"
+#include "IPixelStreamingStreamer.h"
+#include "IPixelStreamingSignallingConnectionObserver.h"
 #include "PixelStreamingProtocolDefs.h"
 #include "RHI.h"
 #include "PlayerSessions.h"
-#include "FixedFPSPump.h"
-#include "Poller.h"
 #include "Stats.h"
+#include "StreamerInputDevices.h"
+#include "PixelStreamingSessionDescriptionObservers.h"
+#include "Dom/JsonObject.h"
 
 namespace UE::PixelStreaming
 {
 	class FPlayerSession;
-	class FVideoEncoderFactory;
-	class FSimulcastEncoderFactory;
+	class FVideoEncoderFactorySimple;
+	class FVideoEncoderFactorySimulcast;
 	class FSetSessionDescriptionObserver;
 	class FPoller;
 	class IPlayerSession;
+	class FVideoSourceGroup;
 
-	class FStreamer : public FSignallingServerConnectionObserver
+	class FStreamer : public IPixelStreamingStreamer, public IPixelStreamingSignallingConnectionObserver
 	{
 	public:
-		static bool IsPlatformCompatible();
-
 		FStreamer(const FString& StreamerId);
 		virtual ~FStreamer();
 
-		bool StartStreaming(const FString& SignallingServerUrl);
-		void StopStreaming();
-
-		bool IsStreaming() const { return bStreamingStarted; }
-		FPlayerSessions& GetPlayerSessions() { return PlayerSessions; }
-
-		void SendPlayerMessage(Protocol::EToPlayerMsg Type, const FString& Descriptor);
-		void SendFreezeFrame(const TArray64<uint8>& JpegBytes);
-		void SendUnfreezeFrame();
-		void SendFileData(TArray<uint8>& ByteData, FString& MimeType, FString& FileExtension);
-		void KickPlayer(FPixelStreamingPlayerId PlayerId);
+		/** IPixelStreamingStreamer implementation */
+		virtual void SetStreamFPS(int32 InFramesPerSecond) override;
+		virtual void SetVideoInput(TSharedPtr<FPixelStreamingVideoInput> Input) override;
+		virtual void SetTargetViewport(TSharedPtr<FSceneViewport> InTargetViewport) override;
+		virtual void SetSignallingServerURL(const FString& InSignallingServerURL) override;
+		virtual void StartStreaming() override;
+		virtual void StopStreaming() override;
+		virtual bool IsStreaming() const override { return bStreamingStarted; }
+		virtual void SendPlayerMessage(Protocol::EToPlayerMsg Type, const FString& Descriptor) override;
+		virtual void SetInputDevice(TSharedPtr<IPixelStreamingInputDevice> InInputDevice) override { InputDevice = InInputDevice; }
+		virtual void FreezeStream(UTexture2D* Texture) override;
+		virtual void UnfreezeStream() override;
+		virtual void SendFileData(TArray<uint8>& ByteData, FString& MimeType, FString& FileExtension) override;
+		virtual void KickPlayer(FPixelStreamingPlayerId PlayerId) override;
+		IPixelStreamingAudioSink* GetPeerAudioSink(FPixelStreamingPlayerId PlayerId) override;
+		IPixelStreamingAudioSink* GetUnlistenedAudioSink() override;
+		FStreamingStartedEvent& OnStreamingStarted() override;
+		FStreamingStoppedEvent& OnStreamingStopped() override;
+		virtual void RegisterCreateInputDevice(IPixelStreamingInputDevice::FCreateInputDeviceFunc& InCreateInputDevice) override;
+		/** End IPixelStreamingStreamer implementation */
+		
+		/** Own methods */
 		void ForceKeyFrame();
+		void PushFrame();
+		FPlayerSessions& GetPlayerSessions() { return PlayerSessions; }
+		// TODO(Luke) hook this back up so that the Engine can change how the interface is working browser side
+		void AddPlayerConfig(TSharedRef<FJsonObject>& JsonObject);		
+		/** End Own methods */
 
 	private:
-		TSharedPtr<IPlayerSession> CreateSession(FPixelStreamingPlayerId PlayerId, int Flags);
+		/** ISignallingServerConnectionObserver implementation */
+		virtual void OnSignallingConfig(const webrtc::PeerConnectionInterface::RTCConfiguration& Config) override;
+		virtual void OnSignallingSessionDescription(FPixelStreamingPlayerId PlayerId, webrtc::SdpType Type, const FString& Sdp) override;
+		virtual void OnSignallingRemoteIceCandidate(FPixelStreamingPlayerId PlayerId, const FString& SdpMid, int SdpMLineIndex, const FString& Sdp) override;
+		virtual void OnSignallingPlayerConnected(FPixelStreamingPlayerId PlayerId, int Flags) override;
+		virtual void OnSignallingPlayerDisconnected(FPixelStreamingPlayerId PlayerId) override;
+		virtual void OnSignallingSFUPeerDataChannels(FPixelStreamingPlayerId SFUId, FPixelStreamingPlayerId PlayerId, int32 SendStreamId, int32 RecvStreamId) override;
+		virtual void OnSignallingConnected() override;
+		virtual void OnSignallingDisconnected(int32 StatusCode, const FString& Reason, bool bWasClean) override;
+		virtual void OnSignallingError(const FString& ErrorMsg) override;
+		/** End ISignallingServerConnectionObserver implementation */
+
+		/** IInputDevice implementation */
+		virtual TSharedPtr<IInputDevice> CreateInputDevice(const TSharedRef<FGenericApplicationMessageHandler>& InMessageHandler) override;
+		/** End IInputDevice implementation */
+
+		/** Own methods */
+		void SendFreezeFrame(TArray<FColor> RawData, const FIntRect& Rect);
 		void SendCachedFreezeFrameTo(FPixelStreamingPlayerId PlayerId) const;
-
-		webrtc::AudioProcessing* SetupAudioProcessingModule();
-
-		// Procedure for WebRTC inter-thread communication
-		void StartWebRtcSignallingThread();
-		void ConnectToSignallingServer(const FString& Url);
-		void DisconnectFromSignallingServer();
-
-		// ISignallingServerConnectionObserver impl
-		virtual void OnConfig(const webrtc::PeerConnectionInterface::RTCConfiguration& Config) override;
-		virtual void OnSessionDescription(FPixelStreamingPlayerId PlayerId, webrtc::SdpType Type, const FString& Sdp) override;
-		virtual void OnRemoteIceCandidate(FPixelStreamingPlayerId PlayerId, const FString& SdpMid, int SdpMLineIndex, const FString& Sdp) override;
-		virtual void OnPlayerConnected(FPixelStreamingPlayerId PlayerId, int Flags) override;
-		virtual void OnPlayerDisconnected(FPixelStreamingPlayerId PlayerId) override;
-		virtual void OnStreamerDataChannels(FPixelStreamingPlayerId SFUId, FPixelStreamingPlayerId PlayerId, int32 SendStreamId, int32 RecvStreamId) override;
-		virtual void OnSignallingServerConnected() override;
-		virtual void OnSignallingServerDisconnected() override;
-
-		// own methods
+		TSharedPtr<IPlayerSession> CreateSession(FPixelStreamingPlayerId PlayerId, int Flags);
 		void OnOffer(FPixelStreamingPlayerId PlayerId, const FString& Sdp);
 		void ModifyTransceivers(std::vector<rtc::scoped_refptr<webrtc::RtpTransceiverInterface>> Transceivers, int Flags);
 		void MungeRemoteSDP(cricket::SessionDescription* SessionDescription);
@@ -74,32 +90,40 @@ namespace UE::PixelStreaming
 		std::vector<webrtc::RtpEncodingParameters> CreateRTPEncodingParams(int Flags);
 		void SendAnswer(TSharedPtr<IPlayerSession> Session, TUniquePtr<webrtc::SessionDescriptionInterface> Sdp);
 		void OnDataChannelOpen(FPixelStreamingPlayerId PlayerId, webrtc::DataChannelInterface* DataChannel);
-		void SetLocalDescription(webrtc::PeerConnectionInterface* PeerConnection, FSetSessionDescriptionObserver* Observer, webrtc::SessionDescriptionInterface* SDP);
+		void SetLocalDescription(webrtc::PeerConnectionInterface* PeerConnection, FPixelStreamingSetSessionDescriptionObserver* Observer, webrtc::SessionDescriptionInterface* SDP);
 		FString GetAudioStreamID() const;
 		FString GetVideoStreamID() const;
 		rtc::scoped_refptr<webrtc::AudioDeviceModule> CreateAudioDeviceModule() const;
+		rtc::scoped_refptr<webrtc::AudioProcessing> SetupAudioProcessingModule();
+		// Procedure for WebRTC inter-thread communication
+		void StartWebRtcSignallingThread();
+		void StopWebRtcSignallingThread();
+		/** End Own methods */
 
 	private:
+		bool bCaptureNextBackBufferAndStream = false;
+
 		FString StreamerId;
-		FString PreviousSignallingServerURL;
+		FString CurrentSignallingServerURL;
 		const FString AudioTrackLabel = TEXT("pixelstreaming_audio_track_label");
 		const FString VideoTrackLabel = TEXT("pixelstreaming_video_track_label");
 
 		TUniquePtr<rtc::Thread> WebRtcSignallingThread;
+		TSharedPtr<IPixelStreamingInputDevice> InputDevice;
 
-		TUniquePtr<FSignallingServerConnection> SignallingServerConnection;
+		TUniquePtr<FPixelStreamingSignallingConnection> SignallingServerConnection;
 		double LastSignallingServerConnectionAttemptTimestamp = 0;
 
 		webrtc::PeerConnectionInterface::RTCConfiguration PeerConnectionConfig;
 
 		/* P2P Peer Connections use these exclusively */
 		rtc::scoped_refptr<webrtc::PeerConnectionFactoryInterface> P2PPeerConnectionFactory;
-		FVideoEncoderFactory* P2PVideoEncoderFactory;
+		FVideoEncoderFactorySimple* P2PVideoEncoderFactory;
 		/* P2P Peer Connections use these exclusively */
 
 		/* SFU Peer Connections use these exclusively */
 		rtc::scoped_refptr<webrtc::PeerConnectionFactoryInterface> SFUPeerConnectionFactory;
-		FSimulcastEncoderFactory* SFUVideoEncoderFactory;
+		FVideoEncoderFactorySimulcast* SFUVideoEncoderFactory;
 		/* SFU Peer Connections use these exclusively */
 
 		rtc::scoped_refptr<webrtc::AudioSourceInterface> AudioSource;
@@ -108,12 +132,18 @@ namespace UE::PixelStreaming
 		// When we send a freeze frame we retain the data so we send freeze frame to new peers if they join during a freeze frame.
 		TArray64<uint8> CachedJpegBytes;
 
-		FThreadSafeBool bStreamingStarted = false;
+		bool bStreamingStarted = false;
 
 		FPlayerSessions PlayerSessions;
-		FStats Stats;
 
 		TUniquePtr<webrtc::SessionDescriptionInterface> SFULocalDescription;
 		TUniquePtr<webrtc::SessionDescriptionInterface> SFURemoteDescription;
+
+		TUniquePtr<FVideoSourceGroup> VideoSourceGroup;
+
+		FStreamingStartedEvent StreamingStartedEvent;
+		FStreamingStoppedEvent StreamingStoppedEvent;
+
+		TSharedPtr<FStreamerInputDevices> StreamerInputDevices;
 	};
 } // namespace UE::PixelStreaming
diff --git a/Engine/Plugins/Media/PixelStreaming/Source/PixelStreaming/Private/StreamerInputDevices.cpp b/Engine/Plugins/Media/PixelStreaming/Source/PixelStreaming/Private/StreamerInputDevices.cpp
new file mode 100644
--- /dev/null
+++ b/Engine/Plugins/Media/PixelStreaming/Source/PixelStreaming/Private/StreamerInputDevices.cpp
@@ -0,0 +1,104 @@
+// Copyright Epic Games, Inc. All Rights Reserved.
+#include "StreamerInputDevices.h"
+
+namespace UE::PixelStreaming
+{
+	FStreamerInputDevices::FStreamerInputDevices(const TSharedRef<FGenericApplicationMessageHandler>& InMessageHandler)
+		: MessageHandler(InMessageHandler)
+	{
+	}
+
+	TSharedPtr<IPixelStreamingInputDevice> FStreamerInputDevices::CreateInputDevice()
+	{
+		TSharedPtr<IPixelStreamingInputDevice> NewInputDevice;
+
+		if (OverridenCreateInputDevice)
+		{
+			NewInputDevice = OverridenCreateInputDevice(MessageHandler);
+		}
+		else
+		{
+			NewInputDevice = StaticCastSharedRef<IPixelStreamingInputDevice>(MakeShared<FInputDevice>(MessageHandler));			
+		}
+
+		{
+			FScopeLock Lock(&InputDeviceLock);
+			InputDevices.Add(NewInputDevice);
+		}
+
+		return NewInputDevice;
+	}
+
+	void FStreamerInputDevices::Tick(float DeltaTime)
+	{
+		FScopeLock Lock(&InputDeviceLock);
+		ForEachDevice([DeltaTime](IInputDevice* Device) {
+			Device->Tick(DeltaTime);
+		});
+	}
+
+	void FStreamerInputDevices::SendControllerEvents()
+	{
+		FScopeLock Lock(&InputDeviceLock);
+		ForEachDevice([](IInputDevice* Device) {
+			Device->SendControllerEvents();
+		});
+	}
+
+	void FStreamerInputDevices::SetMessageHandler(const TSharedRef<FGenericApplicationMessageHandler>& InMessageHandler)
+	{
+		FScopeLock Lock(&InputDeviceLock);
+		MessageHandler = InMessageHandler;
+		ForEachDevice([&MessageHandler = MessageHandler](IInputDevice* Device) {
+			Device->SetMessageHandler(MessageHandler);
+		});
+	}
+
+	bool FStreamerInputDevices::Exec(UWorld* InWorld, const TCHAR* Cmd, FOutputDevice& Ar)
+	{
+		FScopeLock Lock(&InputDeviceLock);
+		ForEachDevice([InWorld, Cmd, &Ar](IInputDevice* Device) {
+			Device->Exec(InWorld, Cmd, Ar);
+		});
+		return true;
+	}
+
+	void FStreamerInputDevices::SetChannelValue(int32 ControllerId, FForceFeedbackChannelType ChannelType, float Value)
+	{
+		FScopeLock Lock(&InputDeviceLock);
+		ForEachDevice([&ControllerId, &ChannelType, &Value](IInputDevice* Device) {
+			Device->SetChannelValue(ControllerId, ChannelType, Value);
+		});
+	}
+
+	void FStreamerInputDevices::SetChannelValues(int32 ControllerId, const FForceFeedbackValues& Values)
+	{
+		FScopeLock Lock(&InputDeviceLock);
+		ForEachDevice([&ControllerId, &Values](IInputDevice* Device) {
+			Device->SetChannelValues(ControllerId, Values);
+		});
+	}
+	
+	void FStreamerInputDevices::OverrideInputDevice(IPixelStreamingInputDevice::FCreateInputDeviceFunc& CreateInputDeviceFunc)
+	{
+		OverridenCreateInputDevice = CreateInputDeviceFunc;
+	}
+
+	void FStreamerInputDevices::ForEachDevice(TFunction<void(IInputDevice*)> const& Visitor)
+	{
+		for (int i = 0; i < InputDevices.Num(); /* skip */)
+		{
+			TWeakPtr<IInputDevice>& WeakDevice = InputDevices[i];
+			TSharedPtr<IInputDevice> Device = WeakDevice.Pin();
+			if (Device)
+			{
+				Visitor(Device.Get());
+				++i;
+			}
+			else
+			{
+				InputDevices.RemoveAt(i);
+			}
+		}
+	}
+} // namespace UE::PixelStreaming
diff --git a/Engine/Plugins/Media/PixelStreaming/Source/PixelStreaming/Private/StreamerInputDevices.h b/Engine/Plugins/Media/PixelStreaming/Source/PixelStreaming/Private/StreamerInputDevices.h
new file mode 100644
--- /dev/null
+++ b/Engine/Plugins/Media/PixelStreaming/Source/PixelStreaming/Private/StreamerInputDevices.h
@@ -0,0 +1,36 @@
+// Copyright Epic Games, Inc. All Rights Reserved.
+
+#pragma once
+
+#include "IPixelStreamingInputDevice.h"
+#include "InputDevice.h"
+
+namespace UE::PixelStreaming
+{
+	class FStreamerInputDevices : public IInputDevice
+	{
+	public:
+		FStreamerInputDevices(const TSharedRef<FGenericApplicationMessageHandler>& InMessageHandler);
+		virtual ~FStreamerInputDevices() = default;
+
+		TSharedPtr<IPixelStreamingInputDevice> CreateInputDevice();
+		virtual void Tick(float DeltaTime) override;
+		virtual void SendControllerEvents() override;
+		virtual void SetMessageHandler(const TSharedRef<FGenericApplicationMessageHandler>& InMessageHandler) override;
+		virtual bool Exec(UWorld* InWorld, const TCHAR* Cmd, FOutputDevice& Ar) override;
+		virtual void SetChannelValue(int32 ControllerId, FForceFeedbackChannelType ChannelType, float Value) override;
+		virtual void SetChannelValues(int32 ControllerId, const FForceFeedbackValues& Values) override;
+		
+		void OverrideInputDevice(IPixelStreamingInputDevice::FCreateInputDeviceFunc& CreateInputDeviceFunc);
+
+	private:
+		/** Reference to the message handler which events should be passed to. */
+		TSharedRef<FGenericApplicationMessageHandler> MessageHandler;
+		TArray<TWeakPtr<IInputDevice>> InputDevices;
+		FCriticalSection InputDeviceLock;
+				
+		IPixelStreamingInputDevice::FCreateInputDeviceFunc OverridenCreateInputDevice;
+
+		void ForEachDevice(TFunction<void(IInputDevice*)> const& Visitor);
+	};
+} // namespace UE::PixelStreaming
