Early Pixel Streaming 5.1 features, bugfixes, and refactors that benefit main. (From various authors)

diff --git a/Engine/Plugins/Media/PixelStreaming/Source/PixelStreaming/Private/Tests/FrameAdaptTest.cpp b/Engine/Plugins/Media/PixelStreaming/Source/PixelStreaming/Private/Tests/FrameAdaptTest.cpp
new file mode 100644
--- /dev/null
+++ b/Engine/Plugins/Media/PixelStreaming/Source/PixelStreaming/Private/Tests/FrameAdaptTest.cpp
@@ -0,0 +1,270 @@
+// Copyright Epic Games, Inc. All Rights Reserved.
+
+#include "Misc/AutomationTest.h"
+#include "Tests/AutomationCommon.h"
+#include "PixelStreamingFrameAdapterProcess.h"
+#include "PixelStreamingFrameAdapter.h"
+#include "PixelStreamingSourceFrame.h"
+#include "PixelStreamingVideoInput.h"
+
+#if WITH_DEV_AUTOMATION_TESTS
+
+namespace UE::PixelStreaming
+{
+	class FMockSourceFrame : public FPixelStreamingSourceFrame
+	{
+	public:
+		FMockSourceFrame(int32 InWidth, int32 InHeight)
+			: FPixelStreamingSourceFrame(nullptr)
+			, Width(InWidth)
+			, Height(InHeight)
+		{
+		}
+
+		virtual int32 GetWidth() const override { return Width; }
+		virtual int32 GetHeight() const override { return Height; }
+
+	private:
+		int32 Width;
+		int32 Height;
+	};
+
+	class FMockAdaptedVideoFrameLayer : public IPixelStreamingAdaptedVideoFrameLayer
+	{
+	public:
+		FMockAdaptedVideoFrameLayer(int32 InWidth, int32 InHeight)
+			: Width(InWidth)
+			, Height(InHeight)
+		{
+		}
+
+		virtual ~FMockAdaptedVideoFrameLayer() = default;
+
+		virtual int32 GetWidth() const override { return Width; }
+		virtual int32 GetHeight() const override { return Height; }
+
+		void SetWorking(bool InWorking) { Working = InWorking; }
+		bool IsWorking() const { return Working; }
+
+	private:
+		int32 Width = 0;
+		int32 Height = 0;
+		bool Working = false;
+	};
+
+	class FMockAdapterProcess : public FPixelStreamingFrameAdapterProcess
+	{
+	public:
+		FMockAdapterProcess(float InScale)
+		:Scale(InScale)
+		{
+		}
+
+		virtual ~FMockAdapterProcess() = default;
+
+		bool HasResChanged() const { return ResChanged; }
+		int GetBufferCreateCount() const { return BuffersCreated; }
+		FMockAdaptedVideoFrameLayer* GetCurrentWriteBuffer() { return StaticCast<FMockAdaptedVideoFrameLayer*>(GetWriteBuffer().Get()); }
+
+		void MockFinish()
+		{
+			FMockAdaptedVideoFrameLayer* WriteBuffer = StaticCast<FMockAdaptedVideoFrameLayer*>(GetWriteBuffer().Get());
+			WriteBuffer->SetWorking(false);
+			EndProcess();
+		}
+
+	protected:
+		virtual void OnSourceResolutionChanged(int32 OldWidth, int32 OldHeight, int32 NewWidth, int32 NewHeight) override
+		{
+			ResChanged = true;
+		}
+
+		virtual TSharedPtr<IPixelStreamingAdaptedVideoFrameLayer> CreateOutputBuffer(int32 SourceWidth, int32 SourceHeight) override
+		{
+			BuffersCreated++;
+			return MakeShared<FMockAdaptedVideoFrameLayer>(SourceWidth * Scale, SourceHeight * Scale);
+		}
+
+		virtual void BeginProcess(const FPixelStreamingSourceFrame& SourceFrame) override
+		{
+			FMockAdaptedVideoFrameLayer* WriteBuffer = StaticCast<FMockAdaptedVideoFrameLayer*>(GetWriteBuffer().Get());
+			WriteBuffer->SetWorking(true);
+		}
+
+	private:
+		float Scale = 1.0f;
+		bool ResChanged = false;
+		int BuffersCreated = 0;
+	};
+
+	class FMockFrameAdapter : public FPixelStreamingFrameAdapter
+	{
+	public:
+		FMockFrameAdapter(TSharedPtr<FPixelStreamingVideoInput> VideoInput)
+			: FPixelStreamingFrameAdapter(VideoInput)
+		{
+		}
+
+		FMockFrameAdapter(TSharedPtr<FPixelStreamingVideoInput> VideoInput, TArray<float> LayerScales)
+			: FPixelStreamingFrameAdapter(VideoInput, LayerScales)
+		{
+		}
+
+		virtual ~FMockFrameAdapter() = default;
+
+		FMockAdapterProcess* GetLayerAdapter(int LayerIndex)
+		{
+			return StaticCast<FMockAdapterProcess*>(LayerAdapters[LayerIndex].Get());
+		}
+
+	protected:
+		virtual TSharedPtr<FPixelStreamingFrameAdapterProcess> CreateAdaptProcess(float Scale) override
+		{
+			return MakeShared<FMockAdapterProcess>(Scale);
+		}
+	};
+
+	IMPLEMENT_SIMPLE_AUTOMATION_TEST(FAdaptProcessTest, "PixelStreaming.AdaptProcess", EAutomationTestFlags::ApplicationContextMask | EAutomationTestFlags::EngineFilter)
+	bool FAdaptProcessTest::RunTest(const FString& Parameters)
+	{
+		const int32 MockFrame1Width = 32;
+		const int32 MockFrame1Height = 96;
+
+		const int32 MockFrame2Width = 67;
+		const int32 MockFrame2Height = 101;
+
+		TSharedPtr<FMockAdapterProcess> MockAdapterProcess = MakeShared<FMockAdapterProcess>(1.0f);
+		TestFalse("Initialized after create.", MockAdapterProcess->IsInitialized());
+		TestFalse("Busy after create.", MockAdapterProcess->IsBusy());
+		TestFalse("HasOutput after create.", MockAdapterProcess->HasOutput());
+
+		FMockSourceFrame MockFrame1(MockFrame1Width, MockFrame1Height);
+		MockAdapterProcess->Process(MockFrame1);
+
+		TestTrue("Initialized after Process.", MockAdapterProcess->IsInitialized());
+		TestTrue("Busy after Process.", MockAdapterProcess->IsBusy());
+		TestFalse("HasOutput after Process.", MockAdapterProcess->HasOutput());
+		TestFalse("OnSourceResolutionChanged after Process.", MockAdapterProcess->HasResChanged());
+		TestTrue("Three output buffers created after Process.", MockAdapterProcess->GetBufferCreateCount() == 3);
+		TestTrue("Output has correct width after Process.", MockAdapterProcess->GetCurrentWriteBuffer()->GetWidth() == MockFrame1Width);
+		TestTrue("Output has correct height after Process.", MockAdapterProcess->GetCurrentWriteBuffer()->GetHeight() == MockFrame1Height);
+		TestTrue("Output is working after Process.", MockAdapterProcess->GetCurrentWriteBuffer()->IsWorking());
+
+		MockAdapterProcess->MockFinish();
+
+		TestTrue("Still Initialized after Finish.", MockAdapterProcess->IsInitialized());
+		TestFalse("Busy after Finish.", MockAdapterProcess->IsBusy());
+		TestTrue("HasOutput after Finish.", MockAdapterProcess->HasOutput());
+		TestTrue("GetOutputLayerWidth correct after Finish.", MockAdapterProcess->GetOutputLayerWidth() == MockFrame1Width);
+		TestTrue("GetOutputLayerHeight correct after Finish.", MockAdapterProcess->GetOutputLayerHeight() == MockFrame1Height);
+		TestFalse("OnSourceResolutionChanged after Finish.", MockAdapterProcess->HasResChanged());
+		TestTrue("No new buffers created after Finish.", MockAdapterProcess->GetBufferCreateCount() == 3);
+		TestTrue("Output has correct width after Finish.", MockAdapterProcess->GetCurrentWriteBuffer()->GetWidth() == MockFrame1Width);
+		TestTrue("Output has correct height after Finish.", MockAdapterProcess->GetCurrentWriteBuffer()->GetHeight() == MockFrame1Height);
+		TestFalse("Output is no longer working after Finish.", MockAdapterProcess->GetCurrentWriteBuffer()->IsWorking());
+
+		FMockSourceFrame MockFrame2(MockFrame2Width, MockFrame2Height);
+		MockAdapterProcess->Process(MockFrame2);
+
+		TestTrue("Still Initialized after Process new res.", MockAdapterProcess->IsInitialized());
+		TestTrue("Busy after Process new res.", MockAdapterProcess->IsBusy());
+		TestFalse("HasOutput after Process new res.", MockAdapterProcess->HasOutput());
+		TestTrue("OnSourceResolutionChanged after Process new res.", MockAdapterProcess->HasResChanged());
+		TestTrue("3 new buffers created after Process new res.", MockAdapterProcess->GetBufferCreateCount() == 6);
+		TestTrue("Output has correct width after Process new res.", MockAdapterProcess->GetCurrentWriteBuffer()->GetWidth() == MockFrame2Width);
+		TestTrue("Output has correct height after Process new res.", MockAdapterProcess->GetCurrentWriteBuffer()->GetHeight() == MockFrame2Height);
+		TestTrue("Output is working after Process new res.", MockAdapterProcess->GetCurrentWriteBuffer()->IsWorking());
+
+		MockAdapterProcess->MockFinish();
+
+		FMockSourceFrame MockFrame3(MockFrame1Width, MockFrame1Height);
+		MockAdapterProcess->Process(MockFrame3);
+		MockAdapterProcess->MockFinish();
+
+		TSharedPtr<FMockAdaptedVideoFrameLayer> ReadBuffer0 = StaticCastSharedPtr<FMockAdaptedVideoFrameLayer>(MockAdapterProcess->ReadOutput());
+		TSharedPtr<FMockAdaptedVideoFrameLayer> ReadBuffer1 = StaticCastSharedPtr<FMockAdaptedVideoFrameLayer>(MockAdapterProcess->ReadOutput());
+		TestTrue("Output buffer doesnt flip when input doesnt change", ReadBuffer0 == ReadBuffer1);
+
+		MockAdapterProcess->Process(MockFrame3);
+
+		TSharedPtr<FMockAdaptedVideoFrameLayer> ReadBuffer2 = StaticCastSharedPtr<FMockAdaptedVideoFrameLayer>(MockAdapterProcess->ReadOutput());
+		TestTrue("Output buffer doesnt flip when process hasnt completed", ReadBuffer0 == ReadBuffer2);
+
+		MockAdapterProcess->MockFinish();
+
+		TSharedPtr<FMockAdaptedVideoFrameLayer> ReadBuffer3 = StaticCastSharedPtr<FMockAdaptedVideoFrameLayer>(MockAdapterProcess->ReadOutput());
+		TestTrue("Output buffer does flip when process completed", ReadBuffer0 != ReadBuffer3);
+
+		MockAdapterProcess->Process(MockFrame3);
+		MockAdapterProcess->MockFinish();
+
+		TSharedPtr<FMockAdaptedVideoFrameLayer> ReadBuffer4 = StaticCastSharedPtr<FMockAdaptedVideoFrameLayer>(MockAdapterProcess->ReadOutput());
+		TestTrue("All three buffers read after the third finished process", ReadBuffer3 != ReadBuffer4 && ReadBuffer0 != ReadBuffer4);
+
+		return true;
+	}
+
+	IMPLEMENT_SIMPLE_AUTOMATION_TEST(FFrameAdapterTest, "PixelStreaming.FrameAdapterTest", EAutomationTestFlags::ApplicationContextMask | EAutomationTestFlags::EngineFilter)
+	bool FFrameAdapterTest::RunTest(const FString& Parameters)
+	{
+		TSharedPtr<FPixelStreamingVideoInput> MockVideoInput = MakeShared<FPixelStreamingVideoInput>();
+		TArray<float> LayerScales{ 1.0f, 0.5f, 0.25f };
+		FMockFrameAdapter MockFrameAdapter(MockVideoInput, LayerScales);
+
+		TestFalse("IsReady after create", MockFrameAdapter.IsReady());
+		TestTrue("GetNumLayers() returns expected count", MockFrameAdapter.GetNumLayers() == 0);
+
+		const int32 MockSourceFrameWidth = 1024;
+		const int32 MockSourceFrameHeight = 768;
+		FMockSourceFrame MockSourceFrame(1024, 768);
+		MockVideoInput->OnFrame.Broadcast(MockSourceFrame);
+
+		TestFalse("IsReady after OnFrame but before Adapt finishes.", MockFrameAdapter.IsReady());
+		TestTrue("GetNumLayers() doesnt change after OnFrame", MockFrameAdapter.GetNumLayers() == 3);
+
+		for (int i = 0; i < 3; ++i)
+		{
+			FMockAdapterProcess* Process = MockFrameAdapter.GetLayerAdapter(i);
+			TestTrue("Adapter process is busy after OnFrame", Process->IsBusy());
+			TestFalse("Adapter process has no output after OnFrame", Process->HasOutput());
+		}
+
+		for (int i = 0; i < 3; ++i)
+		{
+			FMockAdapterProcess* Process = MockFrameAdapter.GetLayerAdapter(i);
+			Process->MockFinish();
+		}
+
+		for (int i = 0; i < 3; ++i)
+		{
+			FMockAdapterProcess* Process = MockFrameAdapter.GetLayerAdapter(i);
+			TestFalse("Adapter process is not busy after MockFinish", Process->IsBusy());
+			TestTrue("Adapter process has output after MockFinish", Process->HasOutput());
+		}
+
+		TestTrue("IsReady after all adapt processes finished.", MockFrameAdapter.IsReady());
+		TestTrue("GetNumLayers() doesnt change after all adapt processes finished", MockFrameAdapter.GetNumLayers() == 3);
+
+		for (int i = 0; i < 3; ++i)
+		{
+			const int32 OutputWidth = MockFrameAdapter.GetWidth(i);
+			const int32 OutputHeight = MockFrameAdapter.GetHeight(i);
+			const int32 ExpectedWidth = MockSourceFrameWidth * LayerScales[i];
+			const int32 ExpectedHeight = MockSourceFrameHeight * LayerScales[i];
+
+			TestTrue("Adapter process output width matches expected", OutputWidth == ExpectedWidth);
+			TestTrue("Adapter process output height matches expected", OutputHeight == ExpectedHeight);
+
+			TSharedPtr<FMockAdaptedVideoFrameLayer> Output = StaticCastSharedPtr<FMockAdaptedVideoFrameLayer>(MockFrameAdapter.ReadOutput(i));
+
+			TestTrue("ReadOutput returns buffer with expected width", Output->GetWidth() == ExpectedWidth);
+			TestTrue("ReadOutput returns buffer with expected height", Output->GetHeight() == ExpectedHeight);
+		}
+
+
+		return true;
+	}
+
+} // namespace UE::PixelStreaming
+
+#endif // WITH_DEV_AUTOMATION_TESTS
diff --git a/Engine/Plugins/Media/PixelStreaming/Source/PixelStreaming/Private/Tests/FrameMetadataTest.cpp b/Engine/Plugins/Media/PixelStreaming/Source/PixelStreaming/Private/Tests/FrameMetadataTest.cpp
new file mode 100644
--- /dev/null
+++ b/Engine/Plugins/Media/PixelStreaming/Source/PixelStreaming/Private/Tests/FrameMetadataTest.cpp
@@ -0,0 +1,57 @@
+// Copyright Epic Games, Inc. All Rights Reserved.
+
+#include "Misc/AutomationTest.h"
+#include "Tests/AutomationCommon.h"
+#include "PixelStreamingFrameMetadata.h"
+
+#if WITH_DEV_AUTOMATION_TESTS
+
+namespace UE::PixelStreaming
+{
+
+	IMPLEMENT_SIMPLE_AUTOMATION_TEST(FFrameMetadataTest, "PixelStreaming.FrameMetadata", EAutomationTestFlags::ApplicationContextMask | EAutomationTestFlags::EngineFilter)
+	bool FFrameMetadataTest::RunTest(const FString& Parameters)
+	{
+		FPixelStreamingFrameMetadata Metadata;
+
+		TestTrue("Layer initialized", Metadata.Layer == 0);
+		TestTrue("SourceTime initialized", Metadata.SourceTime == 0);
+		TestTrue("AdaptCallTime initialized", Metadata.AdaptCallTime == 0);
+		TestTrue("AdaptProcessStartTime initialized", Metadata.AdaptProcessStartTime == 0);
+		TestTrue("AdaptProcessFinalizeTime initialized", Metadata.AdaptProcessFinalizeTime == 0);
+		TestTrue("AdaptProcessEndTime initialized", Metadata.AdaptProcessEndTime == 0);
+		TestTrue("UseCount initialized", Metadata.UseCount == 0);
+		TestTrue("FirstEncodeStartTime initialized", Metadata.FirstEncodeStartTime == 0);
+		TestTrue("LastEncodeStartTime initialized", Metadata.LastEncodeStartTime == 0);
+		TestTrue("LastEncodedEndTime initialized", Metadata.LastEncodeEndTime == 0);
+
+		Metadata.Layer = 1;
+		Metadata.SourceTime = 2;
+		Metadata.AdaptCallTime = 3;
+		Metadata.AdaptProcessStartTime = 4;
+		Metadata.AdaptProcessFinalizeTime = 5;
+		Metadata.AdaptProcessEndTime = 6;
+		Metadata.UseCount = 7;
+		Metadata.FirstEncodeStartTime = 8;
+		Metadata.LastEncodeStartTime = 9;
+		Metadata.LastEncodeEndTime = 10;
+
+		FPixelStreamingFrameMetadata MetadataCopy = Metadata.Copy();
+
+		TestTrue("Layer copied", MetadataCopy.Layer == 1);
+		TestTrue("SourceTime copied", MetadataCopy.SourceTime == 2);
+		TestTrue("AdaptCallTime copied", MetadataCopy.AdaptCallTime == 3);
+		TestTrue("AdaptProcessStartTime copied", MetadataCopy.AdaptProcessStartTime == 4);
+		TestTrue("AdaptProcessFinalizeTime copied", MetadataCopy.AdaptProcessFinalizeTime == 5);
+		TestTrue("AdaptProcessEndTime copied", MetadataCopy.AdaptProcessEndTime == 6);
+		TestTrue("UseCount copied", MetadataCopy.UseCount == 7);
+		TestTrue("FirstEncodeStartTime copied", MetadataCopy.FirstEncodeStartTime == 8);
+		TestTrue("LastEncodeStartTime copied", MetadataCopy.LastEncodeStartTime == 9);
+		TestTrue("LastEncodedEndTime copied", MetadataCopy.LastEncodeEndTime == 10);
+
+		return true;
+	}
+
+} // namespace UE::PixelStreaming
+
+#endif // WITH_DEV_AUTOMATION_TESTS
diff --git a/Engine/Plugins/Media/PixelStreaming/Source/PixelStreaming/Private/Tests/ServerTests.cpp b/Engine/Plugins/Media/PixelStreaming/Source/PixelStreaming/Private/Tests/ServerTests.cpp
new file mode 100644
--- /dev/null
+++ b/Engine/Plugins/Media/PixelStreaming/Source/PixelStreaming/Private/Tests/ServerTests.cpp
@@ -0,0 +1,85 @@
+// Copyright Epic Games, Inc. All Rights Reserved.
+
+#include "Misc/AutomationTest.h"
+#include "Tests/AutomationCommon.h"
+#include "PixelStreamingServers.h"
+#include "IPixelStreamingModule.h"
+#include "Misc/Paths.h"
+
+#if WITH_DEV_AUTOMATION_TESTS
+
+namespace UE::PixelStreaming::Servers
+{
+
+	DEFINE_LATENT_AUTOMATION_COMMAND_ONE_PARAMETER(FWaitForServer, TSharedPtr<FServerBase>, Server);
+	bool FWaitForServer::Update()
+	{
+		if (Server)
+		{
+			return Server->IsTimedOut() || Server->IsReady();
+		}
+		return true;
+	}
+
+	DEFINE_LATENT_AUTOMATION_COMMAND_ONE_PARAMETER(FCleanupServer, TSharedPtr<FServerBase>, Server);
+	bool FCleanupServer::Update()
+	{
+		if (Server)
+		{
+			Server->Stop();
+			return !FPaths::DirectoryExists(Server->GetPathOnDisk());
+		}
+		else
+		{
+			return true;
+		}
+	}
+
+	IMPLEMENT_SIMPLE_AUTOMATION_TEST(FLaunchTransientCirrusTest, "PixelStreaming.LaunchTransientCirrus", EAutomationTestFlags::ApplicationContextMask | EAutomationTestFlags::EngineFilter)
+	bool FLaunchTransientCirrusTest::RunTest(const FString& Parameters)
+	{
+
+		TSharedPtr<FServerBase> SignallingServer = MakeSignallingServer();
+		FLaunchArgs LaunchArgs;
+		LaunchArgs.bEphemeral = true;
+
+		SignallingServer->Launch(LaunchArgs);
+
+		TestTrue("Transient cirrus server directory exists", FPaths::DirectoryExists(SignallingServer->GetPathOnDisk()));
+
+		SignallingServer->Stop();
+
+		TestTrue("Transient cirrus server directory should not exist", !FPaths::DirectoryExists(SignallingServer->GetPathOnDisk()));
+
+		return true;
+	}
+
+	IMPLEMENT_SIMPLE_AUTOMATION_TEST(FLaunchTransientCirrusWithParams, "PixelStreaming.LaunchTransientCirrusWithParams", EAutomationTestFlags::ApplicationContextMask | EAutomationTestFlags::EngineFilter)
+	bool FLaunchTransientCirrusWithParams::RunTest(const FString& Parameters)
+	{
+
+		TSharedPtr<FServerBase> SignallingServer = MakeSignallingServer();
+		FLaunchArgs LaunchArgs;
+		LaunchArgs.bEphemeral = true;
+
+		SignallingServer->OnReady.AddLambda([this](TMap<EEndpoint, FString> Endpoints) {
+			TestTrue("Server was ready.", true);
+		});
+
+		SignallingServer->OnFailedToReady.AddLambda([this]() {
+			TestTrue("Server was ready.", false);
+		});
+
+		SignallingServer->Launch(LaunchArgs);
+
+		TestTrue("Transient cirrus server directory exists", FPaths::DirectoryExists(SignallingServer->GetPathOnDisk()));
+
+		ADD_LATENT_AUTOMATION_COMMAND(FWaitForServer(SignallingServer));
+		ADD_LATENT_AUTOMATION_COMMAND(FCleanupServer(SignallingServer));
+
+		return true;
+	}
+
+} // namespace UE::PixelStreaming::Servers
+
+#endif // WITH_DEV_AUTOMATION_TESTS
\ No newline at end of file
diff --git a/Engine/Plugins/Media/PixelStreaming/Source/PixelStreaming/Private/Tests/SignallingServerConnectionTests.cpp b/Engine/Plugins/Media/PixelStreaming/Source/PixelStreaming/Private/Tests/SignallingServerConnectionTests.cpp
--- a/Engine/Plugins/Media/PixelStreaming/Source/PixelStreaming/Private/Tests/SignallingServerConnectionTests.cpp
+++ b/Engine/Plugins/Media/PixelStreaming/Source/PixelStreaming/Private/Tests/SignallingServerConnectionTests.cpp
@@ -1,49 +1,50 @@
 // Copyright Epic Games, Inc. All Rights Reserved.
 #include "Misc/AutomationTest.h"
 #include "Tests/AutomationCommon.h"
-#include "SignallingServerConnection.h"
+#include "PixelStreamingSignallingConnection.h"
 #include "Serialization/JsonSerializer.h"
 #include "IWebSocket.h"
 
 namespace
 {
-	using namespace UE::PixelStreaming;
+	//using namespace UE::PixelStreaming;
 
 	// For validating callbacks from the signalling server connection
-	class FTestSSConnectionObserver : public FSignallingServerConnectionObserver
+	class FTestSSConnectionObserver : public IPixelStreamingSignallingConnectionObserver
 	{
 	public:
 		virtual ~FTestSSConnectionObserver() = default;
 
-		virtual void OnSignallingServerConnected() override {}
-		virtual void OnSignallingServerDisconnected() override {}
+		virtual void OnSignallingConnected() override {}
+		virtual void OnSignallingDisconnected(int32 StatusCode, const FString& Reason, bool bWasClean) override {}
+		virtual void OnSignallingError(const FString& ErrorMsg) override {}
 
-		virtual void OnConfig(const webrtc::PeerConnectionInterface::RTCConfiguration& Config) override { SetConfig = Config; }
+		virtual void OnSignallingConfig(const webrtc::PeerConnectionInterface::RTCConfiguration& Config) override { SetConfig = Config; }
 
 		//// Streamer-only
-		virtual void OnSessionDescription(FPixelStreamingPlayerId PlayerId, webrtc::SdpType Type, const FString& Sdp) override
+		virtual void OnSignallingSessionDescription(FPixelStreamingPlayerId PlayerId, webrtc::SdpType Type, const FString& Sdp) override
 		{
 			SetSDP = { Type, Sdp, PlayerId };
 		}
-		virtual void OnRemoteIceCandidate(FPixelStreamingPlayerId PlayerId, const FString& SdpMid, int SdpMLineIndex, const FString& Sdp)
+		virtual void OnSignallingRemoteIceCandidate(FPixelStreamingPlayerId PlayerId, const FString& SdpMid, int SdpMLineIndex, const FString& Sdp)
 		{
 			SetIceCandidate = { SdpMid, SdpMLineIndex, Sdp, PlayerId };
 		}
-		//virtual void OnPlayerConnected(FPixelStreamingPlayerId PlayerId, int Flags) { unimplemented(); }
-		//virtual void OnPlayerDisconnected(FPixelStreamingPlayerId PlayerId) { unimplemented(); }
-		//virtual void OnStreamerDataChannels(FPixelStreamingPlayerId SFUId, FPixelStreamingPlayerId PlayerId, int32 SendStreamId, int32 RecvStreamId) { unimplemented(); }
+		// virtual void OnPlayerConnected(FPixelStreamingPlayerId PlayerId, int Flags) { unimplemented(); }
+		// virtual void OnPlayerDisconnected(FPixelStreamingPlayerId PlayerId) { unimplemented(); }
+		// virtual void OnSFUPeerDataChannels(FPixelStreamingPlayerId SFUId, FPixelStreamingPlayerId PlayerId, int32 SendStreamId, int32 RecvStreamId) { unimplemented(); }
 
 		//// Player-only
-		virtual void OnSessionDescription(webrtc::SdpType Type, const FString& Sdp) override
+		virtual void OnSignallingSessionDescription(webrtc::SdpType Type, const FString& Sdp) override
 		{
 			SetSDP = { Type, Sdp };
 		}
-		virtual void OnRemoteIceCandidate(const FString& SdpMid, int SdpMLineIndex, const FString& Sdp)
+		virtual void OnSignallingRemoteIceCandidate(const FString& SdpMid, int SdpMLineIndex, const FString& Sdp)
 		{
 			SetIceCandidate = { SdpMid, SdpMLineIndex, Sdp };
 		}
-		//virtual void OnPlayerCount(uint32 Count) { unimplemented(); }
-		//virtual void OnPeerDataChannels(int32 SendStreamId, int32 RecvStreamId) { unimplemented(); }
+		// virtual void OnPlayerCount(uint32 Count) { unimplemented(); }
+		// virtual void OnPeerDataChannels(int32 SendStreamId, int32 RecvStreamId) { unimplemented(); }
 
 		void Reset()
 		{
@@ -142,12 +143,12 @@ namespace
 	};
 
 	BEGIN_DEFINE_SPEC(FSignallingServerConnectionSpec, "System.Plugins.PixelStreaming.SignallingServerConnection",  EAutomationTestFlags::EditorContext | EAutomationTestFlags::ClientContext | EAutomationTestFlags::ProductFilter)
-	TUniquePtr<FSignallingServerConnection> SignallingServerConnection; // the object being tested
+	TUniquePtr<FPixelStreamingSignallingConnection> SignallingServerConnection; // the object being tested
 	FTestSSConnectionObserver Observer;
-	FSignallingServerConnection::FWebSocketFactory WebSocketFactory;
+	FPixelStreamingSignallingConnection::FWebSocketFactory WebSocketFactory;
 	TSharedPtr<FMockWebSocket> SocketEndpoint; // the socket created for the SignallingServerConnection. Used to fake send data.
-	FDelegateHandle OnMessageDelegateHandle; // so we can remove the delegate after tests
-	TSharedPtr<FJsonObject> LastJsonMsg; // the parsed json of the last message received from the SignallingServerConnection
+	FDelegateHandle OnMessageDelegateHandle;   // so we can remove the delegate after tests
+	TSharedPtr<FJsonObject> LastJsonMsg;	   // the parsed json of the last message received from the SignallingServerConnection
 	const FString ServerID = "MOCK SERVER";
 	const FPixelStreamingPlayerId MockPlayerId = "MOCK";
 	const FMockSessionDescriptionInterface MockSDP;
@@ -158,11 +159,11 @@ namespace
 	{
 		Describe("Mocked connection tests", [this]() {
 			BeforeEach([this]() {
-				WebSocketFactory = [&](const FString&) {
+				WebSocketFactory = [this](const FString&) {
 					SocketEndpoint = MakeShared<FMockWebSocket>();
 					return SocketEndpoint;
 				};
-				SignallingServerConnection = MakeUnique<FSignallingServerConnection>(WebSocketFactory, Observer, ServerID);
+				SignallingServerConnection = MakeUnique<FPixelStreamingSignallingConnection>(WebSocketFactory, Observer, ServerID);
 				SignallingServerConnection->Connect("Fake Url");
 				OnMessageDelegateHandle = SocketEndpoint->OnMockResponseEvent.AddLambda([this](const FString& message) {
 					OnMessageReceived(message);
diff --git a/Engine/Plugins/Media/PixelStreaming/Source/PixelStreaming/Private/Tests/SourceFrameTest.cpp b/Engine/Plugins/Media/PixelStreaming/Source/PixelStreaming/Private/Tests/SourceFrameTest.cpp
new file mode 100644
--- /dev/null
+++ b/Engine/Plugins/Media/PixelStreaming/Source/PixelStreaming/Private/Tests/SourceFrameTest.cpp
@@ -0,0 +1,42 @@
+// Copyright Epic Games, Inc. All Rights Reserved.
+
+#include "Misc/AutomationTest.h"
+#include "Tests/AutomationCommon.h"
+#include "PixelStreamingSourceFrame.h"
+
+#if WITH_DEV_AUTOMATION_TESTS
+
+namespace UE::PixelStreaming
+{
+
+	IMPLEMENT_SIMPLE_AUTOMATION_TEST(FSourceFrameTest, "PixelStreaming.SourceFrame", EAutomationTestFlags::ApplicationContextMask | EAutomationTestFlags::EngineFilter)
+	bool FSourceFrameTest::RunTest(const FString& Parameters)
+	{
+		FRHITextureCreateDesc TextureDesc =
+			FRHITextureCreateDesc::Create2D(TEXT("Test Texture"), 512, 128, EPixelFormat::PF_B8G8R8A8)
+			.SetClearValue(FClearValueBinding::None)
+			.SetFlags(ETextureCreateFlags::None)
+			.SetInitialState(ERHIAccess::Present)
+			.DetermineInititialState();
+
+		FTexture2DRHIRef TestTexture = GDynamicRHI->RHICreateTexture(TextureDesc);
+		const uint64 PreCreateTime = FPlatformTime::Cycles64();
+
+		FPlatformProcess::Sleep(0.1f);
+
+		FPixelStreamingSourceFrame Frame(TestTexture);
+
+		FPlatformProcess::Sleep(0.1f);
+
+		const uint64 PostCreateTime = FPlatformTime::Cycles64();
+
+		TestTrue("Frame Width", Frame.GetWidth() == 512);
+		TestTrue("Frame Height", Frame.GetHeight() == 128);
+		TestTrue("Frame Create Time", Frame.Metadata.SourceTime > PreCreateTime && Frame.Metadata.SourceTime < PostCreateTime);
+
+		return true;
+	}
+
+} // namespace UE::PixelStreaming
+
+#endif // WITH_DEV_AUTOMATION_TESTS
diff --git a/Engine/Plugins/Media/PixelStreaming/Source/PixelStreaming/Private/Utils.h b/Engine/Plugins/Media/PixelStreaming/Source/PixelStreaming/Private/Utils.h
--- a/Engine/Plugins/Media/PixelStreaming/Source/PixelStreaming/Private/Utils.h
+++ b/Engine/Plugins/Media/PixelStreaming/Source/PixelStreaming/Private/Utils.h
@@ -2,19 +2,8 @@
 
 #pragma once
 
-#include "RHI.h"
-#include "CommonRenderResources.h"
-#include "ScreenRendering.h"
-#include "RHIStaticStates.h"
-#include "Modules/ModuleManager.h"
-#include "PixelStreamingPrivate.h"
 #include "WebRTCIncludes.h"
 #include "Async/Async.h"
-#include "RenderGraphBuilder.h"
-#include "RenderGraphUtils.h"
-#include "RHIGPUReadback.h"
-#include "PixelShaderUtils.h"
-#include "Runtime/Renderer/Private/ScreenPass.h"
 
 namespace UE::PixelStreaming
 {
@@ -49,7 +38,7 @@ namespace UE::PixelStreaming
 			FPlatformProcess::ReturnSynchEventToPool(TaskEvent);
 		}
 	}
-
+#if WEBRTC_VERSION == 84
 	inline webrtc::SdpVideoFormat CreateH264Format(webrtc::H264::Profile profile, webrtc::H264::Level level)
 	{
 		const absl::optional<std::string> ProfileString =
@@ -60,9 +49,23 @@ namespace UE::PixelStreaming
 			{ { cricket::kH264FmtpProfileLevelId, *ProfileString },
 				{ cricket::kH264FmtpLevelAsymmetryAllowed, "1" },
 				{ cricket::kH264FmtpPacketizationMode, "1" } });
+
 	}
+#elif WEBRTC_VERSION == 96
+	inline webrtc::SdpVideoFormat CreateH264Format(webrtc::H264Profile profile, webrtc::H264Level level)
+	{
+		const absl::optional<std::string> ProfileString =
+			webrtc::H264ProfileLevelIdToString(webrtc::H264ProfileLevelId(profile, level));
+		check(ProfileString);
+		return webrtc::SdpVideoFormat(
+			cricket::kH264CodecName,
+			{ { cricket::kH264FmtpProfileLevelId, *ProfileString },
+				{ cricket::kH264FmtpLevelAsymmetryAllowed, "1" },
+				{ cricket::kH264FmtpPacketizationMode, "1" } });
 
-	inline void MemCpyStride(void* Dest, const void* Src, size_t DestStride, size_t SrcStride, size_t Height)
+	}
+#endif
+inline void MemCpyStride(void* Dest, const void* Src, size_t DestStride, size_t SrcStride, size_t Height)
 	{
 		char* DestPtr = static_cast<char*>(Dest);
 		const char* SrcPtr = static_cast<const char*>(Src);
@@ -73,201 +76,13 @@ namespace UE::PixelStreaming
 		}
 	}
 
-	inline FTextureRHIRef CreateRHITexture(uint32 Width, uint32 Height)
-	{
-		// Create empty texture
-		FRHITextureCreateDesc Desc =
-			FRHITextureCreateDesc::Create2D(TEXT("PixelStreamingBlankTexture"), Width, Height, PF_B8G8R8A8)
-			.SetClearValue(FClearValueBinding::None)
-			.SetFlags(ETextureCreateFlags::RenderTargetable);
-
-		if (RHIGetInterfaceType() == ERHIInterfaceType::Vulkan)
-		{
-			Desc.AddFlags(ETextureCreateFlags::External)
-				.SetInitialState(ERHIAccess::Present);
-		}
-		else
-		{
-			Desc.AddFlags(ETextureCreateFlags::Shared)
-				.SetInitialState(ERHIAccess::CopyDest);
-		}
-
-		return GDynamicRHI->RHICreateTexture(Desc);
-	}
-
-	inline FTextureRHIRef CreateCPUReadbackTexture(uint32 Width, uint32 Height)
-	{
-		// Create empty texture
-		FRHITextureCreateDesc Desc =
-			FRHITextureCreateDesc::Create2D(TEXT("PixelStreamingBlankTexture"), Width, Height, PF_B8G8R8A8)
-			.SetClearValue(FClearValueBinding::None)
-			.SetFlags(ETextureCreateFlags::HideInVisualizeTexture | ETextureCreateFlags::CPUReadback);
-
-		if (RHIGetInterfaceType() == ERHIInterfaceType::Vulkan)
-		{
-			Desc.SetInitialState(ERHIAccess::Present);
-		}
-		else
-		{
-			Desc.SetInitialState(ERHIAccess::CopyDest);
-		}
-
-		return GDynamicRHI->RHICreateTexture(Desc);
-	}
-
-	inline void CopyTextureToRHI(FTextureRHIRef SourceTexture, FTextureRHIRef DestinationTexture, FGPUFenceRHIRef& CopyFence)
-	{
-		FRHICommandListImmediate& RHICmdList = FRHICommandListExecutor::GetImmediateCommandList();
-
-		IRendererModule* RendererModule = &FModuleManager::GetModuleChecked<IRendererModule>("Renderer");
-
-		// #todo-renderpasses there's no explicit resolve here? Do we need one?
-		FRHIRenderPassInfo RPInfo(DestinationTexture, ERenderTargetActions::Load_Store);
-
-		RHICmdList.BeginRenderPass(RPInfo, TEXT("CopyBackbuffer"));
-
-		{
-			RHICmdList.SetViewport(0, 0, 0.0f, DestinationTexture->GetDesc().Extent.X, DestinationTexture->GetDesc().Extent.Y, 1.0f);
-
-			FGraphicsPipelineStateInitializer GraphicsPSOInit;
-			RHICmdList.ApplyCachedRenderTargets(GraphicsPSOInit);
-			GraphicsPSOInit.BlendState = TStaticBlendState<>::GetRHI();
-			GraphicsPSOInit.RasterizerState = TStaticRasterizerState<>::GetRHI();
-			GraphicsPSOInit.DepthStencilState = TStaticDepthStencilState<false, CF_Always>::GetRHI();
-
-			// New engine version...
-			FGlobalShaderMap* ShaderMap = GetGlobalShaderMap(GMaxRHIFeatureLevel);
-			TShaderMapRef<FScreenVS> VertexShader(ShaderMap);
-			TShaderMapRef<FScreenPS> PixelShader(ShaderMap);
-
-			GraphicsPSOInit.BoundShaderState.VertexDeclarationRHI = GFilterVertexDeclaration.VertexDeclarationRHI;
-			GraphicsPSOInit.BoundShaderState.VertexShaderRHI = VertexShader.GetVertexShader();
-			GraphicsPSOInit.BoundShaderState.PixelShaderRHI = PixelShader.GetPixelShader();
-
-			GraphicsPSOInit.PrimitiveType = PT_TriangleList;
-
-			SetGraphicsPipelineState(RHICmdList, GraphicsPSOInit, 0);
-
-			if (DestinationTexture->GetDesc().Extent.X != SourceTexture->GetDesc().Extent.X || DestinationTexture->GetDesc().Extent.Y != SourceTexture->GetDesc().Extent.Y)
-			{
-				PixelShader->SetParameters(RHICmdList, TStaticSamplerState<SF_Bilinear>::GetRHI(), SourceTexture);
-			}
-			else
-			{
-				PixelShader->SetParameters(RHICmdList, TStaticSamplerState<SF_Point>::GetRHI(), SourceTexture);
-			}
-
-			FIntPoint TargetBufferSize(DestinationTexture->GetDesc().Extent.X, DestinationTexture->GetDesc().Extent.Y);
-			RendererModule->DrawRectangle(RHICmdList, 0, 0, // Dest X, Y
-				DestinationTexture->GetDesc().Extent.X,		// Dest Width
-				DestinationTexture->GetDesc().Extent.Y,		// Dest Height
-				0, 0,										// Source U, V
-				1, 1,										// Source USize, VSize
-				TargetBufferSize,							// Target buffer size
-				FIntPoint(1, 1),							// Source texture size
-				VertexShader, EDRF_Default);
-		}
-
-		RHICmdList.EndRenderPass();
-
-		RHICmdList.WriteGPUFence(CopyFence);
-
-		RHICmdList.ImmediateFlush(EImmediateFlushType::WaitForOutstandingTasksOnly);
-	}
-
-	inline void CopyTextureToReadbackTexture(FTextureRHIRef SourceTexture, TSharedPtr<FRHIGPUTextureReadback> GPUTextureReadback, void* OutBuffer)
-	{
-		FRHICommandListImmediate& RHICmdList = FRHICommandListExecutor::GetImmediateCommandList();
-		FRDGBuilder GraphBuilder(RHICmdList);
-
-		FRDGTextureRef RDGSourceTexture = RegisterExternalTexture(GraphBuilder, SourceTexture, TEXT("SourceCopyTextureToReadbackTexture"));
-		FRDGTextureRef RDGStagingTexture = nullptr;
-
-		if(RDGSourceTexture->Desc.Format != EPixelFormat::PF_B8G8R8A8)
-		{
-			// We need the pixel format to be BGRA8 so we first draw it to a staging texture
-			{
-				const FRDGTextureDesc Desc(FRDGTextureDesc::Create2D(RDGSourceTexture->Desc.Extent, PF_B8G8R8A8, FClearValueBinding::None, TexCreate_ShaderResource | TexCreate_RenderTargetable));
-				RDGStagingTexture = GraphBuilder.CreateTexture(Desc, TEXT("StagingCopyTextureToReadbackTexture"));
-			}
-
-			FGlobalShaderMap* ShaderMap = GetGlobalShaderMap(GMaxRHIFeatureLevel);
-
-			TShaderMapRef<FScreenPassVS> VertexShader(ShaderMap);
-
-			// Setup the pixel shader
-			TShaderMapRef<FCopyRectPS> PixelShader(ShaderMap);
-
-			FCopyRectPS::FParameters* PixelShaderParameters = GraphBuilder.AllocParameters<FCopyRectPS::FParameters>();
-			PixelShaderParameters->InputTexture = RDGSourceTexture;
-			PixelShaderParameters->InputSampler = TStaticSamplerState<SF_Bilinear, AM_Clamp, AM_Clamp, AM_Clamp>::GetRHI();
-			PixelShaderParameters->RenderTargets[0] = FRenderTargetBinding(RDGStagingTexture, ERenderTargetLoadAction::ELoad);
-
-			ClearUnusedGraphResources(PixelShader, PixelShaderParameters);
-
-			// We are not doing any clever blending stuff so we just use defaults here
-			FRHIBlendState* BlendState = FScreenPassPipelineState::FDefaultBlendState::GetRHI();
-			FRHIDepthStencilState* DepthStencilState = FScreenPassPipelineState::FDefaultDepthStencilState::GetRHI();
-
-			// Create the pipline state that will execute
-			const FScreenPassPipelineState PipelineState(VertexShader, PixelShader, BlendState, DepthStencilState);
-
-			// Add the pass the the graph builder
-			GraphBuilder.AddPass(
-				RDG_EVENT_NAME("PixelStreamingChangePixelFormat"),
-				PixelShaderParameters,
-				ERDGPassFlags::Raster,
-				[PipelineState, Extent = RDGSourceTexture->Desc.Extent, PixelShader, PixelShaderParameters](FRHICommandList& RHICmdList)
-				{
-					PipelineState.Validate();
-					
-					RHICmdList.SetViewport(0.0f, 0.0f, 0.0f, Extent.X, Extent.Y, 1.0f);
-
-					SetScreenPassPipelineState(RHICmdList, PipelineState);
-
-					SetShaderParameters(RHICmdList, PixelShader, PixelShader.GetPixelShader(), *PixelShaderParameters);
-
-					DrawRectangle(
-						RHICmdList,
-						0, 0, Extent.X, Extent.Y,
-						0, 0, Extent.X, Extent.Y,
-						Extent,
-						Extent,
-						PipelineState.VertexShader,
-						EDRF_UseTriangleOptimization);
-				});
-		}
-		else
-		{
-			// Otherwise if the PixelFormat is already BGRA8 we can just use the SourceTexture as the staging texture
-			RDGStagingTexture = RDGSourceTexture;
-		}
-
-		// Do the copy from staging RBGA8 texture into the readback
-		AddEnqueueCopyPass(GraphBuilder, GPUTextureReadback.Get(), RDGStagingTexture);
-
-		GraphBuilder.Execute();
-
-		// Lock and copy out the content of the TextureReadback to the CPU
-		int32 OutRowPitchInPixels;
-		int32 BlockSize = GPixelFormats[RDGStagingTexture->Desc.Format].BlockBytes;
-		void* ResultsBuffer = GPUTextureReadback->Lock(OutRowPitchInPixels);
-		if (RDGSourceTexture->Desc.Extent.X == OutRowPitchInPixels)
-		{
-			// Source pixel width is the same as the stride of the result buffer (ie no padding), we can do a plain memcpy
-			FPlatformMemory::Memcpy(OutBuffer, ResultsBuffer, (RDGSourceTexture->Desc.Extent.X * RDGSourceTexture->Desc.Extent.Y * BlockSize));
-		} 
-		else
-		{
-			// Source pixel width differs from the stride of the result buffer (ie padding), do a memcpy that accounts for this
-			MemCpyStride(OutBuffer, ResultsBuffer, RDGSourceTexture->Desc.Extent.X * BlockSize, OutRowPitchInPixels * BlockSize, RDGSourceTexture->Desc.Extent.Y); 
-		}
-		GPUTextureReadback->Unlock();
-	}
-
 	inline size_t SerializeToBuffer(rtc::CopyOnWriteBuffer& Buffer, size_t Pos, const void* Data, size_t DataSize)
 	{
-		FMemory::Memcpy(&Buffer[Pos], reinterpret_cast<const uint8_t*>(Data), DataSize);
+#if WEBRTC_VERSION == 84
+		FMemory::Memcpy(&Buffer[Pos], Data, DataSize);
+#elif WEBRTC_VERSION == 96
+		FMemory::Memcpy(Buffer.MutableData() + Pos, Data, DataSize);
+#endif
 		return Pos + DataSize;
 	}
-} // namespace UE::PixelStreaming
\ No newline at end of file
+} // namespace UE::PixelStreaming
diff --git a/Engine/Plugins/Media/PixelStreaming/Source/PixelStreaming/Private/UtilsRender.h b/Engine/Plugins/Media/PixelStreaming/Source/PixelStreaming/Private/UtilsRender.h
new file mode 100644
--- /dev/null
+++ b/Engine/Plugins/Media/PixelStreaming/Source/PixelStreaming/Private/UtilsRender.h
@@ -0,0 +1,196 @@
+// Copyright Epic Games, Inc. All Rights Reserved.
+
+#pragma once
+
+#include "RHIGPUReadback.h"
+#include "RenderGraphBuilder.h"
+#include "RenderGraphUtils.h"
+#include "Runtime/Renderer/Private/ScreenPass.h"
+#include "Modules/ModuleManager.h"
+#include "ScreenRendering.h"
+#include "Utils.h"
+
+namespace UE::PixelStreaming
+{
+	inline FTextureRHIRef CreateRHITexture(uint32 Width, uint32 Height)
+	{
+		// Create empty texture
+		FRHITextureCreateDesc TextureDesc =
+			FRHITextureCreateDesc::Create2D(TEXT("PixelStreamingBlankTexture"), Width, Height, EPixelFormat::PF_B8G8R8A8)
+			.SetClearValue(FClearValueBinding::None)
+			.SetFlags(ETextureCreateFlags::RenderTargetable)
+			.SetInitialState(ERHIAccess::Present)
+			.DetermineInititialState();
+
+		if (RHIGetInterfaceType() == ERHIInterfaceType::Vulkan)
+		{
+			TextureDesc.AddFlags(ETextureCreateFlags::External);
+		}
+		else
+		{
+			TextureDesc.AddFlags(ETextureCreateFlags::Shared);
+		}
+
+		return GDynamicRHI->RHICreateTexture(TextureDesc);
+	}
+
+	/*
+	 * Copy from one texture to another.
+	 * Assumes SourceTexture is in ERHIAccess::CopySrc and DestTexture is in ERHIAccess::CopyDest
+	 * Fence can be nullptr if no fence is to be used.
+	 */
+	inline void CopyTexture(FRHICommandList& RHICmdList, FTextureRHIRef SourceTexture, FTextureRHIRef DestTexture, FRHIGPUFence* Fence)
+	{
+		if (SourceTexture->GetDesc().Format == DestTexture->GetDesc().Format
+			&& SourceTexture->GetDesc().Extent.X == DestTexture->GetDesc().Extent.X
+			&& SourceTexture->GetDesc().Extent.Y == DestTexture->GetDesc().Extent.Y)
+		{
+
+			RHICmdList.Transition(FRHITransitionInfo(SourceTexture, ERHIAccess::Unknown, ERHIAccess::CopySrc));
+			RHICmdList.Transition(FRHITransitionInfo(DestTexture, ERHIAccess::Unknown, ERHIAccess::CopyDest));
+
+			// source and dest are the same. simple copy
+			RHICmdList.CopyTexture(SourceTexture, DestTexture, {});
+		}
+		else
+		{
+			IRendererModule* RendererModule = &FModuleManager::GetModuleChecked<IRendererModule>("Renderer");
+
+			RHICmdList.Transition(FRHITransitionInfo(SourceTexture, ERHIAccess::Unknown, ERHIAccess::SRVMask));
+			RHICmdList.Transition(FRHITransitionInfo(DestTexture, ERHIAccess::Unknown, ERHIAccess::RTV));
+
+			// source and destination are different. rendered copy
+			FRHIRenderPassInfo RPInfo(DestTexture, ERenderTargetActions::Load_Store);
+			RHICmdList.BeginRenderPass(RPInfo, TEXT("PixelStreaming::CopyTexture"));
+			{
+				FGlobalShaderMap* ShaderMap = GetGlobalShaderMap(GMaxRHIFeatureLevel);
+				TShaderMapRef<FScreenVS> VertexShader(ShaderMap);
+				TShaderMapRef<FScreenPS> PixelShader(ShaderMap);
+
+				RHICmdList.SetViewport(0, 0, 0.0f, DestTexture->GetDesc().Extent.X, DestTexture->GetDesc().Extent.Y, 1.0f);
+
+				FGraphicsPipelineStateInitializer GraphicsPSOInit;
+				RHICmdList.ApplyCachedRenderTargets(GraphicsPSOInit);
+				GraphicsPSOInit.BlendState = TStaticBlendState<>::GetRHI();
+				GraphicsPSOInit.RasterizerState = TStaticRasterizerState<>::GetRHI();
+				GraphicsPSOInit.DepthStencilState = TStaticDepthStencilState<false, CF_Always>::GetRHI();
+				GraphicsPSOInit.BoundShaderState.VertexDeclarationRHI = GFilterVertexDeclaration.VertexDeclarationRHI;
+				GraphicsPSOInit.BoundShaderState.VertexShaderRHI = VertexShader.GetVertexShader();
+				GraphicsPSOInit.BoundShaderState.PixelShaderRHI = PixelShader.GetPixelShader();
+				GraphicsPSOInit.PrimitiveType = PT_TriangleList;
+				SetGraphicsPipelineState(RHICmdList, GraphicsPSOInit, 0);
+
+				PixelShader->SetParameters(RHICmdList, TStaticSamplerState<SF_Point>::GetRHI(), SourceTexture);
+
+				FIntPoint TargetBufferSize(DestTexture->GetDesc().Extent.X, DestTexture->GetDesc().Extent.Y);
+				RendererModule->DrawRectangle(RHICmdList, 0, 0, // Dest X, Y
+					DestTexture->GetDesc().Extent.X,			// Dest Width
+					DestTexture->GetDesc().Extent.Y,			// Dest Height
+					0, 0,										// Source U, V
+					1, 1,										// Source USize, VSize
+					TargetBufferSize,							// Target buffer size
+					FIntPoint(1, 1),							// Source texture size
+					VertexShader, EDRF_Default);
+			}
+
+			RHICmdList.EndRenderPass();
+
+			RHICmdList.Transition(FRHITransitionInfo(SourceTexture, ERHIAccess::SRVMask, ERHIAccess::CopySrc));
+			RHICmdList.Transition(FRHITransitionInfo(DestTexture, ERHIAccess::RTV, ERHIAccess::CopyDest));
+		}
+
+		if (Fence != nullptr)
+		{
+			RHICmdList.WriteGPUFence(Fence);
+		}
+	}
+
+	inline void CopyTextureToReadbackTexture(FTextureRHIRef SourceTexture, TSharedPtr<FRHIGPUTextureReadback> GPUTextureReadback, void* OutBuffer)
+	{
+		FRHICommandListImmediate& RHICmdList = FRHICommandListExecutor::GetImmediateCommandList();
+		FRDGBuilder GraphBuilder(RHICmdList);
+
+		FRDGTextureRef RDGSourceTexture = RegisterExternalTexture(GraphBuilder, SourceTexture, TEXT("SourceCopyTextureToReadbackTexture"));
+		FRDGTextureRef RDGStagingTexture = nullptr;
+
+		if (RDGSourceTexture->Desc.Format != EPixelFormat::PF_B8G8R8A8)
+		{
+			// We need the pixel format to be BGRA8 so we first draw it to a staging texture
+			{
+				const FRDGTextureDesc Desc(FRDGTextureDesc::Create2D(RDGSourceTexture->Desc.Extent, PF_B8G8R8A8, FClearValueBinding::None, TexCreate_ShaderResource | TexCreate_RenderTargetable));
+				RDGStagingTexture = GraphBuilder.CreateTexture(Desc, TEXT("StagingCopyTextureToReadbackTexture"));
+			}
+
+			FGlobalShaderMap* ShaderMap = GetGlobalShaderMap(GMaxRHIFeatureLevel);
+
+			TShaderMapRef<FScreenPassVS> VertexShader(ShaderMap);
+
+			// Setup the pixel shader
+			TShaderMapRef<FCopyRectPS> PixelShader(ShaderMap);
+
+			FCopyRectPS::FParameters* PixelShaderParameters = GraphBuilder.AllocParameters<FCopyRectPS::FParameters>();
+			PixelShaderParameters->InputTexture = RDGSourceTexture;
+			PixelShaderParameters->InputSampler = TStaticSamplerState<SF_Bilinear, AM_Clamp, AM_Clamp, AM_Clamp>::GetRHI();
+			PixelShaderParameters->RenderTargets[0] = FRenderTargetBinding(RDGStagingTexture, ERenderTargetLoadAction::ELoad);
+
+			ClearUnusedGraphResources(PixelShader, PixelShaderParameters);
+
+			// We are not doing any clever blending stuff so we just use defaults here
+			FRHIBlendState* BlendState = FScreenPassPipelineState::FDefaultBlendState::GetRHI();
+			FRHIDepthStencilState* DepthStencilState = FScreenPassPipelineState::FDefaultDepthStencilState::GetRHI();
+
+			// Create the pipline state that will execute
+			const FScreenPassPipelineState PipelineState(VertexShader, PixelShader, BlendState, DepthStencilState);
+
+			// Add the pass the the graph builder
+			GraphBuilder.AddPass(
+				RDG_EVENT_NAME("PixelStreamingChangePixelFormat"),
+				PixelShaderParameters,
+				ERDGPassFlags::Raster,
+				[PipelineState, Extent = RDGSourceTexture->Desc.Extent, PixelShader, PixelShaderParameters](FRHICommandList& RHICmdList) {
+					PipelineState.Validate();
+
+					RHICmdList.SetViewport(0.0f, 0.0f, 0.0f, Extent.X, Extent.Y, 1.0f);
+
+					SetScreenPassPipelineState(RHICmdList, PipelineState);
+
+					SetShaderParameters(RHICmdList, PixelShader, PixelShader.GetPixelShader(), *PixelShaderParameters);
+
+					DrawRectangle(
+						RHICmdList,
+						0, 0, Extent.X, Extent.Y,
+						0, 0, Extent.X, Extent.Y,
+						Extent,
+						Extent,
+						PipelineState.VertexShader,
+						EDRF_UseTriangleOptimization);
+				});
+		}
+		else
+		{
+			// Otherwise if the PixelFormat is already BGRA8 we can just use the SourceTexture as the staging texture
+			RDGStagingTexture = RDGSourceTexture;
+		}
+
+		// Do the copy from staging RBGA8 texture into the readback
+		AddEnqueueCopyPass(GraphBuilder, GPUTextureReadback.Get(), RDGStagingTexture);
+
+		GraphBuilder.Execute();
+
+		// Lock and copy out the content of the TextureReadback to the CPU
+		int32 OutRowPitchInPixels;
+		int32 BlockSize = GPixelFormats[RDGStagingTexture->Desc.Format].BlockBytes;
+		void* ResultsBuffer = GPUTextureReadback->Lock(OutRowPitchInPixels);
+		if (RDGSourceTexture->Desc.Extent.X == OutRowPitchInPixels)
+		{
+			// Source pixel width is the same as the stride of the result buffer (ie no padding), we can do a plain memcpy
+			FPlatformMemory::Memcpy(OutBuffer, ResultsBuffer, (RDGSourceTexture->Desc.Extent.X * RDGSourceTexture->Desc.Extent.Y * BlockSize));
+		}
+		else
+		{
+			// Source pixel width differs from the stride of the result buffer (ie padding), do a memcpy that accounts for this
+			MemCpyStride(OutBuffer, ResultsBuffer, RDGSourceTexture->Desc.Extent.X * BlockSize, OutRowPitchInPixels * BlockSize, RDGSourceTexture->Desc.Extent.Y);
+		}
+		GPUTextureReadback->Unlock();
+	}
+} // namespace UE::PixelStreaming
diff --git a/Engine/Plugins/Media/PixelStreaming/Source/PixelStreaming/Private/VideoDecoderFactory.cpp b/Engine/Plugins/Media/PixelStreaming/Source/PixelStreaming/Private/VideoDecoderFactory.cpp
--- a/Engine/Plugins/Media/PixelStreaming/Source/PixelStreaming/Private/VideoDecoderFactory.cpp
+++ b/Engine/Plugins/Media/PixelStreaming/Source/PixelStreaming/Private/VideoDecoderFactory.cpp
@@ -2,16 +2,35 @@
 #include "VideoDecoderFactory.h"
 #include "Utils.h"
 #include "VideoDecoderStub.h"
+#include "absl/strings/match.h"
+#include "VideoDecoderVPX.h"
 
 namespace UE::PixelStreaming
 {
 	std::vector<webrtc::SdpVideoFormat> FVideoDecoderFactory::GetSupportedFormats() const
 	{
-		return { CreateH264Format(webrtc::H264::kProfileConstrainedBaseline, webrtc::H264::kLevel3_1) };
+		std::vector<webrtc::SdpVideoFormat> video_formats;
+		video_formats.push_back(webrtc::SdpVideoFormat(cricket::kVp8CodecName));
+		// video_formats.push_back(webrtc::SdpVideoFormat(cricket::kVp9CodecName));
+#if WEBRTC_VERSION == 84
+		video_formats.push_back(CreateH264Format(webrtc::H264::kProfileConstrainedBaseline, webrtc::H264::kLevel3_1));
+#elif WEBRTC_VERSION == 96
+		video_formats.push_back(CreateH264Format(webrtc::H264Profile::kProfileConstrainedBaseline, webrtc::H264Level::kLevel3_1));
+#endif
+		return video_formats;
 	}
 
 	std::unique_ptr<webrtc::VideoDecoder> FVideoDecoderFactory::CreateVideoDecoder(const webrtc::SdpVideoFormat& format)
 	{
+		if (absl::EqualsIgnoreCase(format.name, cricket::kVp8CodecName))
+		{
+			return std::make_unique<VideoDecoderVPX>();
+			// return webrtc::VP8Decoder::Create();
+		}
+		else if (absl::EqualsIgnoreCase(format.name, cricket::kVp9CodecName))
+		{
+			return webrtc::VP9Decoder::Create();
+		}
 		return std::make_unique<FVideoDecoderStub>();
 	}
 } // namespace UE::PixelStreaming
diff --git a/Engine/Plugins/Media/PixelStreaming/Source/PixelStreaming/Private/VideoDecoderStub.h b/Engine/Plugins/Media/PixelStreaming/Source/PixelStreaming/Private/VideoDecoderStub.h
--- a/Engine/Plugins/Media/PixelStreaming/Source/PixelStreaming/Private/VideoDecoderStub.h
+++ b/Engine/Plugins/Media/PixelStreaming/Source/PixelStreaming/Private/VideoDecoderStub.h
@@ -8,9 +8,14 @@ namespace UE::PixelStreaming
 	class FVideoDecoderStub : public webrtc::VideoDecoder
 	{
 	public:
+#if WEBRTC_VERSION == 84
 		virtual int32_t InitDecode(const webrtc::VideoCodec* codec_settings, int32_t number_of_cores) override { return 0; }
+#elif WEBRTC_VERSION == 96
+		virtual bool Configure(const Settings& settings) override { return false; }
+#endif
 		virtual int32_t Decode(const webrtc::EncodedImage& input_image, bool missing_frames, int64_t render_time_ms) override { return 0; }
 		virtual int32_t RegisterDecodeCompleteCallback(webrtc::DecodedImageCallback* callback) override { return 0; }
 		virtual int32_t Release() override { return 0; }
+
 	};
 } // namespace UE::PixelStreaming
diff --git a/Engine/Plugins/Media/PixelStreaming/Source/PixelStreaming/Private/VideoDecoderVPX.cpp b/Engine/Plugins/Media/PixelStreaming/Source/PixelStreaming/Private/VideoDecoderVPX.cpp
new file mode 100644
--- /dev/null
+++ b/Engine/Plugins/Media/PixelStreaming/Source/PixelStreaming/Private/VideoDecoderVPX.cpp
@@ -0,0 +1,31 @@
+// Copyright Epic Games, Inc. All Rights Reserved.
+
+#include "VideoDecoderVPX.h"
+
+namespace UE::PixelStreaming
+{
+	VideoDecoderVPX::VideoDecoderVPX()
+	{
+		VideoDecoder = webrtc::VP8Decoder::Create();
+	}
+
+	int32 VideoDecoderVPX::InitDecode(const webrtc::VideoCodec* codec_settings, int32 number_of_cores)
+	{
+		return VideoDecoder->InitDecode(codec_settings, number_of_cores);
+	}
+
+	int32 VideoDecoderVPX::Decode(const webrtc::EncodedImage& input_image, bool missing_frames, int64_t render_time_ms)
+	{
+		return VideoDecoder->Decode(input_image, missing_frames, render_time_ms);
+	}
+
+	int32 VideoDecoderVPX::RegisterDecodeCompleteCallback(webrtc::DecodedImageCallback* callback)
+	{
+		return VideoDecoder->RegisterDecodeCompleteCallback(callback);
+	}
+
+	int32 VideoDecoderVPX::Release()
+	{
+		return VideoDecoder->Release();
+	}
+} // namespace UE::PixelStreaming
diff --git a/Engine/Plugins/Media/PixelStreaming/Source/PixelStreaming/Private/VideoDecoderVPX.h b/Engine/Plugins/Media/PixelStreaming/Source/PixelStreaming/Private/VideoDecoderVPX.h
new file mode 100644
--- /dev/null
+++ b/Engine/Plugins/Media/PixelStreaming/Source/PixelStreaming/Private/VideoDecoderVPX.h
@@ -0,0 +1,24 @@
+// Copyright Epic Games, Inc. All Rights Reserved.
+
+#pragma once
+
+#include "WebRTCIncludes.h"
+
+namespace UE::PixelStreaming
+{
+	class VideoDecoderVPX : public webrtc::VideoDecoder
+	{
+	public:
+		VideoDecoderVPX();
+		virtual ~VideoDecoderVPX() = default;
+
+		virtual int32 InitDecode(const webrtc::VideoCodec* codec_settings, int32 number_of_cores) override;
+		virtual int32 Decode(const webrtc::EncodedImage& input_image, bool missing_frames, int64_t render_time_ms) override;
+		virtual int32 RegisterDecodeCompleteCallback(webrtc::DecodedImageCallback* callback) override;
+		virtual int32 Release() override;
+		virtual const char* ImplementationName() const { return "VideoDecoderVPX"; }
+
+	private:
+		std::unique_ptr<webrtc::VideoDecoder> VideoDecoder;
+	};
+} // namespace UE::PixelStreaming
diff --git a/Engine/Plugins/Media/PixelStreaming/Source/PixelStreaming/Private/VideoEncoderAdapterSimulcast.cpp b/Engine/Plugins/Media/PixelStreaming/Source/PixelStreaming/Private/VideoEncoderAdapterSimulcast.cpp
new file mode 100644
--- /dev/null
+++ b/Engine/Plugins/Media/PixelStreaming/Source/PixelStreaming/Private/VideoEncoderAdapterSimulcast.cpp
@@ -0,0 +1,588 @@
+// Copyright Epic Games, Inc. All Rights Reserved.
+#include "VideoEncoderAdapterSimulcast.h"
+#include "FrameBufferH264.h"
+#include "FrameBufferI420.h"
+#include "modules/video_coding/codecs/vp8/include/vp8.h"
+#include "modules/video_coding/codecs/vp9/include/vp9.h"
+#include "VideoEncoderFactorySimple.h"
+
+namespace
+{
+	uint32_t SumStreamMaxBitrate(int Streams, const webrtc::VideoCodec& Codec)
+	{
+		uint32_t BitrateSum = 0;
+		for (int i = 0; i < Streams; ++i)
+		{
+			BitrateSum += Codec.simulcastStream[i].maxBitrate;
+		}
+		return BitrateSum;
+	}
+
+	int GetNumberOfStreams(const webrtc::VideoCodec& Codec)
+	{
+		int Streams = Codec.numberOfSimulcastStreams < 1 ? 1 : Codec.numberOfSimulcastStreams;
+		uint32_t SimulcastMaxBitrate = SumStreamMaxBitrate(Streams, Codec);
+		if (SimulcastMaxBitrate == 0)
+		{
+			Streams = 1;
+		}
+		return Streams;
+	}
+
+	int GetNumActiveStreams(const webrtc::VideoCodec& Codec)
+	{
+		int NumConfiguredStreams = GetNumberOfStreams(Codec);
+		int NumActiveStreams = 0;
+		for (int i = 0; i < NumConfiguredStreams; ++i)
+		{
+			if (Codec.simulcastStream[i].active)
+			{
+				++NumActiveStreams;
+			}
+		}
+		return NumActiveStreams;
+	}
+
+	bool StreamResolutionCompare(const webrtc::SpatialLayer& A, const webrtc::SpatialLayer& B)
+	{
+		return std::tie(A.height, A.width, A.maxBitrate, A.maxFramerate) < std::tie(B.height, B.width, B.maxBitrate, B.maxFramerate);
+	}
+
+	void PopulateStreamCodec(const webrtc::VideoCodec& CodecSettings, int StreamIndex, uint32_t StartBitrateKbps, webrtc::VideoCodec* StreamCodec)
+	{
+		*StreamCodec = CodecSettings;
+
+		// Stream specific settings.
+		StreamCodec->numberOfSimulcastStreams = 0;
+		StreamCodec->width = CodecSettings.simulcastStream[StreamIndex].width;
+		StreamCodec->height = CodecSettings.simulcastStream[StreamIndex].height;
+		StreamCodec->maxBitrate = CodecSettings.simulcastStream[StreamIndex].maxBitrate;
+		StreamCodec->minBitrate = CodecSettings.simulcastStream[StreamIndex].minBitrate;
+		StreamCodec->maxFramerate = CodecSettings.simulcastStream[StreamIndex].maxFramerate;
+		StreamCodec->qpMax = CodecSettings.simulcastStream[StreamIndex].qpMax;
+		StreamCodec->active = CodecSettings.simulcastStream[StreamIndex].active;
+		if (CodecSettings.codecType == webrtc::kVideoCodecH264)
+		{
+			StreamCodec->H264()->numberOfTemporalLayers = CodecSettings.simulcastStream[StreamIndex].numberOfTemporalLayers;
+		}
+		StreamCodec->startBitrate = StartBitrateKbps;
+	}
+
+	// An EncodedImageCallback implementation that forwards on calls to a
+	// UE::PixelStreaming::FVideoEncoderAdapterSimulcast, but with the stream index it's registered with as
+	// the first parameter to Encoded.
+	class AdapterEncodedImageCallback : public webrtc::EncodedImageCallback
+	{
+	public:
+		AdapterEncodedImageCallback(UE::PixelStreaming::FVideoEncoderAdapterSimulcast* adapter,
+			size_t stream_idx)
+			: adapter_(adapter), stream_idx_(stream_idx)
+		{
+		}
+
+		EncodedImageCallback::Result OnEncodedImage(
+			const webrtc::EncodedImage& encoded_image,
+			const webrtc::CodecSpecificInfo* codec_specific_info
+#if WEBRTC_VERSION == 84
+			, const webrtc::RTPFragmentationHeader* fragmentation
+#endif
+		) override
+		{
+			return adapter_->OnEncodedImage(stream_idx_, encoded_image, codec_specific_info
+#if WEBRTC_VERSION == 84
+				, fragmentation
+#endif
+			);
+		}
+
+	private:
+		UE::PixelStreaming::FVideoEncoderAdapterSimulcast* const adapter_;
+		const size_t stream_idx_;
+	};
+} // namespace
+
+namespace UE::PixelStreaming
+{
+	FVideoEncoderAdapterSimulcast::FVideoEncoderAdapterSimulcast(FVideoEncoderFactorySimulcast& InSimulcastFactory, const webrtc::SdpVideoFormat& format)
+		: Initialized(false)
+		, SimulcastEncoderFactory(InSimulcastFactory)
+		, VideoFormat(format)
+		, EncodedCompleteCallback(nullptr)
+	{
+		memset(&CurrentCodec, 0, sizeof(webrtc::VideoCodec));
+	}
+
+	FVideoEncoderAdapterSimulcast::~FVideoEncoderAdapterSimulcast()
+	{
+		check(!IsInitialized());
+	}
+
+	int FVideoEncoderAdapterSimulcast::Release()
+	{
+		{
+			// Lock during deleting an encoder
+			FScopeLock Lock(&StreamInfosGuard);
+			while (!StreamInfos.empty())
+			{
+				std::unique_ptr<VideoEncoder> Encoder = std::move(StreamInfos.back().Encoder);
+				// Even though it seems very unlikely, there are no guarantees that the
+				// encoder will not call back after being Release()'d. Therefore, we first
+				// disable the callbacks here.
+				Encoder->RegisterEncodeCompleteCallback(nullptr);
+				Encoder->Release();
+				StreamInfos.pop_back(); // Deletes callback adapter.
+			}
+		}
+
+		Initialized = false;
+
+		return WEBRTC_VIDEO_CODEC_OK;
+	}
+
+	int FVideoEncoderAdapterSimulcast::InitEncode(const webrtc::VideoCodec* codec_settings, const webrtc::VideoEncoder::Settings& settings)
+	{
+		int NumberOfStreams = GetNumberOfStreams(*codec_settings);
+		int NumActiveStreams = GetNumActiveStreams(*codec_settings);
+
+		CurrentCodec = *codec_settings;
+
+		const auto MinMax = std::minmax_element(std::begin(CurrentCodec.simulcastStream), std::begin(CurrentCodec.simulcastStream) + NumberOfStreams, StreamResolutionCompare);
+		LowestResolutionStreamIndex = std::distance(std::begin(CurrentCodec.simulcastStream), MinMax.first);
+		HighestResolutionStreamIndex = std::distance(std::begin(CurrentCodec.simulcastStream), MinMax.second);
+
+		const webrtc::SdpVideoFormat Format(CurrentCodec.codecType == webrtc::kVideoCodecVP8 ? "VP8" : "H264", VideoFormat.parameters);
+
+		if (NumActiveStreams == 1)
+		{
+			// with one stream we just proxy the pixelstreaming encoder
+			FVideoEncoderFactorySimple* EncoderFactory = SimulcastEncoderFactory.GetOrCreateEncoderFactory(0);
+			std::unique_ptr<VideoEncoder> Encoder = EncoderFactory->CreateVideoEncoder(Format);
+
+			int ReturnCode = Encoder->InitEncode(&CurrentCodec, settings);
+			if (ReturnCode < 0)
+			{
+				// Explicitly destroy the current encoder; because we haven't registered
+				// a StreamInfo for it yet, Release won't do anything about it.
+				Encoder.reset();
+				Release();
+				return ReturnCode;
+			}
+
+			std::unique_ptr<webrtc::EncodedImageCallback> Callback(new AdapterEncodedImageCallback(this, 0));
+			Encoder->RegisterEncodeCompleteCallback(Callback.get());
+			StreamInfos.emplace_back(
+				std::move(Encoder), std::move(Callback),
+				std::make_unique<webrtc::FramerateController>(CurrentCodec.maxFramerate),
+				CurrentCodec.width, CurrentCodec.height, true, true);
+		}
+		else
+		{
+			for (int i = 0; i < NumberOfStreams; ++i)
+			{
+				webrtc::VideoCodec StreamCodec;
+				uint32_t StartBitrateKbps = CurrentCodec.simulcastStream[i].targetBitrate;
+				PopulateStreamCodec(CurrentCodec, i, StartBitrateKbps, &StreamCodec);
+
+				FVideoEncoderFactorySimple* EncoderFactory = SimulcastEncoderFactory.GetOrCreateEncoderFactory(i);
+				std::unique_ptr<VideoEncoder> Encoder = EncoderFactory->CreateVideoEncoder(Format);
+
+				int ReturnCode = Encoder->InitEncode(&StreamCodec, settings);
+				if (ReturnCode < 0)
+				{
+					// Explicitly destroy the current encoder; because we haven't registered
+					// a StreamInfo for it yet, Release won't do anything about it.
+					Encoder.reset();
+					Release();
+					return ReturnCode;
+				}
+
+				std::unique_ptr<webrtc::EncodedImageCallback> Callback(new AdapterEncodedImageCallback(this, i));
+				Encoder->RegisterEncodeCompleteCallback(Callback.get());
+				StreamInfos.emplace_back(
+					std::move(Encoder), std::move(Callback),
+					std::make_unique<webrtc::FramerateController>(StreamCodec.maxFramerate),
+					StreamCodec.width, StreamCodec.height, true, true);
+			}
+		}
+
+		Initialized = true;
+
+		return WEBRTC_VIDEO_CODEC_OK;
+	}
+
+	int FVideoEncoderAdapterSimulcast::EncodeVP8(const webrtc::VideoFrame& InputImage, const std::vector<webrtc::VideoFrameType>* FrameTypes, bool bSendKeyFrame)
+	{
+		// TODO Currently this does not get his because the encoder factory will just return a vpx encoder but if we change it to use this encoder
+		// webrtc will try to call ToI420() on the simulcast buffer which is currently unimplemented because the buffers are pre created.
+		// we either need to short circuit that behaviour or change what webrtc is trying to do. info.supports_native_handle from the vpx encoders
+		// is the culprit.
+
+		const FFrameBufferI420Simulcast* FrameBuffer = static_cast<FFrameBufferI420Simulcast*>(InputImage.video_frame_buffer().get());
+		check(FrameBuffer->GetFrameBufferType() == EPixelStreamingFrameBufferType::Simulcast);
+
+		rtc::scoped_refptr<webrtc::I420BufferInterface> SourceBuffer;
+		for (size_t StreamIdx = 0; StreamIdx < StreamInfos.size(); ++StreamIdx)
+		{
+			// Don't encode frames in resolutions that we don't intend to send.
+			if (!StreamInfos[StreamIdx].bSendStream)
+			{
+				continue;
+			}
+
+			// extract the specific layer frame buffer
+			webrtc::VideoFrame NewFrame(InputImage);
+			const int LayerIndex = StreamInfos.size() == 1 ? FrameBuffer->GetNumLayers() - 1 : StreamIdx;
+			rtc::scoped_refptr<FFrameBufferI420> LayerFrameBuffer = new rtc::RefCountedObject<FFrameBufferI420>(FrameBuffer->GetFrameAdapter(), LayerIndex);
+			NewFrame.set_video_frame_buffer(LayerFrameBuffer);
+
+			const uint32_t FrameTimestampMs = 1000 * InputImage.timestamp() / 90000; // kVideoPayloadTypeFrequency;
+			// If adapter is passed through and only one sw encoder does simulcast,
+			// frame types for all streams should be passed to the encoder unchanged.
+			// Otherwise a single per-encoder frame type is passed.
+			std::vector<webrtc::VideoFrameType> StreamFrameTypes(StreamInfos.size() == 1 ? GetNumberOfStreams(CurrentCodec) : 1);
+			if (bSendKeyFrame)
+			{
+				std::fill(StreamFrameTypes.begin(), StreamFrameTypes.end(), webrtc::VideoFrameType::kVideoFrameKey);
+				StreamInfos[StreamIdx].KeyFrameRequest = false;
+			}
+			else
+			{
+#if WEBRTC_VERSION == 84
+				if (StreamInfos[StreamIdx].FramerateController->DropFrame(FrameTimestampMs))
+#elif WEBRTC_VERSION == 96
+				if (StreamInfos[StreamIdx].FramerateController->ShouldDropFrame(FrameTimestampMs))
+#endif
+				{
+					continue;
+				}
+				std::fill(StreamFrameTypes.begin(), StreamFrameTypes.end(), webrtc::VideoFrameType::kVideoFrameDelta);
+			}
+#if WEBRTC_VERSION == 84
+			StreamInfos[StreamIdx].FramerateController->AddFrame(FrameTimestampMs);
+#elif WEBRTC_VERSION == 96
+			StreamInfos[StreamIdx].FramerateController->KeepFrame(FrameTimestampMs);
+#endif
+
+			int RtcError = StreamInfos[StreamIdx].Encoder->Encode(NewFrame, &StreamFrameTypes);
+			if (RtcError != WEBRTC_VIDEO_CODEC_OK)
+			{
+				return RtcError;
+			}
+		}
+		return WEBRTC_VIDEO_CODEC_OK;
+	}
+
+	int FVideoEncoderAdapterSimulcast::EncodeVP9(const webrtc::VideoFrame& InputImage, const std::vector<webrtc::VideoFrameType>* FrameTypes, bool bSendKeyFrame)
+	{
+		return WEBRTC_VIDEO_CODEC_OK;
+	}
+
+	int FVideoEncoderAdapterSimulcast::EncodeH264(const webrtc::VideoFrame& InputImage, const std::vector<webrtc::VideoFrameType>* FrameTypes, bool bSendKeyFrame)
+	{
+		// ignore any init frames
+		// NOTE: It seems most of the time they never even get here because the internals of WebRTC decide to drop it. It seems it thinks
+		// the encoder is paused or something.
+		const IPixelStreamingFrameBuffer* FrameBuffer = static_cast<IPixelStreamingFrameBuffer*>(InputImage.video_frame_buffer().get());
+		if (FrameBuffer->GetFrameBufferType() == EPixelStreamingFrameBufferType::Initialize)
+		{
+			return WEBRTC_VIDEO_CODEC_OK;
+		}
+		// separate out our pixelstreaming image sources
+		check(FrameBuffer->GetFrameBufferType() == EPixelStreamingFrameBufferType::Simulcast);
+		const FFrameBufferH264Simulcast* SimulcastFrameBuffer = static_cast<const FFrameBufferH264Simulcast*>(FrameBuffer);
+		for (size_t StreamIdx = 0; StreamIdx < StreamInfos.size(); ++StreamIdx)
+		{
+			// Don't encode frames in resolutions that we don't intend to send.
+			if (!StreamInfos[StreamIdx].bSendStream)
+			{
+				continue;
+			}
+			webrtc::VideoFrame NewFrame(InputImage);
+			// grab the simulcast frame source, extract the frame source for this layer and wrap that in a new frame buffer appropriate for H264
+			const int LayerIndex = StreamInfos.size() == 1 ? SimulcastFrameBuffer->GetNumLayers() - 1 : StreamIdx;
+			rtc::scoped_refptr<FFrameBufferH264> LayerFrameBuffer = new rtc::RefCountedObject<FFrameBufferH264>(SimulcastFrameBuffer->GetFrameSource(), LayerIndex);
+			NewFrame.set_video_frame_buffer(LayerFrameBuffer);
+
+			const uint32_t FrameTimestampMs = 1000 * NewFrame.timestamp() / 90000;
+			// If adapter is passed through and only one sw encoder does simulcast,
+			// frame types for all streams should be passed to the encoder unchanged.
+			// Otherwise a single per-encoder frame type is passed.
+			std::vector<webrtc::VideoFrameType> StreamFrameTypes(StreamInfos.size() == 1 ? GetNumberOfStreams(CurrentCodec) : 1);
+			if (bSendKeyFrame)
+			{
+				std::fill(StreamFrameTypes.begin(), StreamFrameTypes.end(), webrtc::VideoFrameType::kVideoFrameKey);
+				StreamInfos[StreamIdx].KeyFrameRequest = false;
+			}
+			else
+			{
+#if WEBRTC_VERSION == 84
+				if (StreamInfos[StreamIdx].FramerateController->DropFrame(FrameTimestampMs))
+#elif WEBRTC_VERSION == 96
+				if (StreamInfos[StreamIdx].FramerateController->ShouldDropFrame(FrameTimestampMs))
+#endif
+				{
+					return WEBRTC_VIDEO_CODEC_OK;
+				}
+				std::fill(StreamFrameTypes.begin(), StreamFrameTypes.end(), webrtc::VideoFrameType::kVideoFrameDelta);
+			}
+#if WEBRTC_VERSION == 84
+			StreamInfos[StreamIdx].FramerateController->AddFrame(FrameTimestampMs);
+#elif WEBRTC_VERSION == 96
+			StreamInfos[StreamIdx].FramerateController->KeepFrame(FrameTimestampMs);
+#endif
+			int RtcError = StreamInfos[StreamIdx].Encoder->Encode(NewFrame, &StreamFrameTypes);
+			if (RtcError != WEBRTC_VIDEO_CODEC_OK)
+			{
+				return RtcError;
+			}
+		}
+		return WEBRTC_VIDEO_CODEC_OK;
+	}
+
+	int FVideoEncoderAdapterSimulcast::Encode(const webrtc::VideoFrame& input_image, const std::vector<webrtc::VideoFrameType>* frame_types)
+	{
+		if (!IsInitialized())
+		{
+			return WEBRTC_VIDEO_CODEC_UNINITIALIZED;
+		}
+		if (EncodedCompleteCallback == nullptr)
+		{
+			return WEBRTC_VIDEO_CODEC_UNINITIALIZED;
+		}
+
+		// All active streams should generate a key frame if
+		// a key frame is requested by any stream.
+		bool bSendKeyFrame = false;
+		if (frame_types)
+		{
+			for (size_t i = 0; i < frame_types->size(); ++i)
+			{
+				if (frame_types->at(i) == webrtc::VideoFrameType::kVideoFrameKey)
+				{
+					bSendKeyFrame = true;
+					break;
+				}
+			}
+		}
+
+		for (size_t StreamIdx = 0; StreamIdx < StreamInfos.size(); ++StreamIdx)
+		{
+			if (StreamInfos[StreamIdx].KeyFrameRequest && StreamInfos[StreamIdx].bSendStream)
+			{
+				bSendKeyFrame = true;
+				break;
+			}
+		}
+
+		switch (CurrentCodec.codecType)
+		{
+			case webrtc::kVideoCodecVP8:
+				return EncodeVP8(input_image, frame_types, bSendKeyFrame);
+			case webrtc::kVideoCodecVP9:
+				return EncodeVP9(input_image, frame_types, bSendKeyFrame);
+			case webrtc::kVideoCodecH264:
+				return EncodeH264(input_image, frame_types, bSendKeyFrame);
+			default:
+				return WEBRTC_VIDEO_CODEC_ERROR;
+		}
+	}
+
+	int FVideoEncoderAdapterSimulcast::RegisterEncodeCompleteCallback(webrtc::EncodedImageCallback* callback)
+	{
+		EncodedCompleteCallback = callback;
+		return WEBRTC_VIDEO_CODEC_OK;
+	}
+
+	void FVideoEncoderAdapterSimulcast::SetRates(const RateControlParameters& parameters)
+	{
+		if (!IsInitialized())
+		{
+			RTC_LOG(LS_WARNING) << "SetRates while not initialized";
+			return;
+		}
+
+		if (parameters.framerate_fps < 1.0)
+		{
+			RTC_LOG(LS_WARNING) << "Invalid framerate: " << parameters.framerate_fps;
+			return;
+		}
+
+		CurrentCodec.maxFramerate = static_cast<uint32_t>(parameters.framerate_fps + 0.5);
+
+		if (StreamInfos.size() == 1)
+		{
+			// Not doing simulcast.
+			StreamInfos[0].Encoder->SetRates(parameters);
+			return;
+		}
+
+		for (size_t StreamIdx = 0; StreamIdx < StreamInfos.size(); ++StreamIdx)
+		{
+			uint32_t StreamBitrateKbps = parameters.bitrate.GetSpatialLayerSum(StreamIdx) / 1000;
+
+			// Need a key frame if we have not sent this stream before.
+			if (StreamBitrateKbps > 0 && !StreamInfos[StreamIdx].bSendStream)
+			{
+				StreamInfos[StreamIdx].KeyFrameRequest = true;
+			}
+			StreamInfos[StreamIdx].bSendStream = StreamBitrateKbps > 0;
+
+			// Slice the temporal layers out of the full allocation and pass it on to
+			// the encoder handling the current simulcast stream.
+			RateControlParameters StreamParameters = parameters;
+			StreamParameters.bitrate = webrtc::VideoBitrateAllocation();
+			for (int i = 0; i < webrtc::kMaxTemporalStreams; ++i)
+			{
+				if (parameters.bitrate.HasBitrate(StreamIdx, i))
+				{
+					StreamParameters.bitrate.SetBitrate(0, i, parameters.bitrate.GetBitrate(StreamIdx, i));
+				}
+			}
+
+			// Assign link allocation proportionally to spatial layer allocation.
+			if (!parameters.bandwidth_allocation.IsZero() && parameters.bitrate.get_sum_bps() > 0)
+			{
+				StreamParameters.bandwidth_allocation =
+					webrtc::DataRate::BitsPerSec((parameters.bandwidth_allocation.bps() * StreamParameters.bitrate.get_sum_bps()) / parameters.bitrate.get_sum_bps());
+				// Make sure we don't allocate bandwidth lower than target bitrate.
+				if (StreamParameters.bandwidth_allocation.bps() < StreamParameters.bitrate.get_sum_bps())
+				{
+					StreamParameters.bandwidth_allocation =
+						webrtc::DataRate::BitsPerSec(StreamParameters.bitrate.get_sum_bps());
+				}
+			}
+#if WEBRTC_VERSION == 84
+			StreamParameters.framerate_fps = std::min<double>(parameters.framerate_fps, StreamInfos[StreamIdx].FramerateController->GetTargetRate());
+#elif WEBRTC_VERSION == 96
+			StreamParameters.framerate_fps = std::min<double>(parameters.framerate_fps, StreamInfos[StreamIdx].FramerateController->GetMaxFramerate());
+#endif
+
+			StreamInfos[StreamIdx].Encoder->SetRates(StreamParameters);
+		}
+	}
+
+	void FVideoEncoderAdapterSimulcast::OnPacketLossRateUpdate(float packet_loss_rate)
+	{
+		for (StreamInfo& Info : StreamInfos)
+		{
+			Info.Encoder->OnPacketLossRateUpdate(packet_loss_rate);
+		}
+	}
+
+	void FVideoEncoderAdapterSimulcast::OnRttUpdate(int64_t rtt_ms)
+	{
+		for (StreamInfo& Info : StreamInfos)
+		{
+			Info.Encoder->OnRttUpdate(rtt_ms);
+		}
+	}
+
+	void FVideoEncoderAdapterSimulcast::OnLossNotification(const LossNotification& loss_notification)
+	{
+		for (StreamInfo& Info : StreamInfos)
+		{
+			Info.Encoder->OnLossNotification(loss_notification);
+		}
+	}
+
+	webrtc::EncodedImageCallback::Result FVideoEncoderAdapterSimulcast::OnEncodedImage(
+		size_t stream_idx,
+		const webrtc::EncodedImage& encodedImage,
+		const webrtc::CodecSpecificInfo* codecSpecificInfo
+#if WEBRTC_VERSION == 84
+		, const webrtc::RTPFragmentationHeader* fragmentation
+#endif
+	)
+	{
+		webrtc::EncodedImage StreamImage(encodedImage);
+		webrtc::CodecSpecificInfo StreamCodecSpecific = *codecSpecificInfo;
+
+		StreamImage.SetSpatialIndex(stream_idx);
+
+		return EncodedCompleteCallback->OnEncodedImage(StreamImage, &StreamCodecSpecific
+#if WEBRTC_VERSION == 84
+			, fragmentation
+#endif
+		);
+	}
+
+	bool FVideoEncoderAdapterSimulcast::IsInitialized() const
+	{
+		return Initialized;
+	}
+
+	webrtc::VideoEncoder::EncoderInfo FVideoEncoderAdapterSimulcast::GetEncoderInfo() const
+	{
+		if (StreamInfos.size() == 1)
+		{
+			// Not using simulcast adapting functionality, just pass through.
+			return StreamInfos[0].Encoder->GetEncoderInfo();
+		}
+
+		VideoEncoder::EncoderInfo EncoderInfo;
+		EncoderInfo.implementation_name = "PixelStreamingSimulcastEncoderAdapter";
+		EncoderInfo.requested_resolution_alignment = 1;
+		EncoderInfo.supports_native_handle = true;
+		EncoderInfo.scaling_settings.thresholds = absl::nullopt;
+		if (StreamInfos.empty())
+		{
+			return EncoderInfo;
+		}
+
+		EncoderInfo.scaling_settings = VideoEncoder::ScalingSettings::kOff;
+		int NumActiveStreams = GetNumActiveStreams(CurrentCodec);
+
+		for (size_t i = 0; i < StreamInfos.size(); ++i)
+		{
+			VideoEncoder::EncoderInfo EncoderImplInfo = StreamInfos[i].Encoder->GetEncoderInfo();
+
+			if (i == 0)
+			{
+				// Encoder name indicates names of all sub-encoders.
+				EncoderInfo.implementation_name += " (";
+				EncoderInfo.implementation_name += EncoderImplInfo.implementation_name;
+
+				EncoderInfo.supports_native_handle = EncoderImplInfo.supports_native_handle;
+				EncoderInfo.has_trusted_rate_controller = EncoderImplInfo.has_trusted_rate_controller;
+				EncoderInfo.is_hardware_accelerated = EncoderImplInfo.is_hardware_accelerated;
+				EncoderInfo.has_internal_source = EncoderImplInfo.has_internal_source;
+			}
+			else
+			{
+				EncoderInfo.implementation_name += ", ";
+				EncoderInfo.implementation_name += EncoderImplInfo.implementation_name;
+
+				// Native handle supported if any encoder supports it.
+				EncoderInfo.supports_native_handle |= EncoderImplInfo.supports_native_handle;
+
+				// Trusted rate controller only if all encoders have it.
+				EncoderInfo.has_trusted_rate_controller &= EncoderImplInfo.has_trusted_rate_controller;
+
+				// Uses hardware support if any of the encoders uses it.
+				// For example, if we are having issues with down-scaling due to
+				// pipelining delay in HW encoders we need higher encoder usage
+				// thresholds in CPU adaptation.
+				EncoderInfo.is_hardware_accelerated |= EncoderImplInfo.is_hardware_accelerated;
+
+				// Has internal source only if all encoders have it.
+				EncoderInfo.has_internal_source &= EncoderImplInfo.has_internal_source;
+			}
+
+			// Nasty hack to allow us to manually convert VPX frames to I420 later in the encode block
+			if (CurrentCodec.codecType == webrtc::kVideoCodecVP8)
+			{
+				EncoderInfo.supports_native_handle = true;
+			}
+
+			EncoderInfo.fps_allocation[i] = EncoderImplInfo.fps_allocation[0];
+			EncoderInfo.requested_resolution_alignment = cricket::LeastCommonMultiple(EncoderInfo.requested_resolution_alignment, EncoderImplInfo.requested_resolution_alignment);
+			if (NumActiveStreams == 1 && CurrentCodec.simulcastStream[i].active)
+			{
+				EncoderInfo.scaling_settings = EncoderImplInfo.scaling_settings;
+			}
+		}
+		EncoderInfo.implementation_name += ")";
+
+		return EncoderInfo;
+	}
+} // namespace UE::PixelStreaming
diff --git a/Engine/Plugins/Media/PixelStreaming/Source/PixelStreaming/Private/VideoEncoderAdapterSimulcast.h b/Engine/Plugins/Media/PixelStreaming/Source/PixelStreaming/Private/VideoEncoderAdapterSimulcast.h
new file mode 100644
--- /dev/null
+++ b/Engine/Plugins/Media/PixelStreaming/Source/PixelStreaming/Private/VideoEncoderAdapterSimulcast.h
@@ -0,0 +1,94 @@
+// Copyright Epic Games, Inc. All Rights Reserved.
+#pragma once
+
+#include "CoreMinimal.h"
+#include "WebRTCIncludes.h"
+#include "VideoEncoderFactorySimulcast.h"
+
+namespace webrtc
+{
+	class SimulcastRateAllocator;
+	class VideoEncoderFactory;
+} // namespace webrtc
+
+namespace UE::PixelStreaming
+{
+	// This is highly modified version of webrtc::simulcast_encoder_adapter
+	class FVideoEncoderAdapterSimulcast : public webrtc::VideoEncoder
+	{
+	public:
+		FVideoEncoderAdapterSimulcast(FVideoEncoderFactorySimulcast& InSimulcastFactory, const webrtc::SdpVideoFormat& format);
+		~FVideoEncoderAdapterSimulcast() override;
+
+		// Implements VideoEncoder.
+		int Release() override;
+		int InitEncode(const webrtc::VideoCodec* codec_settings, const webrtc::VideoEncoder::Settings& settings) override;
+		int Encode(const webrtc::VideoFrame& input_image, const std::vector<webrtc::VideoFrameType>* frame_types) override;
+		int RegisterEncodeCompleteCallback(webrtc::EncodedImageCallback* callback) override;
+
+		void SetRates(const RateControlParameters& parameters) override;
+		void OnPacketLossRateUpdate(float packet_loss_rate) override;
+		void OnRttUpdate(int64_t rtt_ms) override;
+		void OnLossNotification(const LossNotification& loss_notification) override;
+
+		// Eventual handler for the contained encoders' EncodedImageCallbacks, but
+		// called from an internal helper that also knows the correct stream
+		// index.
+		webrtc::EncodedImageCallback::Result OnEncodedImage(
+			size_t stream_idx,
+			const webrtc::EncodedImage& encoded_image,
+			const webrtc::CodecSpecificInfo* codec_specific_info
+#if WEBRTC_VERSION == 84
+			, const webrtc::RTPFragmentationHeader* fragmentation
+#endif
+		);
+
+		EncoderInfo GetEncoderInfo() const override;
+
+	private:
+		struct StreamInfo
+		{
+			StreamInfo(std::unique_ptr<webrtc::VideoEncoder> InEncoder,
+				std::unique_ptr<webrtc::EncodedImageCallback> InCallback,
+				std::unique_ptr<webrtc::FramerateController> InFramerateController,
+				uint16_t InWidth,
+				uint16_t InHeight,
+				bool bInKeyFrame,
+				bool bInSendStream)
+				: Encoder(std::move(InEncoder))
+				, Callback(std::move(InCallback))
+				, FramerateController(std::move(InFramerateController))
+				, Width(InWidth)
+				, Height(InHeight)
+				, KeyFrameRequest(bInKeyFrame)
+				, bSendStream(bInSendStream)
+			{
+			}
+			std::unique_ptr<webrtc::VideoEncoder> Encoder;
+			std::unique_ptr<webrtc::EncodedImageCallback> Callback;
+			std::unique_ptr<webrtc::FramerateController> FramerateController;
+			uint16_t Width;
+			uint16_t Height;
+			bool KeyFrameRequest;
+			bool bSendStream;
+		};
+
+		bool IsInitialized() const;
+
+		TAtomic<bool> Initialized;
+		FVideoEncoderFactorySimulcast& SimulcastEncoderFactory;
+		const webrtc::SdpVideoFormat VideoFormat;
+		webrtc::VideoCodec CurrentCodec;
+		FCriticalSection StreamInfosGuard;
+		std::vector<StreamInfo> StreamInfos;
+		webrtc::EncodedImageCallback* EncodedCompleteCallback;
+
+		int LowestResolutionStreamIndex;
+		int HighestResolutionStreamIndex;
+
+		//int EncodeStream(const webrtc::VideoFrame& InputImage, TSharedPtr<FTextureTripleBuffer> LayerFrameSource, size_t StreamIdx, bool bSendKeyFrame);
+		int EncodeVP8(const webrtc::VideoFrame& InputImage, const std::vector<webrtc::VideoFrameType>* FrameTypes, bool bSendKeyFrame);
+		int EncodeVP9(const webrtc::VideoFrame& InputImage, const std::vector<webrtc::VideoFrameType>* FrameTypes, bool bSendKeyFrame);
+		int EncodeH264(const webrtc::VideoFrame& InputImage, const std::vector<webrtc::VideoFrameType>* FrameTypes, bool bSendKeyFrame);
+	};
+} // namespace UE::PixelStreaming
diff --git a/Engine/Plugins/Media/PixelStreaming/Source/PixelStreaming/Private/VideoEncoderFactorySimple.cpp b/Engine/Plugins/Media/PixelStreaming/Source/PixelStreaming/Private/VideoEncoderFactorySimple.cpp
new file mode 100644
--- /dev/null
+++ b/Engine/Plugins/Media/PixelStreaming/Source/PixelStreaming/Private/VideoEncoderFactorySimple.cpp
@@ -0,0 +1,161 @@
+// Copyright Epic Games, Inc. All Rights Reserved.
+
+#include "VideoEncoderFactorySimple.h"
+#include "absl/strings/match.h"
+#include "VideoEncoderRTC.h"
+#include "VideoEncoderVPX.h"
+#include "Settings.h"
+#include "VideoEncoderFactory.h"
+#include "PixelStreamingPrivate.h"
+#include "Utils.h"
+
+namespace UE::PixelStreaming
+{
+	FVideoEncoderFactorySimple::FVideoEncoderFactorySimple()
+	{
+	}
+
+	FVideoEncoderFactorySimple::~FVideoEncoderFactorySimple()
+	{
+	}
+
+	std::vector<webrtc::SdpVideoFormat> FVideoEncoderFactorySimple::GetSupportedFormats() const
+	{
+		std::vector<webrtc::SdpVideoFormat> video_formats;
+
+		switch (UE::PixelStreaming::Settings::GetSelectedCodec())
+		{
+			case UE::PixelStreaming::Settings::ECodec::VP8:
+				video_formats.push_back(webrtc::SdpVideoFormat(cricket::kVp8CodecName));
+				break;
+			case UE::PixelStreaming::Settings::ECodec::VP9:
+				video_formats.push_back(webrtc::SdpVideoFormat(cricket::kVp9CodecName));
+				break;
+			case UE::PixelStreaming::Settings::ECodec::H264:
+			default:
+#if WEBRTC_VERSION == 84
+				video_formats.push_back(UE::PixelStreaming::CreateH264Format(webrtc::H264::kProfileConstrainedBaseline, webrtc::H264::kLevel3_1));
+				video_formats.push_back(UE::PixelStreaming::CreateH264Format(webrtc::H264::kProfileBaseline, webrtc::H264::kLevel3_1));
+#elif WEBRTC_VERSION == 96
+				video_formats.push_back(UE::PixelStreaming::CreateH264Format(webrtc::H264Profile::kProfileConstrainedBaseline, webrtc::H264Level::kLevel3_1));
+				video_formats.push_back(UE::PixelStreaming::CreateH264Format(webrtc::H264Profile::kProfileBaseline, webrtc::H264Level::kLevel3_1));
+#endif
+				break;
+		}
+		return video_formats;
+	}
+
+	FVideoEncoderFactorySimple::CodecInfo FVideoEncoderFactorySimple::QueryVideoEncoder(const webrtc::SdpVideoFormat& format) const
+	{
+		webrtc::VideoEncoderFactory::CodecInfo CodecInfo = { false
+#if WEBRTC_VERSION == 84
+			, false 
+#endif
+		};
+#if WEBRTC_VERSION == 84
+		CodecInfo.is_hardware_accelerated = true;
+#endif
+		CodecInfo.has_internal_source = false;
+		return CodecInfo;
+	}
+
+	std::unique_ptr<webrtc::VideoEncoder> FVideoEncoderFactorySimple::CreateVideoEncoder(const webrtc::SdpVideoFormat& format)
+	{
+		if (absl::EqualsIgnoreCase(format.name, cricket::kVp8CodecName))
+		{
+			return std::make_unique<FVideoEncoderVPX>(8);
+		}
+		else if (absl::EqualsIgnoreCase(format.name, cricket::kVp9CodecName))
+		{
+			return std::make_unique<FVideoEncoderVPX>(9);
+		}
+		else
+		{
+			// Lock during encoder creation
+			FScopeLock Lock(&ActiveEncodersGuard);
+			auto VideoEncoder = std::make_unique<FVideoEncoderRTC>(*this);
+			ActiveEncoders.Add(VideoEncoder.get());
+			return VideoEncoder;
+		}
+	}
+
+#if WEBRTC_VERSION == 84
+	void FVideoEncoderFactorySimple::OnEncodedImage(const webrtc::EncodedImage& encoded_image, const webrtc::CodecSpecificInfo* codec_specific_info, const webrtc::RTPFragmentationHeader* fragmentation)
+#elif WEBRTC_VERSION == 96
+	void FVideoEncoderFactorySimple::OnEncodedImage(const webrtc::EncodedImage& encoded_image, const webrtc::CodecSpecificInfo* codec_specific_info)
+#endif
+	{
+		// Lock as we send encoded image to each encoder.
+		FScopeLock Lock(&ActiveEncodersGuard);
+
+		// Go through each encoder and send our encoded image to its callback
+		for (FVideoEncoderRTC* Encoder : ActiveEncoders)
+		{
+			Encoder->SendEncodedImage(encoded_image, codec_specific_info
+#if WEBRTC_VERSION == 84
+				, fragmentation
+#endif
+			);
+		}
+	}
+
+	void FVideoEncoderFactorySimple::ReleaseVideoEncoder(FVideoEncoderRTC* Encoder)
+	{
+		// Lock during deleting an encoder
+		FScopeLock Lock(&ActiveEncodersGuard);
+		ActiveEncoders.Remove(Encoder);
+	}
+
+	void FVideoEncoderFactorySimple::ForceKeyFrame()
+	{
+		FScopeLock Lock(&ActiveEncodersGuard);
+		HardwareEncoder->SetForceNextKeyframe();
+	}
+
+	FVideoEncoderH264Wrapper* FVideoEncoderFactorySimple::GetOrCreateHardwareEncoder(int Width, int Height, int MaxBitrate, int TargetBitrate, int MaxFramerate)
+	{
+		if (HardwareEncoder == nullptr)
+		{
+			// Make AVEncoder frame factory.
+			TUniquePtr<FEncoderFrameFactory> FrameFactory = MakeUnique<FEncoderFrameFactory>();
+
+			// Make the encoder config
+			AVEncoder::FVideoEncoder::FLayerConfig EncoderConfig;
+			EncoderConfig.Width = Width;
+			EncoderConfig.Height = Height;
+			EncoderConfig.MaxFramerate = MaxFramerate;
+			EncoderConfig.TargetBitrate = TargetBitrate;
+			EncoderConfig.MaxBitrate = MaxBitrate;
+
+			// Make the actual AVEncoder encoder.
+			const TArray<AVEncoder::FVideoEncoderInfo>& Available = AVEncoder::FVideoEncoderFactory::Get().GetAvailable();
+			TUniquePtr<AVEncoder::FVideoEncoder> Encoder = AVEncoder::FVideoEncoderFactory::Get().Create(Available[0].ID, FrameFactory->GetOrCreateVideoEncoderInput(), EncoderConfig);
+			if (Encoder.IsValid())
+			{
+				Encoder->SetOnEncodedPacket([this](uint32 InLayerIndex, const TSharedPtr<AVEncoder::FVideoEncoderInputFrame> InFrame, const AVEncoder::FCodecPacket& InPacket) {
+					// Note: this is a static method call.
+					FVideoEncoderH264Wrapper::OnEncodedPacket(this, InLayerIndex, InFrame, InPacket);
+				});
+
+				// Make the hardware encoder wrapper
+				HardwareEncoder = MakeUnique<FVideoEncoderH264Wrapper>(MoveTemp(FrameFactory), MoveTemp(Encoder));
+				return HardwareEncoder.Get();
+			}
+			else
+			{
+				UE_LOG(LogPixelStreaming, Error, TEXT("Could not create encoder. Check encoder config or perhaps you used up all your HW encoders."));
+				// We could not make the encoder, so indicate the id was not set successfully.
+				return nullptr;
+			}
+		}
+		else
+		{
+			return HardwareEncoder.Get();
+		}
+	}
+
+	FVideoEncoderH264Wrapper* FVideoEncoderFactorySimple::GetHardwareEncoder()
+	{
+		return HardwareEncoder.Get();
+	}
+} // namespace UE::PixelStreaming
diff --git a/Engine/Plugins/Media/PixelStreaming/Source/PixelStreaming/Private/VideoEncoderFactorySimple.h b/Engine/Plugins/Media/PixelStreaming/Source/PixelStreaming/Private/VideoEncoderFactorySimple.h
new file mode 100644
--- /dev/null
+++ b/Engine/Plugins/Media/PixelStreaming/Source/PixelStreaming/Private/VideoEncoderFactorySimple.h
@@ -0,0 +1,51 @@
+// Copyright Epic Games, Inc. All Rights Reserved.
+#pragma once
+
+#include "WebRTCIncludes.h"
+
+namespace UE::PixelStreaming 
+{
+	class FVideoEncoderRTC;
+	class FVideoEncoderH264Wrapper;
+	
+	class FVideoEncoderFactorySimple : public webrtc::VideoEncoderFactory
+	{
+	public:
+		FVideoEncoderFactorySimple();
+		virtual ~FVideoEncoderFactorySimple() override;
+
+		virtual std::vector<webrtc::SdpVideoFormat> GetSupportedFormats() const override;
+
+		// Always returns our H264 hardware encoders codec_info for now/
+		virtual CodecInfo QueryVideoEncoder(const webrtc::SdpVideoFormat& Format) const override;
+
+		// Always returns our H264 hardware encoder for now.
+		virtual std::unique_ptr<webrtc::VideoEncoder> CreateVideoEncoder(const webrtc::SdpVideoFormat& Format) override;
+
+		void ReleaseVideoEncoder(FVideoEncoderRTC* Encoder);
+		void ForceKeyFrame();
+			
+		FVideoEncoderH264Wrapper* GetOrCreateHardwareEncoder(int Width, int Height, int MaxBitrate, int TargetBitrate, int MaxFramerate);
+		FVideoEncoderH264Wrapper* GetHardwareEncoder();
+
+#if WEBRTC_VERSION == 84
+		void OnEncodedImage(const webrtc::EncodedImage& Encoded_image, const webrtc::CodecSpecificInfo* CodecSpecificInfo, const webrtc::RTPFragmentationHeader* Fragmentation);
+#elif WEBRTC_VERSION == 96
+		void OnEncodedImage(const webrtc::EncodedImage& Encoded_image, const webrtc::CodecSpecificInfo* CodecSpecificInfo);
+#endif
+
+	private:
+#if WEBRTC_VERSION == 84
+		static webrtc::SdpVideoFormat CreateH264Format(webrtc::H264::Profile Profile, webrtc::H264::Level Level);
+#elif WEBRTC_VERSION == 96
+		static webrtc::SdpVideoFormat CreateH264Format(webrtc::H264Profile Profile, webrtc::H264Level Level);
+#endif
+
+		TUniquePtr<FVideoEncoderH264Wrapper> HardwareEncoder;
+
+		// Encoders assigned to each peer. Each one of these will be assigned to one of the hardware encoders
+		// raw pointers here because the factory interface wants a unique_ptr here so we cant own the object.
+		TArray<FVideoEncoderRTC*> ActiveEncoders;
+		FCriticalSection ActiveEncodersGuard;
+	};
+} // namespace UE::PixelStreaming
diff --git a/Engine/Plugins/Media/PixelStreaming/Source/PixelStreaming/Private/VideoEncoderFactorySimulcast.cpp b/Engine/Plugins/Media/PixelStreaming/Source/PixelStreaming/Private/VideoEncoderFactorySimulcast.cpp
new file mode 100644
--- /dev/null
+++ b/Engine/Plugins/Media/PixelStreaming/Source/PixelStreaming/Private/VideoEncoderFactorySimulcast.cpp
@@ -0,0 +1,66 @@
+// Copyright Epic Games, Inc. All Rights Reserved.
+
+#include "VideoEncoderFactorySimulcast.h"
+#include "Settings.h"
+#include "absl/strings/match.h"
+#include "WebRTCIncludes.h"
+#include "VideoEncoderAdapterSimulcast.h"
+#include "VideoEncoderFactorySimple.h"
+
+namespace UE::PixelStreaming
+{
+	FVideoEncoderFactorySimulcast::FVideoEncoderFactorySimulcast()
+		: PrimaryEncoderFactory(MakeUnique<FVideoEncoderFactorySimple>())
+	{
+		// Make a copy of simulcast settings and sort them based on scaling.
+		for (FEncoderFactoryId i = 0; i < Settings::SimulcastParameters.Layers.Num(); i++)
+		{
+			GetOrCreateEncoderFactory(i);
+		}
+	}
+
+	FVideoEncoderFactorySimulcast::~FVideoEncoderFactorySimulcast()
+	{
+	}
+
+	std::unique_ptr<webrtc::VideoEncoder> FVideoEncoderFactorySimulcast::CreateVideoEncoder(const webrtc::SdpVideoFormat& format)
+	{
+		return std::make_unique<FVideoEncoderAdapterSimulcast>(*this, format);
+	}
+
+	FVideoEncoderFactorySimple* FVideoEncoderFactorySimulcast::GetEncoderFactory(FEncoderFactoryId Id)
+	{
+		FScopeLock Lock(&EncoderFactoriesGuard);
+		if (auto&& Existing = EncoderFactories.Find(Id))
+		{
+			return Existing->Get();
+		}
+		return nullptr;
+	}
+
+	FVideoEncoderFactorySimple* FVideoEncoderFactorySimulcast::GetOrCreateEncoderFactory(FEncoderFactoryId Id)
+	{
+		FVideoEncoderFactorySimple* Existing = GetEncoderFactory(Id);
+		if (Existing == nullptr)
+		{
+			FScopeLock Lock(&EncoderFactoriesGuard);
+			TUniquePtr<FVideoEncoderFactorySimple> EncoderFactory = MakeUnique<FVideoEncoderFactorySimple>();
+			EncoderFactories.Add(Id, MoveTemp(EncoderFactory));
+			return EncoderFactory.Get();
+		}
+		else
+		{
+			return Existing;
+		}
+	}
+
+	std::vector<webrtc::SdpVideoFormat> FVideoEncoderFactorySimulcast::GetSupportedFormats() const
+	{
+		return PrimaryEncoderFactory->GetSupportedFormats();
+	}
+
+	FVideoEncoderFactorySimulcast::CodecInfo FVideoEncoderFactorySimulcast::QueryVideoEncoder(const webrtc::SdpVideoFormat& format) const
+	{
+		return PrimaryEncoderFactory->QueryVideoEncoder(format);
+	}
+} // namespace UE::PixelStreaming
diff --git a/Engine/Plugins/Media/PixelStreaming/Source/PixelStreaming/Private/VideoEncoderFactorySimulcast.h b/Engine/Plugins/Media/PixelStreaming/Source/PixelStreaming/Private/VideoEncoderFactorySimulcast.h
new file mode 100644
--- /dev/null
+++ b/Engine/Plugins/Media/PixelStreaming/Source/PixelStreaming/Private/VideoEncoderFactorySimulcast.h
@@ -0,0 +1,28 @@
+// Copyright Epic Games, Inc. All Rights Reserved.
+#pragma once
+
+#include "WebRTCIncludes.h"
+
+namespace UE::PixelStreaming 
+{
+	class FVideoEncoderFactorySimple;
+	
+	class FVideoEncoderFactorySimulcast : public webrtc::VideoEncoderFactory
+	{
+	public:
+		FVideoEncoderFactorySimulcast();
+		virtual ~FVideoEncoderFactorySimulcast();
+		virtual std::vector<webrtc::SdpVideoFormat> GetSupportedFormats() const ;
+		virtual CodecInfo QueryVideoEncoder(const webrtc::SdpVideoFormat& format) const override;
+		virtual std::unique_ptr<webrtc::VideoEncoder> CreateVideoEncoder(const webrtc::SdpVideoFormat& format);
+
+		using FEncoderFactoryId = uint64;
+		FVideoEncoderFactorySimple* GetOrCreateEncoderFactory(FEncoderFactoryId Id);
+		FVideoEncoderFactorySimple* GetEncoderFactory(FEncoderFactoryId Id);
+	private:
+		TMap<FEncoderFactoryId, TUniquePtr<FVideoEncoderFactorySimple>> EncoderFactories;
+		FCriticalSection EncoderFactoriesGuard;
+
+		TUniquePtr<FVideoEncoderFactorySimple> PrimaryEncoderFactory;
+	};
+} // namespace UE::PixelStreaming
diff --git a/Engine/Plugins/Media/PixelStreaming/Source/PixelStreaming/Private/VideoEncoderH264Wrapper.cpp b/Engine/Plugins/Media/PixelStreaming/Source/PixelStreaming/Private/VideoEncoderH264Wrapper.cpp
--- a/Engine/Plugins/Media/PixelStreaming/Source/PixelStreaming/Private/VideoEncoderH264Wrapper.cpp
+++ b/Engine/Plugins/Media/PixelStreaming/Source/PixelStreaming/Private/VideoEncoderH264Wrapper.cpp
@@ -3,8 +3,9 @@
 #include "VideoEncoderH264Wrapper.h"
 #include "Stats.h"
 #include "VideoEncoderFactory.h"
-#include "EncoderFactory.h"
-#include "FrameBuffer.h"
+#include "VideoEncoderFactorySimple.h"
+#include "FrameBufferH264.h"
+#include "FrameAdapterH264.h"
 
 namespace UE::PixelStreaming
 {
@@ -22,18 +23,11 @@ namespace UE::PixelStreaming
 
 	void FVideoEncoderH264Wrapper::Encode(const webrtc::VideoFrame& WebRTCFrame, bool bKeyframe)
 	{
-		FPixelStreamingFrameBuffer* FrameBuffer = static_cast<FPixelStreamingFrameBuffer*>(WebRTCFrame.video_frame_buffer().get());
-		check(FrameBuffer->GetFrameBufferType() == EPixelStreamingFrameBufferType::Layer);
-		FLayerFrameBuffer* LayerFrameBuffer = static_cast<FLayerFrameBuffer*>(FrameBuffer);
-		TSharedPtr<FPixelStreamingTextureWrapper> SourceTexture = LayerFrameBuffer->GetFrame();
-
-		if (!SourceTexture.IsValid())
-		{
-			return;
-		}
+		FFrameBufferH264* FrameBuffer = static_cast<FFrameBufferH264*>(WebRTCFrame.video_frame_buffer().get());
+		checkf(FrameBuffer->GetFrameBufferType() == EPixelStreamingFrameBufferType::Layer, TEXT("FVideoEncoderH264Wrapper::Encode(): Supplied frame buffer is of incorrect type."));
 
 		AVEncoder::FVideoEncoder::FLayerConfig EncoderConfig = GetCurrentConfig();
-		TSharedPtr<AVEncoder::FVideoEncoderInputFrame> EncoderInputFrame = FrameFactory->GetFrameAndSetTexture(SourceTexture->GetTexture<FTextureRHIRef>());
+		TSharedPtr<AVEncoder::FVideoEncoderInputFrame> EncoderInputFrame = FrameFactory->GetFrameAndSetTexture(FrameBuffer->GetAdaptedLayer()->GetFrameTexture());
 		if (EncoderInputFrame)
 		{
 			EncoderInputFrame->SetTimestampUs(WebRTCFrame.timestamp_us());
@@ -72,10 +66,11 @@ namespace UE::PixelStreaming
 
 	/* ------------------ Static functions below --------------------- */
 
-	void FVideoEncoderH264Wrapper::OnEncodedPacket(FVideoEncoderFactory* Factory, uint32 InLayerIndex, const TSharedPtr<AVEncoder::FVideoEncoderInputFrame> InFrame, const AVEncoder::FCodecPacket& InPacket)
+	void FVideoEncoderH264Wrapper::OnEncodedPacket(FVideoEncoderFactorySimple* Factory, uint32 InLayerIndex, const TSharedPtr<AVEncoder::FVideoEncoderInputFrame> InFrame, const AVEncoder::FCodecPacket& InPacket)
 	{
 		webrtc::EncodedImage Image;
 
+#if WEBRTC_VERSION == 84
 		webrtc::RTPFragmentationHeader FragHeader;
 		//CreateH264FragmentHeader(InPacket.Data.Get(), InPacket.DataSize, FragHeader);
 
@@ -88,7 +83,7 @@ namespace UE::PixelStreaming
 			FragHeader.fragmentationOffset[i] = NALUIndex.payload_start_offset;
 			FragHeader.fragmentationLength[i] = NALUIndex.payload_size;
 		}
-
+#endif
 		Image.timing_.packetization_finish_ms = FTimespan::FromSeconds(FPlatformTime::Seconds()).GetTotalMilliseconds();
 		Image.timing_.encode_start_ms = InPacket.Timings.StartTs.GetTotalMilliseconds();
 		Image.timing_.encode_finish_ms = InPacket.Timings.FinishTs.GetTotalMilliseconds();
@@ -101,7 +96,9 @@ namespace UE::PixelStreaming
 		Image.content_type_ = webrtc::VideoContentType::UNSPECIFIED;
 		Image.qp_ = InPacket.VideoQP;
 		Image.SetSpatialIndex(InLayerIndex);
+#if WEBRTC_VERSION == 84
 		Image._completeFrame = true;
+#endif
 		Image.rotation_ = webrtc::VideoRotation::kVideoRotation_0;
 		Image.SetTimestamp(InFrame->GetTimestampRTP());
 		Image.capture_time_ms_ = InFrame->GetTimestampUs() / 1000.0;
@@ -113,6 +110,10 @@ namespace UE::PixelStreaming
 		CodecInfo.codecSpecific.H264.idr_frame = InPacket.IsKeyFrame;
 		CodecInfo.codecSpecific.H264.base_layer_sync = false;
 
+#if WEBRTC_VERSION == 84
 		Factory->OnEncodedImage(Image, &CodecInfo, &FragHeader);
+#elif WEBRTC_VERSION == 96
+		Factory->OnEncodedImage(Image, &CodecInfo);
+#endif
 	}
 } // namespace UE::PixelStreaming
\ No newline at end of file
diff --git a/Engine/Plugins/Media/PixelStreaming/Source/PixelStreaming/Private/VideoEncoderH264Wrapper.h b/Engine/Plugins/Media/PixelStreaming/Source/PixelStreaming/Private/VideoEncoderH264Wrapper.h
--- a/Engine/Plugins/Media/PixelStreaming/Source/PixelStreaming/Private/VideoEncoderH264Wrapper.h
+++ b/Engine/Plugins/Media/PixelStreaming/Source/PixelStreaming/Private/VideoEncoderH264Wrapper.h
@@ -11,7 +11,7 @@ namespace webrtc
 
 namespace UE::PixelStreaming 
 {
-	class FVideoEncoderFactory;
+	class FVideoEncoderFactorySimple;
 	
 	class FVideoEncoderH264Wrapper
 	{
@@ -28,7 +28,7 @@ namespace UE::PixelStreaming
 		AVEncoder::FVideoEncoder::FLayerConfig GetCurrentConfig();
 		void SetConfig(const AVEncoder::FVideoEncoder::FLayerConfig& NewConfig);
 
-		static void OnEncodedPacket(FVideoEncoderFactory* Factory, 
+		static void OnEncodedPacket(FVideoEncoderFactorySimple* Factory, 
 			uint32 InLayerIndex, const TSharedPtr<AVEncoder::FVideoEncoderInputFrame> InFrame, 
 			const AVEncoder::FCodecPacket& InPacket);
 
diff --git a/Engine/Plugins/Media/PixelStreaming/Source/PixelStreaming/Private/VideoEncoderRTC.cpp b/Engine/Plugins/Media/PixelStreaming/Source/PixelStreaming/Private/VideoEncoderRTC.cpp
--- a/Engine/Plugins/Media/PixelStreaming/Source/PixelStreaming/Private/VideoEncoderRTC.cpp
+++ b/Engine/Plugins/Media/PixelStreaming/Source/PixelStreaming/Private/VideoEncoderRTC.cpp
@@ -1,22 +1,18 @@
 // Copyright Epic Games, Inc. All Rights Reserved.
 
 #include "VideoEncoderRTC.h"
-#include "EncoderFactory.h"
+#include "VideoEncoderFactorySimple.h"
 #include "VideoEncoderFactory.h"
-#include "PixelStreamingFrameBuffer.h"
-#include "PlayerSession.h"
-#include "Stats.h"
-#include "UnrealEngine.h"
-#include "Misc/Paths.h"
+#include "FrameBufferH264.h"
+#include "FrameAdapterH264.h"
 #include "Settings.h"
-#include "PixelStreamingPlayerId.h"
-#include "Async/Async.h"
-#include "VideoEncoderH264Wrapper.h"
 #include "HAL/PlatformFileManager.h"
+#include "Misc/Paths.h"
+#include "Stats.h"
 
 namespace UE::PixelStreaming
 {
-	FVideoEncoderRTC::FVideoEncoderRTC(FVideoEncoderFactory& InFactory)
+	FVideoEncoderRTC::FVideoEncoderRTC(FVideoEncoderFactorySimple& InFactory)
 		: Factory(InFactory)
 	{
 	}
@@ -59,7 +55,7 @@ namespace UE::PixelStreaming
 
 	int32 FVideoEncoderRTC::Encode(webrtc::VideoFrame const& frame, std::vector<webrtc::VideoFrameType> const* frame_types)
 	{
-		FPixelStreamingFrameBuffer* FrameBuffer = static_cast<FPixelStreamingFrameBuffer*>(frame.video_frame_buffer().get());
+		FFrameBufferH264* FrameBuffer = static_cast<FFrameBufferH264*>(frame.video_frame_buffer().get());
 
 		// If initialize frame is passed, this is a dummy frame so WebRTC is happy frames are passing through
 		// This is used so that an encoder can be active as far as WebRTC is concerned by simply have frames transmitted by some other encoder on its behalf.
@@ -68,6 +64,14 @@ namespace UE::PixelStreaming
 			return WEBRTC_VIDEO_CODEC_OK;
 		}
 
+		FPixelStreamingFrameMetadata& FrameMetadata = FrameBuffer->GetAdaptedLayer()->Metadata;
+		FrameMetadata.UseCount++;
+		FrameMetadata.LastEncodeStartTime = FPlatformTime::Cycles64();
+		if (FrameMetadata.UseCount == 1)
+		{
+			FrameMetadata.FirstEncodeStartTime = FrameMetadata.LastEncodeStartTime;
+		}
+
 		bool bKeyframe = false;
 		if ((frame_types && (*frame_types)[0] == webrtc::VideoFrameType::kVideoFrameKey))
 		{
@@ -84,6 +88,10 @@ namespace UE::PixelStreaming
 
 		Encoder->Encode(frame, bKeyframe);
 
+		FrameMetadata.LastEncodeEndTime = FPlatformTime::Cycles64();
+
+		FStats::Get()->AddFrameTimingStats(FrameMetadata);
+
 		return WEBRTC_VIDEO_CODEC_OK;
 	}
 
@@ -162,7 +170,11 @@ namespace UE::PixelStreaming
 		return InEncoderConfig;
 	}
 
-	void FVideoEncoderRTC::SendEncodedImage(webrtc::EncodedImage const& encoded_image, webrtc::CodecSpecificInfo const* codec_specific_info, webrtc::RTPFragmentationHeader const* fragmentation)
+	void FVideoEncoderRTC::SendEncodedImage(webrtc::EncodedImage const& encoded_image, webrtc::CodecSpecificInfo const* codec_specific_info
+#if WEBRTC_VERSION == 84
+		, webrtc::RTPFragmentationHeader const* fragmentation
+#endif
+	)
 	{
 		if (FirstKeyframeCountdown > 0)
 		{
@@ -195,7 +207,11 @@ namespace UE::PixelStreaming
 
 		if (OnEncodedImageCallback)
 		{
-			OnEncodedImageCallback->OnEncodedImage(encoded_image, codec_specific_info, fragmentation);
+			OnEncodedImageCallback->OnEncodedImage(encoded_image, codec_specific_info
+#if WEBRTC_VERSION == 84
+				, fragmentation
+#endif
+			);
 		}
 	}
 } // namespace UE::PixelStreaming
\ No newline at end of file
