Early Pixel Streaming 5.1 features, bugfixes, and refactors that benefit main. (From various authors)

diff --git a/Engine/Plugins/Media/PixelStreaming/Source/PixelStreaming/Private/VideoEncoderRTC.h b/Engine/Plugins/Media/PixelStreaming/Source/PixelStreaming/Private/VideoEncoderRTC.h
--- a/Engine/Plugins/Media/PixelStreaming/Source/PixelStreaming/Private/VideoEncoderRTC.h
+++ b/Engine/Plugins/Media/PixelStreaming/Source/PixelStreaming/Private/VideoEncoderRTC.h
@@ -1,23 +1,19 @@
 // Copyright Epic Games, Inc. All Rights Reserved.
 #pragma once
 
-#include "PixelStreamingPlayerId.h"
-#include "HAL/ThreadSafeBool.h"
 #include "WebRTCIncludes.h"
 #include "VideoEncoder.h"
-#include "Templates/SharedPointer.h"
-#include "Misc/Optional.h"
 #include "VideoEncoderH264Wrapper.h"
 
 namespace UE::PixelStreaming 
 {
-	class FVideoEncoderFactory;
+	class FVideoEncoderFactorySimple;
 
 	// Implementation that is a WebRTC video encoder that allows us tie to our actually underlying non-WebRTC video encoder.
 	class FVideoEncoderRTC : public webrtc::VideoEncoder
 	{
 	public:
-		FVideoEncoderRTC(FVideoEncoderFactory& InFactory);
+		FVideoEncoderRTC(FVideoEncoderFactorySimple& InFactory);
 		virtual ~FVideoEncoderRTC() override;
 
 		// WebRTC Interface
@@ -35,13 +31,17 @@ namespace UE::PixelStreaming
 		// virtual void OnLossNotification(const LossNotification& loss_notification) override;
 		// End WebRTC Interface.
 
-		void SendEncodedImage(webrtc::EncodedImage const& encoded_image, webrtc::CodecSpecificInfo const* codec_specific_info, webrtc::RTPFragmentationHeader const* fragmentation);
+		void SendEncodedImage(webrtc::EncodedImage const& encoded_image, webrtc::CodecSpecificInfo const* codec_specific_info
+#if WEBRTC_VERSION == 84
+			, webrtc::RTPFragmentationHeader const* fragmentation
+#endif
+		);
 
 	private:
 		void UpdateConfig();
 		AVEncoder::FVideoEncoder::FLayerConfig CreateEncoderConfigFromCVars(AVEncoder::FVideoEncoder::FLayerConfig BaseEncoderConfig) const;
 
-		FVideoEncoderFactory& Factory;
+		FVideoEncoderFactorySimple& Factory;
 
 		FVideoEncoderH264Wrapper* HardwareEncoder;
 
diff --git a/Engine/Plugins/Media/PixelStreaming/Source/PixelStreaming/Private/VideoEncoderVPX.cpp b/Engine/Plugins/Media/PixelStreaming/Source/PixelStreaming/Private/VideoEncoderVPX.cpp
new file mode 100644
--- /dev/null
+++ b/Engine/Plugins/Media/PixelStreaming/Source/PixelStreaming/Private/VideoEncoderVPX.cpp
@@ -0,0 +1,71 @@
+// Copyright Epic Games, Inc. All Rights Reserved.
+
+#include "VideoEncoderVPX.h"
+#include "FrameBufferI420.h"
+#include "FrameAdapterI420.h"
+#include "Stats.h"
+
+namespace UE::PixelStreaming
+{
+	FVideoEncoderVPX::FVideoEncoderVPX(int VPXVersion)
+	{
+		if (VPXVersion == 8)
+		{
+			WebRTCVPXEncoder = webrtc::VP8Encoder::Create();
+		}
+		else if (VPXVersion == 9)
+		{
+			WebRTCVPXEncoder = webrtc::VP9Encoder::Create();
+		}
+	}
+
+	int FVideoEncoderVPX::InitEncode(webrtc::VideoCodec const* InCodecSettings, VideoEncoder::Settings const& settings)
+	{
+		return WebRTCVPXEncoder->InitEncode(InCodecSettings, settings);
+	}
+
+	int32 FVideoEncoderVPX::RegisterEncodeCompleteCallback(webrtc::EncodedImageCallback* callback)
+	{
+		return WebRTCVPXEncoder->RegisterEncodeCompleteCallback(callback);
+	}
+
+	int32 FVideoEncoderVPX::Release()
+	{
+		return WebRTCVPXEncoder->Release();
+	}
+
+	int32 FVideoEncoderVPX::Encode(webrtc::VideoFrame const& frame, std::vector<webrtc::VideoFrameType> const* frame_types)
+	{
+		FFrameBufferI420* FrameBuffer = static_cast<FFrameBufferI420*>(frame.video_frame_buffer().get());
+
+		FPixelStreamingFrameMetadata& FrameMetadata = FrameBuffer->GetAdaptedLayer()->Metadata;
+		FrameMetadata.UseCount++;
+		FrameMetadata.LastEncodeStartTime = FPlatformTime::Cycles64();
+		if (FrameMetadata.UseCount == 1)
+		{
+			FrameMetadata.FirstEncodeStartTime = FrameMetadata.LastEncodeStartTime;
+		}
+
+		const int32 EncodeResult = WebRTCVPXEncoder->Encode(frame, frame_types);
+
+		FrameBuffer->GetAdaptedLayer()->Metadata.LastEncodeEndTime = FPlatformTime::Cycles64();
+
+		FStats::Get()->AddFrameTimingStats(FrameBuffer->GetAdaptedLayer()->Metadata);
+
+		return EncodeResult;
+	}
+
+	// Pass rate control parameters from WebRTC to our encoder
+	// This is how WebRTC can control the bitrate/framerate of the encoder.
+	void FVideoEncoderVPX::SetRates(RateControlParameters const& parameters)
+	{
+		WebRTCVPXEncoder->SetRates(parameters);
+	}
+
+	webrtc::VideoEncoder::EncoderInfo FVideoEncoderVPX::GetEncoderInfo() const
+	{
+		VideoEncoder::EncoderInfo info = WebRTCVPXEncoder->GetEncoderInfo();
+		info.supports_native_handle = true;
+		return info;
+	}
+} // namespace UE::PixelStreaming
diff --git a/Engine/Plugins/Media/PixelStreaming/Source/PixelStreaming/Private/VideoEncoderVPX.h b/Engine/Plugins/Media/PixelStreaming/Source/PixelStreaming/Private/VideoEncoderVPX.h
new file mode 100644
--- /dev/null
+++ b/Engine/Plugins/Media/PixelStreaming/Source/PixelStreaming/Private/VideoEncoderVPX.h
@@ -0,0 +1,25 @@
+// Copyright Epic Games, Inc. All Rights Reserved.
+#pragma once
+
+#include "WebRTCIncludes.h"
+
+namespace UE::PixelStreaming 
+{
+	class FVideoEncoderVPX : public webrtc::VideoEncoder
+	{
+	public:
+		FVideoEncoderVPX(int VPXVersion);
+		virtual ~FVideoEncoderVPX() = default;
+
+		// WebRTC Interface
+		virtual int InitEncode(webrtc::VideoCodec const* codec_settings, webrtc::VideoEncoder::Settings const& settings) override;
+		virtual int32 RegisterEncodeCompleteCallback(webrtc::EncodedImageCallback* callback) override;
+		virtual int32 Release() override;
+		virtual int32 Encode(webrtc::VideoFrame const& frame, std::vector<webrtc::VideoFrameType> const* frame_types) override;
+		virtual void SetRates(RateControlParameters const& parameters) override;
+		virtual webrtc::VideoEncoder::EncoderInfo GetEncoderInfo() const override;
+
+	private:
+		std::unique_ptr<webrtc::VideoEncoder> WebRTCVPXEncoder;
+	};
+} // namespace UE::PixelStreaming
diff --git a/Engine/Plugins/Media/PixelStreaming/Source/PixelStreaming/Private/VideoInputBackBuffer.cpp b/Engine/Plugins/Media/PixelStreaming/Source/PixelStreaming/Private/VideoInputBackBuffer.cpp
new file mode 100644
--- /dev/null
+++ b/Engine/Plugins/Media/PixelStreaming/Source/PixelStreaming/Private/VideoInputBackBuffer.cpp
@@ -0,0 +1,27 @@
+// Copyright Epic Games, Inc. All Rights Reserved.
+
+#include "VideoInputBackBuffer.h"
+#include "Framework/Application/SlateApplication.h"
+#include "PixelStreamingSourceFrame.h"
+#include "Settings.h"
+
+namespace UE::PixelStreaming
+{
+	FVideoInputBackBuffer::FVideoInputBackBuffer()
+	{
+		DelegateHandle = FSlateApplication::Get().GetRenderer()->OnBackBufferReadyToPresent().AddRaw(this, &FVideoInputBackBuffer::OnBackBufferReady);
+	}
+
+	FVideoInputBackBuffer::~FVideoInputBackBuffer()
+	{
+		if (!IsEngineExitRequested())
+		{
+			FSlateApplication::Get().GetRenderer()->OnBackBufferReadyToPresent().Remove(DelegateHandle);
+		}
+	}
+
+	void FVideoInputBackBuffer::OnBackBufferReady(SWindow& SlateWindow, const FTexture2DRHIRef& FrameBuffer)
+	{
+		OnFrame.Broadcast(FPixelStreamingSourceFrame(FrameBuffer));
+	}
+} // namespace UE::PixelStreaming
diff --git a/Engine/Plugins/Media/PixelStreaming/Source/PixelStreaming/Private/VideoInputBackBuffer.h b/Engine/Plugins/Media/PixelStreaming/Source/PixelStreaming/Private/VideoInputBackBuffer.h
new file mode 100644
--- /dev/null
+++ b/Engine/Plugins/Media/PixelStreaming/Source/PixelStreaming/Private/VideoInputBackBuffer.h
@@ -0,0 +1,22 @@
+// Copyright Epic Games, Inc. All Rights Reserved.
+
+#pragma once
+
+#include "PixelStreamingVideoInput.h"
+#include "Widgets/SWindow.h"
+#include "RHI.h"
+
+namespace UE::PixelStreaming
+{
+	class FVideoInputBackBuffer : public FPixelStreamingVideoInput
+	{
+	public:
+		FVideoInputBackBuffer();
+		virtual ~FVideoInputBackBuffer();
+
+	private:
+		FDelegateHandle DelegateHandle;
+
+		void OnBackBufferReady(SWindow& SlateWindow, const FTexture2DRHIRef& FrameBuffer);
+	};
+} // namespace UE::PixelStreaming
diff --git a/Engine/Plugins/Media/PixelStreaming/Source/PixelStreaming/Private/VideoSourceBase.cpp b/Engine/Plugins/Media/PixelStreaming/Source/PixelStreaming/Private/VideoSourceBase.cpp
new file mode 100644
--- /dev/null
+++ b/Engine/Plugins/Media/PixelStreaming/Source/PixelStreaming/Private/VideoSourceBase.cpp
@@ -0,0 +1,47 @@
+// Copyright Epic Games, Inc. All Rights Reserved.
+
+#include "VideoSourceBase.h"
+
+namespace UE::PixelStreaming
+{
+	FVideoSourceBase::FVideoSourceBase()
+	{
+		CurrentState = webrtc::MediaSourceInterface::SourceState::kInitializing;
+	}
+
+	FVideoSourceBase::~FVideoSourceBase()
+	{
+		CurrentState = webrtc::MediaSourceInterface::SourceState::kEnded;
+	}
+
+	void FVideoSourceBase::PushFrame()
+	{
+		static int32 sFrameId = 1;
+
+		CurrentState = webrtc::MediaSourceInterface::SourceState::kLive;
+
+		webrtc::VideoFrame Frame = CreateFrame(sFrameId++);
+
+		// might want to allow source frames to be scaled here?
+		const FIntPoint Resolution(Frame.width(), Frame.height());
+		if (!AdaptCaptureFrame(Frame.timestamp_us(), Resolution))
+		{
+			return;
+		}
+
+		OnFrame(Frame);
+	}
+
+	bool FVideoSourceBase::AdaptCaptureFrame(const int64 TimestampUs, FIntPoint Resolution)
+	{
+		// disabled for now since we want to control framerate directly
+
+		//int outWidth, outHeight, cropWidth, cropHeight, cropX, cropY;
+		//if (!AdaptFrame(Resolution.X, Resolution.Y, TimestampUs, &outWidth, &outHeight, &cropWidth, &cropHeight, &cropX, &cropY))
+		//{
+		//	return false;
+		//}
+
+		return true;
+	}
+} // namespace UE::PixelStreaming
diff --git a/Engine/Plugins/Media/PixelStreaming/Source/PixelStreaming/Private/VideoSourceBase.h b/Engine/Plugins/Media/PixelStreaming/Source/PixelStreaming/Private/VideoSourceBase.h
new file mode 100644
--- /dev/null
+++ b/Engine/Plugins/Media/PixelStreaming/Source/PixelStreaming/Private/VideoSourceBase.h
@@ -0,0 +1,36 @@
+// Copyright Epic Games, Inc. All Rights Reserved.
+
+#pragma once
+
+#include "AdaptedVideoTrackSource.h"
+
+namespace UE::PixelStreaming
+{
+	/*
+	 * Video sources are WebRTC video sources that populate WebRTC tracks
+	 * and pass WebRTC video frames to `OnFrame`, which eventually gets passed to a WebRTC video encoder, encoded, and transmitted.
+	 */
+	class FVideoSourceBase : public FAdaptedVideoTrackSource
+	{
+	public:
+		FVideoSourceBase();
+		virtual ~FVideoSourceBase();
+
+		/* Begin UE::PixelStreaming::AdaptedVideoTrackSource overrides */
+		virtual webrtc::MediaSourceInterface::SourceState state() const override { return CurrentState; }
+		virtual bool remote() const override { return false; }
+		virtual bool is_screencast() const override { return false; }
+		virtual absl::optional<bool> needs_denoising() const override { return false; }
+		/* End UE::PixelStreaming::AdaptedVideoTrackSource overrides */
+
+		virtual bool IsReady() const = 0;
+		void PushFrame();
+
+	protected:
+		virtual bool AdaptCaptureFrame(const int64 TimestampUs, FIntPoint Resolution);
+		virtual webrtc::VideoFrame CreateFrame(int32 FrameId) = 0;
+
+	protected:
+		webrtc::MediaSourceInterface::SourceState CurrentState;
+	};
+} // namespace UE::PixelStreaming
diff --git a/Engine/Plugins/Media/PixelStreaming/Source/PixelStreaming/Private/VideoSourceGroup.cpp b/Engine/Plugins/Media/PixelStreaming/Source/PixelStreaming/Private/VideoSourceGroup.cpp
new file mode 100644
--- /dev/null
+++ b/Engine/Plugins/Media/PixelStreaming/Source/PixelStreaming/Private/VideoSourceGroup.cpp
@@ -0,0 +1,179 @@
+// Copyright Epic Games, Inc. All Rights Reserved.
+
+#include "VideoSourceGroup.h"
+#include "VideoSourceP2P.h"
+#include "VideoSourceSFU.h"
+#include "Settings.h"
+
+namespace UE::PixelStreaming
+{
+	void FVideoSourceGroup::SetVideoInput(TSharedPtr<FPixelStreamingVideoInput> InVideoInput)
+	{
+		VideoInput = InVideoInput;
+	}
+
+	void FVideoSourceGroup::SetFPS(int32 InFramesPerSecond)
+	{
+		FramesPerSecond = InFramesPerSecond;
+	}
+
+	rtc::scoped_refptr<webrtc::VideoTrackSourceInterface> FVideoSourceGroup::CreateVideoSource(const TFunction<bool()>& InIsQualityControllerFunc)
+	{
+		rtc::scoped_refptr<FVideoSourceBase> NewVideoSource = new FVideoSourceP2P(VideoInput, InIsQualityControllerFunc);
+		{
+			FScopeLock Lock(&CriticalSection);
+			VideoSources.Add(NewVideoSource);
+		}
+		CheckStartStopThread();
+		return NewVideoSource;
+	}
+
+	rtc::scoped_refptr<webrtc::VideoTrackSourceInterface> FVideoSourceGroup::CreateSFUVideoSource()
+	{
+		rtc::scoped_refptr<FVideoSourceBase> NewVideoSource = new FVideoSourceSFU(VideoInput);
+		{
+			FScopeLock Lock(&CriticalSection);
+			VideoSources.Add(NewVideoSource);
+		}
+		CheckStartStopThread();
+		return NewVideoSource;
+	}
+
+	void FVideoSourceGroup::RemoveVideoSource(const webrtc::VideoTrackSourceInterface* ToRemove)
+	{
+		{
+			FScopeLock Lock(&CriticalSection);
+			VideoSources.RemoveAll([ToRemove](const rtc::scoped_refptr<FVideoSourceBase>& Target) {
+				return Target.get() == ToRemove;
+			});
+		}
+		CheckStartStopThread();
+	}
+
+	void FVideoSourceGroup::Start()
+	{
+		if (!bRunning)
+		{
+			if (VideoSources.Num() > 0)
+			{
+				StartThread();
+			}
+			bRunning = true;
+		}
+	}
+
+	void FVideoSourceGroup::Stop()
+	{
+		if (bRunning)
+		{
+			StopThread();
+			bRunning = false;
+		}
+	}
+
+	void FVideoSourceGroup::Tick()
+	{
+		FScopeLock Lock(&CriticalSection);
+		// for each player session, push a frame
+		for (auto& VideoSource : VideoSources)
+		{
+			if (VideoSource && VideoSource->IsReady())
+			{
+				VideoSource->PushFrame();
+			}
+		}
+	}
+
+	void FVideoSourceGroup::StartThread()
+	{
+		if (Settings::DecoupleFrameRate() && !bThreadRunning)
+		{
+			FrameRunnable = MakeUnique<FFrameThread>(this);
+			FrameThread = FRunnableThread::Create(FrameRunnable.Get(), TEXT("FVideoSourceGroup Thread"));
+			bThreadRunning = true;
+		}
+	}
+
+	void FVideoSourceGroup::StopThread()
+	{
+		if (FrameThread != nullptr)
+		{
+			FrameThread->Kill(true);
+		}
+		FrameThread = nullptr;
+		bThreadRunning = false;
+	}
+
+	void FVideoSourceGroup::CheckStartStopThread()
+	{
+		if (bRunning)
+		{
+			const int32 NumSources = VideoSources.Num();
+			if (bThreadRunning && NumSources == 0)
+			{
+				StopThread();
+			}
+			else if (!bThreadRunning && NumSources > 0)
+			{
+				StartThread();
+			}
+		}
+	}
+
+	bool FVideoSourceGroup::FFrameThread::Init()
+	{
+		return true;
+	}
+
+	uint32 FVideoSourceGroup::FFrameThread::Run()
+	{
+		bIsRunning = true;
+
+		while (bIsRunning)
+		{
+			uint64 PreTickCycles = FPlatformTime::Cycles64();
+
+			PushFrame();
+
+			const double DeltaMs = FPlatformTime::ToMilliseconds64(LastTickCycles - PreTickCycles);
+			const double FrameDeltaMs = 1000.0 / TickGroup->FramesPerSecond;
+			const double SleepMs = FrameDeltaMs - DeltaMs;
+			PreTickCycles = LastTickCycles;
+
+			// Sleep as long as we need for a constant FPS
+			if (SleepMs > 0)
+			{
+				FPlatformProcess::Sleep(static_cast<float>(SleepMs / 1000.0));
+			}
+		}
+
+		return 0;
+	}
+
+	void FVideoSourceGroup::FFrameThread::Stop()
+	{
+		bIsRunning = false;
+	}
+
+	void FVideoSourceGroup::FFrameThread::Exit()
+	{
+		bIsRunning = false;
+	}
+
+	void FVideoSourceGroup::FFrameThread::Tick()
+	{
+		const uint64 NowCycles = FPlatformTime::Cycles64();
+		const double DeltaMs = FPlatformTime::ToMilliseconds64(NowCycles - LastTickCycles);
+		const double FrameDeltaMs = 1000.0 / TickGroup->FramesPerSecond;
+		if (DeltaMs >= FrameDeltaMs)
+		{
+			PushFrame();
+		}
+	}
+
+	void FVideoSourceGroup::FFrameThread::PushFrame()
+	{
+		TickGroup->Tick();
+		LastTickCycles = FPlatformTime::Cycles64();
+	}
+} // namespace UE::PixelStreaming
diff --git a/Engine/Plugins/Media/PixelStreaming/Source/PixelStreaming/Private/VideoSourceP2P.cpp b/Engine/Plugins/Media/PixelStreaming/Source/PixelStreaming/Private/VideoSourceP2P.cpp
new file mode 100644
--- /dev/null
+++ b/Engine/Plugins/Media/PixelStreaming/Source/PixelStreaming/Private/VideoSourceP2P.cpp
@@ -0,0 +1,76 @@
+// Copyright Epic Games, Inc. All Rights Reserved.
+
+#include "VideoSourceP2P.h"
+#include "Settings.h"
+#include "FrameAdapterH264.h"
+#include "FrameAdapterI420.h"
+#include "FrameBufferH264.h"
+#include "FrameBufferI420.h"
+#include "FrameBufferInitialize.h"
+
+namespace UE::PixelStreaming
+{
+	FVideoSourceP2P::FVideoSourceP2P(TSharedPtr<FPixelStreamingVideoInput> InVideoInput, TFunction<bool()> InIsQualityControllerFunc)
+		: IsQualityControllerFunc(InIsQualityControllerFunc)
+	{
+		const float ScaleFactor = Settings::CVarPixelStreamingFrameScale.GetValueOnAnyThread();
+		switch (Settings::GetSelectedCodec())
+		{
+			case Settings::ECodec::VP8:
+			case Settings::ECodec::VP9:
+				if (Settings::CVarPixelStreamingVPXUseCompute.GetValueOnAnyThread())
+				{
+					FrameAdapter = MakeShared<FFrameAdapterI420Compute>(InVideoInput, TArray<float>{ ScaleFactor });
+				}
+				else
+				{
+					FrameAdapter = MakeShared<FFrameAdapterI420CPU>(InVideoInput, TArray<float>{ ScaleFactor });
+				}
+				break;
+			case Settings::ECodec::H264:
+			default:
+				FrameAdapter = MakeShared<FFrameAdapterH264>(InVideoInput, TArray<float>{ ScaleFactor });
+				break;
+		}
+	}
+
+	bool FVideoSourceP2P::IsReady() const
+	{
+		return FrameAdapter->IsReady();
+	}
+
+	webrtc::VideoFrame FVideoSourceP2P::CreateFrame(int32 FrameId)
+	{
+		const int64 TimestampUs = rtc::TimeMicros();
+
+		// Based on quality control we either send a frame buffer with a texture or an empty frame
+		// buffer to indicate that this frame should be skipped by the encoder as this peer is having its
+		// frames transmitted by some other peer (the quality controlling peer).
+
+		rtc::scoped_refptr<webrtc::VideoFrameBuffer> FrameBuffer;
+		if (Settings::IsCodecVPX())
+		{
+			FrameBuffer = new rtc::RefCountedObject<FFrameBufferI420>(FrameAdapter, FrameAdapter->GetNumLayers() - 1);
+		}
+		else
+		{
+			if (IsQualityControllerFunc())
+			{
+				FrameBuffer = new rtc::RefCountedObject<FFrameBufferH264>(FrameAdapter);
+			}
+			else
+			{
+				FrameBuffer = new rtc::RefCountedObject<FFrameBufferInitialize>(FrameAdapter);
+			}
+		}
+
+		webrtc::VideoFrame Frame = webrtc::VideoFrame::Builder()
+									   .set_video_frame_buffer(FrameBuffer)
+									   .set_timestamp_us(TimestampUs)
+									   .set_rotation(webrtc::VideoRotation::kVideoRotation_0)
+									   .set_id(FrameId)
+									   .build();
+
+		return Frame;
+	}
+} // namespace UE::PixelStreaming
diff --git a/Engine/Plugins/Media/PixelStreaming/Source/PixelStreaming/Private/VideoSourceP2P.h b/Engine/Plugins/Media/PixelStreaming/Source/PixelStreaming/Private/VideoSourceP2P.h
new file mode 100644
--- /dev/null
+++ b/Engine/Plugins/Media/PixelStreaming/Source/PixelStreaming/Private/VideoSourceP2P.h
@@ -0,0 +1,29 @@
+// Copyright Epic Games, Inc. All Rights Reserved.
+
+#pragma once
+
+#include "VideoSourceBase.h"
+
+class FPixelStreamingFrameAdapter;
+class FPixelStreamingVideoInput;
+
+namespace UE::PixelStreaming
+{
+	/*
+	 * A video source for P2P peers.
+	 */
+	class FVideoSourceP2P : public FVideoSourceBase
+	{
+	public:
+		FVideoSourceP2P(TSharedPtr<FPixelStreamingVideoInput> InVideoInput, TFunction<bool()> InIsQualityControllerFunc);
+
+		virtual bool IsReady() const override;
+
+	protected:
+		TSharedPtr<FPixelStreamingFrameAdapter> FrameAdapter;
+		TFunction<bool()> IsQualityControllerFunc;
+
+	protected:
+		virtual webrtc::VideoFrame CreateFrame(int32 FrameId) override;
+	};
+} // namespace UE::PixelStreaming
diff --git a/Engine/Plugins/Media/PixelStreaming/Source/PixelStreaming/Private/VideoSourceSFU.cpp b/Engine/Plugins/Media/PixelStreaming/Source/PixelStreaming/Private/VideoSourceSFU.cpp
new file mode 100644
--- /dev/null
+++ b/Engine/Plugins/Media/PixelStreaming/Source/PixelStreaming/Private/VideoSourceSFU.cpp
@@ -0,0 +1,71 @@
+// Copyright Epic Games, Inc. All Rights Reserved.
+
+#include "VideoSourceSFU.h"
+#include "Settings.h"
+#include "FrameAdapterH264.h"
+#include "FrameAdapterI420.h"
+#include "FrameBufferH264.h"
+#include "FrameBufferI420.h"
+
+namespace UE::PixelStreaming
+{
+	FVideoSourceSFU::FVideoSourceSFU(TSharedPtr<FPixelStreamingVideoInput> InVideoInput)
+	{
+		TArray<float> LayerScaling;
+		for (auto& Layer : Settings::SimulcastParameters.Layers)
+		{
+			LayerScaling.Add(1.0f / Layer.Scaling);
+		}
+		LayerScaling.Sort([](float ScaleA, float ScaleB) { return ScaleA < ScaleB; });
+
+		switch (Settings::GetSelectedCodec())
+		{
+			case Settings::ECodec::VP8:
+			case Settings::ECodec::VP9:
+				if (Settings::CVarPixelStreamingVPXUseCompute.GetValueOnAnyThread())
+				{
+					FrameAdapter = MakeShared<FFrameAdapterI420Compute>(InVideoInput, LayerScaling);
+				}
+				else
+				{
+					FrameAdapter = MakeShared<FFrameAdapterI420CPU>(InVideoInput, LayerScaling);
+				}
+				break;
+			case Settings::ECodec::H264:
+			default:
+				FrameAdapter = MakeShared<FFrameAdapterH264>(InVideoInput, LayerScaling);
+				break;
+		}
+	}
+
+	bool FVideoSourceSFU::IsReady() const
+	{
+		return FrameAdapter->IsReady();
+	}
+
+	webrtc::VideoFrame FVideoSourceSFU::CreateFrame(int32 FrameId)
+	{
+		const int64 TimestampUs = rtc::TimeMicros();
+
+		rtc::scoped_refptr<webrtc::VideoFrameBuffer> FrameBuffer;
+		if (Settings::IsCodecVPX())
+		{
+			// VPX Encoding. We only ever have a single layer
+			// TODO if we change to use the simulcast adapter we want the simulcast frame buffer here
+			FrameBuffer = new rtc::RefCountedObject<FFrameBufferI420Simulcast>(FrameAdapter);
+		}
+		else
+		{
+			FrameBuffer = new rtc::RefCountedObject<FFrameBufferH264Simulcast>(FrameAdapter);
+		}
+
+		webrtc::VideoFrame Frame = webrtc::VideoFrame::Builder()
+									   .set_video_frame_buffer(FrameBuffer)
+									   .set_timestamp_us(TimestampUs)
+									   .set_rotation(webrtc::VideoRotation::kVideoRotation_0)
+									   .set_id(FrameId)
+									   .build();
+
+		return Frame;
+	}
+} // namespace UE::PixelStreaming
diff --git a/Engine/Plugins/Media/PixelStreaming/Source/PixelStreaming/Private/VideoSourceSFU.h b/Engine/Plugins/Media/PixelStreaming/Source/PixelStreaming/Private/VideoSourceSFU.h
new file mode 100644
--- /dev/null
+++ b/Engine/Plugins/Media/PixelStreaming/Source/PixelStreaming/Private/VideoSourceSFU.h
@@ -0,0 +1,28 @@
+// Copyright Epic Games, Inc. All Rights Reserved.
+
+#pragma once
+
+#include "VideoSourceBase.h"
+
+class FPixelStreamingFrameAdapter;
+class FPixelStreamingVideoInput;
+
+namespace UE::PixelStreaming
+{
+	/*
+	 * A video source for the SFU.
+	 */
+	class FVideoSourceSFU : public FVideoSourceBase
+	{
+	public:
+		FVideoSourceSFU(TSharedPtr<FPixelStreamingVideoInput> InVideoInput);
+
+		virtual bool IsReady() const override;
+
+	protected:
+		TSharedPtr<FPixelStreamingFrameAdapter> FrameAdapter;
+
+	protected:
+		virtual webrtc::VideoFrame CreateFrame(int32 FrameId) override;
+	};
+} // namespace UE::PixelStreaming
diff --git a/Engine/Plugins/Media/PixelStreaming/Source/PixelStreaming/Private/WebRTCIncludes.h b/Engine/Plugins/Media/PixelStreaming/Source/PixelStreaming/Private/WebRTCIncludes.h
--- a/Engine/Plugins/Media/PixelStreaming/Source/PixelStreaming/Private/WebRTCIncludes.h
+++ b/Engine/Plugins/Media/PixelStreaming/Source/PixelStreaming/Private/WebRTCIncludes.h
@@ -47,7 +47,10 @@ THIRD_PARTY_INCLUDES_START
 #include "rtc_base/arraysize.h"
 #include "rtc_base/net_helpers.h"
 #include "rtc_base/string_utils.h"
+#if WEBRTC_VERSION == 84
 #include "rtc_base/signal_thread.h"
+#endif 
+#include "rtc_base/physical_socket_server.h"
 #include "rtc_base/experiments/rate_control_settings.h"
 
 #include "pc/session_description.h"
@@ -55,7 +58,11 @@ THIRD_PARTY_INCLUDES_START
 
 #include "media/engine/internal_decoder_factory.h"
 #include "media/engine/internal_encoder_factory.h"
+#if WEBRTC_VERSION == 84
 #include "media/base/h264_profile_level_id.h"
+#elif WEBRTC_VERSION == 96
+#include "api/video_codecs/h264_profile_level_id.h"
+#endif 
 #include "media/base/adapted_video_track_source.h"
 #include "media/base/media_channel.h"
 #include "media/base/video_common.h"
@@ -65,8 +72,12 @@ THIRD_PARTY_INCLUDES_START
 #include "modules/audio_device/audio_device_buffer.h"
 #include "modules/audio_processing/include/audio_processing.h"
 #include "modules/video_coding/codecs/h264/include/h264.h"
+#if WEBRTC_VERSION ==  84
 #include "modules/video_coding/utility/framerate_controller.h"
+#endif
 #include "modules/video_coding/utility/simulcast_rate_allocator.h"
+#include "modules/video_coding/codecs/vp8/include/vp8.h"
+#include "modules/video_coding/codecs/vp9/include/vp9.h"
 
 #include "common_video/h264/h264_bitstream_parser.h"
 #include "common_video/h264/h264_common.h"
diff --git a/Engine/Plugins/Media/PixelStreaming/Source/PixelStreaming/Public/IPixelStreamingAudioPlayoutRequester.h b/Engine/Plugins/Media/PixelStreaming/Source/PixelStreaming/Public/IPixelStreamingAudioPlayoutRequester.h
new file mode 100644
--- /dev/null
+++ b/Engine/Plugins/Media/PixelStreaming/Source/PixelStreaming/Public/IPixelStreamingAudioPlayoutRequester.h
@@ -0,0 +1,19 @@
+// Copyright Epic Games, Inc. All Rights Reserved.
+
+#pragma once
+
+#include "WebRTCIncludes.h"
+
+class IPixelStreamingAudioPlayoutRequester
+{
+public:
+	virtual ~IPixelStreamingAudioPlayoutRequester() = default;
+
+	virtual void InitPlayout() = 0;
+	virtual void StartPlayout() = 0;
+	virtual void StopPlayout() = 0;
+	virtual bool Playing() const = 0;
+	virtual bool PlayoutIsInitialized() const = 0;
+	virtual void Uninitialise() = 0;
+	virtual void RegisterAudioCallback(webrtc::AudioTransport* AudioCallback) = 0;
+};
diff --git a/Engine/Plugins/Media/PixelStreaming/Source/PixelStreaming/Public/IPixelStreamingAudioSubmixCapturer.h b/Engine/Plugins/Media/PixelStreaming/Source/PixelStreaming/Public/IPixelStreamingAudioSubmixCapturer.h
new file mode 100644
--- /dev/null
+++ b/Engine/Plugins/Media/PixelStreaming/Source/PixelStreaming/Public/IPixelStreamingAudioSubmixCapturer.h
@@ -0,0 +1,21 @@
+// Copyright Epic Games, Inc. All Rights Reserved.
+
+#pragma once
+
+#include "WebRTCIncludes.h"
+
+class IPixelStreamingAudioSubmixCapturer
+{
+public:
+	virtual ~IPixelStreamingAudioSubmixCapturer() = default;
+
+	virtual bool Init() = 0;
+	virtual bool IsInitialised() const = 0;
+	virtual bool IsCapturing() const = 0;
+	virtual void Uninitialise() = 0;
+	virtual bool StartCapturing() = 0;
+	virtual bool EndCapturing() = 0;
+	virtual uint32_t GetVolume() const = 0;
+	virtual void SetVolume(uint32_t NewVolume) = 0;
+	virtual void RegisterAudioCallback(webrtc::AudioTransport* AudioCb) = 0;
+};
diff --git a/Engine/Plugins/Media/PixelStreaming/Source/PixelStreaming/Public/IPixelStreamingFrameBuffer.h b/Engine/Plugins/Media/PixelStreaming/Source/PixelStreaming/Public/IPixelStreamingFrameBuffer.h
new file mode 100644
--- /dev/null
+++ b/Engine/Plugins/Media/PixelStreaming/Source/PixelStreaming/Public/IPixelStreamingFrameBuffer.h
@@ -0,0 +1,45 @@
+// Copyright Epic Games, Inc. All Rights Reserved.
+
+#pragma once
+
+#include "api/video/video_frame_buffer.h"
+#include "CoreMinimal.h"
+
+enum EPixelStreamingFrameBufferType
+{
+	Initialize,
+	Simulcast,
+	Layer
+};
+
+class IPixelStreamingFrameSource;
+
+/*
+ * The base framebuffer that extends the WebRTC type.
+ */
+class PIXELSTREAMING_API IPixelStreamingFrameBuffer : public webrtc::VideoFrameBuffer
+{
+public:
+	virtual ~IPixelStreamingFrameBuffer() = default;
+
+	virtual EPixelStreamingFrameBufferType GetFrameBufferType() const = 0;
+
+	// Begin webrtc::VideoFrameBuffer interface
+	virtual webrtc::VideoFrameBuffer::Type type() const override { return webrtc::VideoFrameBuffer::Type::kNative; }
+
+	virtual int width() const = 0;
+	virtual int height() const = 0;
+
+	virtual rtc::scoped_refptr<webrtc::I420BufferInterface> ToI420() override
+	{
+		unimplemented();
+		return nullptr;
+	}
+
+	virtual const webrtc::I420BufferInterface* GetI420() const override
+	{
+		unimplemented();
+		return nullptr;
+	}
+	// End webrtc::VideoFrameBuffer interface
+};
diff --git a/Engine/Plugins/Media/PixelStreaming/Source/PixelStreaming/Public/IPixelStreamingFrameSource.h b/Engine/Plugins/Media/PixelStreaming/Source/PixelStreaming/Public/IPixelStreamingFrameSource.h
new file mode 100644
--- /dev/null
+++ b/Engine/Plugins/Media/PixelStreaming/Source/PixelStreaming/Public/IPixelStreamingFrameSource.h
@@ -0,0 +1,36 @@
+// Copyright Epic Games, Inc. All Rights Reserved.
+
+#pragma once
+
+/*
+ * Frame buffers will contain a frame source. This could be a static buffer
+ * of data for the frame or an object which pulls the data from elsewhere.
+ * Implement this interface and create your frame buffer object by implementing
+ * IPixelStreamingFrameBuffer.
+ * Could contain multiple layers.
+ */
+class IPixelStreamingFrameSource
+{
+public:
+	virtual ~IPixelStreamingFrameSource() = default;
+
+	/*
+	 * Returns true when the adapter has frames ready to read.
+	 */
+	virtual bool IsReady() const = 0;
+
+	/*
+	 * Gets the number of layers in the adapted output.
+	 */
+	virtual int32 GetNumLayers() const = 0;
+
+	/*
+	 * Gets the output frame width of the given index.
+	 */
+	virtual int32 GetWidth(int LayerIndex) const = 0;
+
+	/*
+	 * Gets the output frame height of the given index.
+	 */
+	virtual int32 GetHeight(int LayerIndex) const = 0;
+};
diff --git a/Engine/Plugins/Media/PixelStreaming/Source/PixelStreaming/Public/IPixelStreamingInputDevice.h b/Engine/Plugins/Media/PixelStreaming/Source/PixelStreaming/Public/IPixelStreamingInputDevice.h
new file mode 100644
--- /dev/null
+++ b/Engine/Plugins/Media/PixelStreaming/Source/PixelStreaming/Public/IPixelStreamingInputDevice.h
@@ -0,0 +1,44 @@
+// Copyright Epic Games, Inc. All Rights Reserved.
+
+#pragma once
+
+#include "GenericPlatform/GenericPlatformMisc.h"
+#include "IInputDevice.h"
+#include "Slate/SceneViewport.h"
+
+// WEBRTC API
+#include "api/data_channel_interface.h"
+//
+
+
+/* 
+ * Use this interface to override the way input is processed from the WebRTC datachannels.
+ * This can be done by passing a FCreateInputDeviceFunc that returns your implementation into
+ * IPixelStreamingStreamer::RegisterCreateInputDevice.
+ * 
+ * See /PixelStreaming/Private/InputDevice.h for the default implimentation of an IPixelStreamingInputDevice.
+*/ 
+class IPixelStreamingInputDevice : public IInputDevice
+{
+public:
+	/*
+	 * Handle the message from the WebRTC data channel.
+	*/
+	virtual void OnMessage(const webrtc::DataBuffer& Buffer) = 0;
+	
+	/*
+	 * Set the viewport this input device is associated with.
+	*/
+	virtual void SetTargetViewport(TSharedPtr<FSceneViewport> InTargetViewport) = 0;
+	
+	/*
+	 * Set whether the input devices is faking touch events using keyboard and mouse
+	 * this can be useful for debugging.
+	*/
+	virtual bool IsFakingTouchEvents() const = 0;
+
+	/*
+	 * Convinient typedef for a TFunction that Creates an Input device and is passed into the IPixelStreamingStreamer.
+	*/
+	typedef TFunction<TSharedPtr<IPixelStreamingInputDevice>(const TSharedRef<FGenericApplicationMessageHandler>& InMessageHandler)> FCreateInputDeviceFunc;
+};
diff --git a/Engine/Plugins/Media/PixelStreaming/Source/PixelStreaming/Public/IPixelStreamingModule.h b/Engine/Plugins/Media/PixelStreaming/Source/PixelStreaming/Public/IPixelStreamingModule.h
--- a/Engine/Plugins/Media/PixelStreaming/Source/PixelStreaming/Public/IPixelStreamingModule.h
+++ b/Engine/Plugins/Media/PixelStreaming/Source/PixelStreaming/Public/IPixelStreamingModule.h
@@ -5,28 +5,25 @@
 #include "CoreMinimal.h"
 #include "Modules/ModuleInterface.h"
 #include "Modules/ModuleManager.h"
-#include "IInputDeviceModule.h"
 #include "Templates/SharedPointer.h"
 #include "Templates/RefCounting.h"
 #include "PixelStreamingPlayerId.h"
 #include "IPixelStreamingAudioSink.h"
-#include "PixelStreamingPumpable.h"
-#include "IPixelStreamingTextureSourceFactory.h"
-#include "IInputDevice.h"
+#include "IPixelStreamingInputDevice.h"
 #include "api/media_stream_interface.h"
 #include "api/scoped_refptr.h"
 #include "api/video_codecs/video_encoder_factory.h"
 #include "api/video/video_frame.h"
 #include "api/video/video_sink_interface.h"
 
-class UTexture2D;
 class UPixelStreamingInput;
-class FRHIGPUFence;
+class IPixelStreamingStreamer;
+class IPixelStreamingClient;
 
 /**
  * The public interface of the Pixel Streaming module.
  */
-class PIXELSTREAMING_API IPixelStreamingModule : public IInputDeviceModule
+class PIXELSTREAMING_API IPixelStreamingModule : public IModuleInterface
 {
 public:
 	/**
@@ -68,99 +65,49 @@ public:
 	virtual bool IsReady() = 0;
 
 	/**
-	 * Connect to the specified signalling server and begin a streaming session
-	 * @param SignallingServerUrl The url of the signalling server to use. eg. ws://localhost:8888
+	 * Starts streaming on all streamers.
 	 */
-	virtual bool StartStreaming(const FString& SignallingServerUrl) = 0;
+	virtual bool StartStreaming() = 0;
 
 	/*
-	 * Ends the current streaming session
+	 * Stops all streamers from streaming.
 	 */
 	virtual void StopStreaming() = 0;
 
 	/*
-	 * Event fired when the streamer has connected to a signalling server and is ready for peers.
+	 * Creates a new streamer.
+	 * @param StreamerId - The ID of the Streamer to be created.
 	 */
-	DECLARE_EVENT_OneParam(IPixelStreamingModule, FStreamingStartedEvent, IPixelStreamingModule&);
+	virtual TSharedPtr<IPixelStreamingStreamer> CreateStreamer(const FString& StreamerId) = 0;
 
-	/*
-	 * A getter for the OnStreamingStarted event. Intent is for users to call IPixelStreamingModule::Get().OnStreamingStarted().AddXXX.
-	 * @return The bindable OnStreamingStarted event.
+	/* 
+	 * Returns a TArray containing the keys to the currently held streamers.
+	 * @return TArray containing the keys to the currently held streamers.
 	 */
-	virtual FStreamingStartedEvent& OnStreamingStarted() = 0;
+	virtual TArray<FString> GetStreamerIds() = 0;
 
 	/*
-	 * Event fired when the streamer has disconnected from a signalling server and has stopped streaming.
+	 * Get a streamer by an ID.
+	 * @return A pointer to the interface for a streamer. nullptr if the streamer isn't found
 	 */
-	DECLARE_EVENT_OneParam(IPixelStreamingModule, FStreamingStoppedEvent, IPixelStreamingModule&);
+	virtual TSharedPtr<IPixelStreamingStreamer> GetStreamer(const FString& StreamerId) = 0;
 
 	/*
-	 * A getter for the OnStreamingStopped event. Intent is for users to call IPixelStreamingModule::Get().OnStreamingStopped().AddXXX.
-	 * @return The bindable OnStreamingStopped event.
-	 */
-	virtual FStreamingStoppedEvent& OnStreamingStopped() = 0;
-
-	/**
-	 * Returns a reference to the input device. The lifetime of this reference
-	 * is that of the underlying shared pointer.
-	 * @return A reference to the input device.
-	 */
-	virtual IInputDevice& GetInputDevice() = 0;
-
-	/**
-	 * Add any player config JSON to the given object which relates to
-	 * configuring the input system for the pixel streaming on the browser.
-	 * @param JsonObject - The JSON object to add fields to.
-	 */
-	virtual void AddPlayerConfig(TSharedRef<class FJsonObject>& JsonObject) = 0;
-
-	/**
-	 * Send a data response back to the browser where we are sending video. This
-	 * could be used as a response to a UI interaction, for example.
-	 * @param Descriptor - A generic descriptor string.
-	 */
-	virtual void SendResponse(const FString& Descriptor) = 0;
-
-	/**
-	 * Send a data command back to the browser where we are sending video. This
-	 * is different to a response as a command is low-level and coming from UE4
-	 * rather than the pixel streamed application.
-	 * @param Descriptor - A generic descriptor string.
-	 */
-	virtual void SendCommand(const FString& Descriptor) = 0;
-
-	/**
-	 * Freeze Pixel Streaming.
-	 * @param Texture - The freeze frame to display. If null then the back buffer is captured.
-	 */
-	virtual void FreezeFrame(UTexture2D* Texture) = 0;
-
-	/**
-	 * Unfreeze Pixel Streaming.
-	 */
-	virtual void UnfreezeFrame() = 0;
-
-	/**
-	 * Send a file to the browser where we are sending video.
-	 * @param FilePath - The freeze frame to display. If null then the back buffer is captured.
+	 * Remove a streamer by an ID
+	 * @param StreamerId	-	The ID of the streamer to be removed.
+	 * @return The removed streamer. nullptr if the streamer wasn't found.
 	 */
-	virtual void SendFileData(TArray<uint8>& ByteData, FString& MimeType, FString& FileExtension) = 0;
+	virtual TSharedPtr<IPixelStreamingStreamer> DeleteStreamer(const FString& StreamerId) = 0;
 
-	/**
-	 * Kick a player by player id.
-	 * @param PlayerId - The ID of the player to kick
-	 */
-	virtual void KickPlayer(FPixelStreamingPlayerId PlayerId) = 0;
-
-	/**
-	 * Get the audio sink associated with a specific peer/player.
+	/*
+	 * Allows the creation of Video Tracks that are fed the backbuffer
 	 */
-	virtual IPixelStreamingAudioSink* GetPeerAudioSink(FPixelStreamingPlayerId PlayerId) = 0;
+	virtual rtc::scoped_refptr<webrtc::VideoTrackSourceInterface> CreateExternalVideoSource() = 0;
 
-	/**
-	 * Get an audio sink that has no peers/players listening to it.
+	/*
+	 * Used to clean up sources no longer needed
 	 */
-	virtual IPixelStreamingAudioSink* GetUnlistenedAudioSink() = 0;
+	virtual void ReleaseExternalVideoSource(const webrtc::VideoTrackSourceInterface* InVideoSource) = 0;
 
 	/**
 	 * Tell the input device about a new pixel streaming input component.
@@ -169,8 +116,7 @@ public:
 	virtual void AddInputComponent(UPixelStreamingInput* InInputComponent) = 0;
 
 	/**
-	 * Tell the input device that a pixel streaming input component is no longer
-	 * relevant.
+	 * Tell the input device that a pixel streaming input component is no longer relevant.
 	 * @param InInputComponent - The pixel streaming input component which is no longer relevant.
 	 */
 	virtual void RemoveInputComponent(UPixelStreamingInput* InInputComponent) = 0;
@@ -188,50 +134,15 @@ public:
 	virtual webrtc::VideoEncoderFactory* CreateVideoEncoderFactory() = 0;
 
 	/**
-	 * Register a "Pumpable" with the Pixel Streaming module and it will be pumped at the WebRTC max FPS (usually 60).
-	 * @param Pumpable - A thing that will be called at a fixed interval (typically 60fps).
-	 */
-	virtual void RegisterPumpable(rtc::scoped_refptr<FPixelStreamingPumpable> Pumpable) = 0;
-
-	/**
-	 * Unregister a "Pumpable" from the Pixel Streaming module.
-	 * @param Pumpable - The pumpable you no longer wish to pump.
-	 */
-	virtual void UnregisterPumpable(rtc::scoped_refptr<FPixelStreamingPumpable> Pumpable) = 0;
-
-	/**
-	 * Create a task that is continuously polled until it is done and can be short circuited.
-	 * This is useful for in Pixel Streaming for checking fences during texture copies.
-	 * @param Task - The task to execute one time, this will happen on the "Poller" thread.
-	 * @param IsTaskFinished - A check that will be called as often as possible until the task is finished.
-	 * @param bKeepRunning - A flag to short circuit the completion checking (useful if we no longer care because it took too long).
-	 */
-	virtual void AddPollerTask(TFunction<void()> Task, TFunction<bool()> IsTaskFinished, TSharedRef<bool, ESPMode::ThreadSafe> bKeepRunning) = 0;
-
-	/**
-	 * Create a WebRTC video source that is has its lifetime managed outside of Pixel Streaming.
-	 * @return A ref-counted WebRTC video source, when WebRTC is done with it the ref count will decrement by one, if that is the only remaining ref the source will be destroyed.
-	 */
-	virtual rtc::scoped_refptr<webrtc::VideoTrackSourceInterface> CreateExternalVideoSource(FName SourceType = FName(TEXT("Backbuffer"))) = 0;
-
-	/**
-	 * The Pixel Streaming module can have n active texture source types that dictate
-	 * the type of video tracks that are made when a new peer joins.
-	 * @param SourceTypes - An array of registered source types to make tracks from, e.g. "Backbuffer".
-	 */
-	virtual void SetActiveTextureSourceTypes(const TArray<FName>& SourceTypes) = 0;
-
-	/**
-	 * The n active texture source types that dictate
-	 * how many video tracks of each source type are made when a new peer joins.
-	 * @return The array of active registered source types that tracks are made from, e.g. "Backbuffer".
+	 * Get the Default Streamer ID
+	 * @return FString 
 	 */
-	virtual const TArray<FName>& GetActiveTextureSourceTypes() const = 0;
+	virtual FString GetDefaultStreamerID() = 0;
 
 	/**
-	 * Get the texture source factory which you can use to register new texture source types and create
-	 * texture sources
-	 * @return The texture source factory in use by the Pixel Streaming module
+	 * @brief A method for iterating through all of the streamers on the module
+	 * 
+	 * @param Func The lambda to execute with each streamer
 	 */
-	virtual IPixelStreamingTextureSourceFactory& GetTextureSourceFactory() = 0;
+	virtual void ForEachStreamer(const TFunction<void(TSharedPtr<IPixelStreamingStreamer>)>& Func) = 0;
 };
diff --git a/Engine/Plugins/Media/PixelStreaming/Source/PixelStreaming/Public/IPixelStreamingSignallingConnectionObserver.h b/Engine/Plugins/Media/PixelStreaming/Source/PixelStreaming/Public/IPixelStreamingSignallingConnectionObserver.h
new file mode 100644
--- /dev/null
+++ b/Engine/Plugins/Media/PixelStreaming/Source/PixelStreaming/Public/IPixelStreamingSignallingConnectionObserver.h
@@ -0,0 +1,30 @@
+// Copyright Epic Games, Inc. All Rights Reserved.
+
+#pragma once
+
+#include "PixelStreamingPlayerId.h"
+#include "WebRTCIncludes.h"
+
+// callback interface for `FPixelStreamingSignallingConnection`
+class IPixelStreamingSignallingConnectionObserver
+{
+public:
+	virtual ~IPixelStreamingSignallingConnectionObserver() = default;
+	virtual void OnSignallingConnected() = 0;
+	virtual void OnSignallingDisconnected(int32 StatusCode, const FString& Reason, bool bWasClean) = 0;
+	virtual void OnSignallingError(const FString& ErrorMsg) = 0;
+	virtual void OnSignallingConfig(const webrtc::PeerConnectionInterface::RTCConfiguration& Config) = 0;
+
+	// Streamer-only
+	virtual void OnSignallingSessionDescription(FPixelStreamingPlayerId PlayerId, webrtc::SdpType Type, const FString& Sdp) { unimplemented(); }
+	virtual void OnSignallingRemoteIceCandidate(FPixelStreamingPlayerId PlayerId, const FString& SdpMid, int SdpMLineIndex, const FString& Sdp) { unimplemented(); }
+	virtual void OnSignallingPlayerConnected(FPixelStreamingPlayerId PlayerId, int Flags) { unimplemented(); }
+	virtual void OnSignallingPlayerDisconnected(FPixelStreamingPlayerId PlayerId) { unimplemented(); }
+	virtual void OnSignallingSFUPeerDataChannels(FPixelStreamingPlayerId SFUId, FPixelStreamingPlayerId PlayerId, int32 SendStreamId, int32 RecvStreamId) { unimplemented(); }
+	
+	// Player-only
+	virtual void OnSignallingSessionDescription(webrtc::SdpType Type, const FString& Sdp) { unimplemented(); }
+	virtual void OnSignallingRemoteIceCandidate(const FString& SdpMid, int SdpMLineIndex, const FString& Sdp) { unimplemented(); }
+	virtual void OnSignallingPlayerCount(uint32 Count) { unimplemented(); }
+	virtual void OnSignallingPeerDataChannels(int32 SendStreamId, int32 RecvStreamId) { unimplemented(); }
+};
diff --git a/Engine/Plugins/Media/PixelStreaming/Source/PixelStreaming/Public/PixelStreamingAudioComponent.h b/Engine/Plugins/Media/PixelStreaming/Source/PixelStreaming/Public/PixelStreamingAudioComponent.h
--- a/Engine/Plugins/Media/PixelStreaming/Source/PixelStreaming/Public/PixelStreamingAudioComponent.h
+++ b/Engine/Plugins/Media/PixelStreaming/Source/PixelStreaming/Public/PixelStreamingAudioComponent.h
@@ -78,6 +78,13 @@ protected:
 	//~ End UActorComponent interface
 
 public:
+	/**
+	 *   The Pixel Streaming streamer of the player that we wish to listen to.
+	 *   If this is left blank this component will use the default Pixel Streaming streamer
+	 */
+	UPROPERTY(EditAnywhere, BlueprintReadWrite, Category = "Pixel Streaming Audio Component")
+	FString StreamerToHear;
+
 	/**
 	 *   The Pixel Streaming player/peer whose audio we wish to listen to.
 	 *   If this is left blank this component will listen to the first non-listened to peer that connects after this component is ready.
@@ -97,10 +104,14 @@ private:
 	TSharedPtr<FWebRTCSoundGenerator, ESPMode::ThreadSafe> SoundGenerator;
 
 public:
-	// Listen to a specific player. If the player is not found this component will be silent.
+	// Listen to a specific player on the default streamer. If the player is not found this component will be silent.
 	UFUNCTION(BlueprintCallable, Category = "Pixel Streaming Audio Component")
 	bool ListenTo(FString PlayerToListenTo);
 
+	// Listen to a specific player. If the player is not found this component will be silent.
+	UFUNCTION(BlueprintCallable, Category = "Pixel Streaming Audio Component")
+	bool StreamerListenTo(FString StreamerId, FString PlayerToListenTo);
+
 	// True if listening to a connected WebRTC peer through Pixel Streaming.
 	UFUNCTION(BlueprintCallable, Category = "Pixel Streaming Audio Component")
 	bool IsListeningToPlayer();
diff --git a/Engine/Plugins/Media/PixelStreaming/Source/PixelStreaming/Public/PixelStreamingFrameAdapter.h b/Engine/Plugins/Media/PixelStreaming/Source/PixelStreaming/Public/PixelStreamingFrameAdapter.h
new file mode 100644
--- /dev/null
+++ b/Engine/Plugins/Media/PixelStreaming/Source/PixelStreaming/Public/PixelStreamingFrameAdapter.h
@@ -0,0 +1,79 @@
+// Copyright Epic Games, Inc. All Rights Reserved.
+
+#pragma once
+
+#include "RHI.h"
+#include "IPixelStreamingFrameSource.h"
+
+class FPixelStreamingVideoInput;
+class FPixelStreamingFrameAdapterProcess;
+class FPixelStreamingSourceFrame;
+
+/*
+ * Wraps the output of the adapt process. Extend this for your own result types.
+ * You must implement GetWidth and GetHeight to return the width and height of the
+ * frame. Add your own method to extract the adapted data.
+ */
+class IPixelStreamingAdaptedVideoFrameLayer
+{
+public:
+	virtual ~IPixelStreamingAdaptedVideoFrameLayer() = default;
+
+	virtual int32 GetWidth() const = 0;
+	virtual int32 GetHeight() const = 0;
+};
+
+/*
+ * Takes the input frame and adapts it to the input type of the selected encoder.
+ * Extend this and implement CreateAdaptProcess to adapt frames to your encoder.
+ * CreateAdaptProcess should create your own processes that implements FPixelStreamingFrameAdapterProcess.
+ */
+class PIXELSTREAMING_API FPixelStreamingFrameAdapter : public IPixelStreamingFrameSource
+{
+public:
+	FPixelStreamingFrameAdapter(TSharedPtr<FPixelStreamingVideoInput> VideoInput);
+	FPixelStreamingFrameAdapter(TSharedPtr<FPixelStreamingVideoInput> VideoInput, TArray<float> LayerScales);
+	virtual ~FPixelStreamingFrameAdapter();
+
+	/*
+	 * Returns true when the adapter has frames ready to read.
+	 */
+	virtual bool IsReady() const override;
+
+	/*
+	 * Gets the number of layers in the adapted output.
+	 */
+	virtual int32 GetNumLayers() const override { return LayerAdapters.Num(); }
+
+	/*
+	 * Gets the output frame width of the given index.
+	 */
+	virtual int32 GetWidth(int LayerIndex) const override;
+
+	/*
+	 * Gets the output frame height of the given index.
+	 */
+	virtual int32 GetHeight(int LayerIndex) const override;
+
+	/*
+	 * Gets a single frame of output for the given index.
+	 */
+	TSharedPtr<IPixelStreamingAdaptedVideoFrameLayer> ReadOutput(int32 LayerIndex);
+
+protected:
+	/*
+	 * Create the appropriate adapter process for the adapter
+	 */
+	virtual TSharedPtr<FPixelStreamingFrameAdapterProcess> CreateAdaptProcess(float Scale) = 0;
+
+protected:
+	TWeakPtr<FPixelStreamingVideoInput> VideoInputPtr;
+	FDelegateHandle OnFrameDelegateHandle;
+	TArray<float> LayerScales;
+	TArray<TSharedPtr<FPixelStreamingFrameAdapterProcess>> LayerAdapters;
+	mutable FCriticalSection LayersGuard;
+
+private:
+	void AddLayer(float Scale);
+	void OnFrame(const FPixelStreamingSourceFrame& SourceFrame);
+};
diff --git a/Engine/Plugins/Media/PixelStreaming/Source/PixelStreaming/Public/PixelStreamingFrameAdapterProcess.h b/Engine/Plugins/Media/PixelStreaming/Source/PixelStreaming/Public/PixelStreamingFrameAdapterProcess.h
new file mode 100644
--- /dev/null
+++ b/Engine/Plugins/Media/PixelStreaming/Source/PixelStreaming/Public/PixelStreamingFrameAdapterProcess.h
@@ -0,0 +1,108 @@
+// Copyright Epic Games, Inc. All Rights Reserved.
+
+#pragma once
+
+#include "Containers/TripleBuffer.h"
+#include "RHI.h"
+
+class FPixelStreamingSourceFrame;
+class IPixelStreamingAdaptedVideoFrameLayer;
+
+/*
+ * Extend this class to add your own adaption/conversion logic for incoming frames
+ * to adapt the data to the type your selected encoder requires.
+ * Implement Consume with your logic to convert the frame data.
+ * This is implemented as a parallel process. Call OnBegin to mark this process as busy
+ * and call OnComplete once the work is complete.
+ */
+class PIXELSTREAMING_API FPixelStreamingFrameAdapterProcess : public TSharedFromThis<FPixelStreamingFrameAdapterProcess>
+{
+public:
+	FPixelStreamingFrameAdapterProcess() = default;
+	virtual ~FPixelStreamingFrameAdapterProcess() = default;
+
+	/*
+	 * Called when an input frame needs processing.
+	 */
+	void Process(const FPixelStreamingSourceFrame& SourceFrame);
+
+	/*
+	 * Returns true if Initialize() has been called.
+	 * Output data can depend on the incoming frames so we do lazy initialization when we first consume.
+	 */
+	bool IsInitialized() const { return bInitialized; }
+
+	/*
+	 * Returns true when this process is actively working on adapting frame data.
+	 */
+	bool IsBusy() const { return bBusy; }
+
+	/*
+	 * Returns true if this process has a frame in the output buffer ready to be read.
+	 */
+	bool HasOutput() const { return bHasOutput; }
+
+	/*
+	 * Gets the output frame from the output buffer.
+	 */
+	TSharedPtr<IPixelStreamingAdaptedVideoFrameLayer> ReadOutput();
+
+	/*
+	 * Gets the width of the output in pixels.
+	 */
+	int32 GetOutputLayerWidth();
+
+	/*
+	 * Gets the height of the output in pixels.
+	 */
+	int32 GetOutputLayerHeight();
+
+protected:
+	/*
+	 * Initializes the process to be ready for work. Called once at startup and any time the
+	 * source resolution changes.
+	 * SourceWidth and SourceHeight are the pixel dimensions of the frame to work on.
+	 */
+	virtual void Initialize(int32 SourceWidth, int32 SourceHeight);
+
+	/*
+	 * A callback that can be implemented if work needs to be done when the source
+	 * resolution changes.
+	 */
+	virtual void OnSourceResolutionChanged(int32 OldWidth, int32 OldHeight, int32 NewWidth, int32 NewHeight) {}
+
+	/*
+	 * Implement this to create a buffer for the output.
+	 */
+	virtual TSharedPtr<IPixelStreamingAdaptedVideoFrameLayer> CreateOutputBuffer(int32 SourceWidth, int32 SourceHeight) = 0;
+
+	/*
+	 * Gets a buffer ready for writing to for the output.
+	 */
+	TSharedPtr<IPixelStreamingAdaptedVideoFrameLayer> GetWriteBuffer();
+
+	/*
+	 * Implement this with your specific process to adapt the incoming frame.
+	 */
+	virtual void BeginProcess(const FPixelStreamingSourceFrame& SourceFrame) = 0;
+
+	/*
+	 * Call this to mark the end of processing. Will commit the current write buffer into the read buffer.
+	 */
+	void EndProcess();
+
+private:
+	bool bInitialized = false;
+	bool bBusy = false;
+	bool bHasOutput = false;
+
+	int32 ExpectedSourceWidth = 0;
+	int32 ExpectedSourceHeight = 0;
+
+	uint64 LastWriteSwap = 0;
+	uint64 LastReadSwap = 0;
+
+	TTripleBuffer<TSharedPtr<IPixelStreamingAdaptedVideoFrameLayer>> Buffer;
+
+	bool ResolutionChanged(int32 SourceWidth, int32 SourceHeight) const;
+};
diff --git a/Engine/Plugins/Media/PixelStreaming/Source/PixelStreaming/Public/PixelStreamingFrameMetadata.h b/Engine/Plugins/Media/PixelStreaming/Source/PixelStreaming/Public/PixelStreamingFrameMetadata.h
new file mode 100644
--- /dev/null
+++ b/Engine/Plugins/Media/PixelStreaming/Source/PixelStreaming/Public/PixelStreamingFrameMetadata.h
@@ -0,0 +1,30 @@
+// Copyright Epic Games, Inc. All Rights Reserved.
+
+#pragma once
+
+struct PIXELSTREAMING_API FPixelStreamingFrameMetadata
+{
+	// Identifier for the pipeline this frame took
+	FString ProcessName = "Unknown";
+
+	// Frames can be composed of multiple frames
+	int32 Layer = 0;
+
+	// The time this frame was sourced
+	uint64 SourceTime = 0;
+
+	// Adapt process timings (happens once per frame)
+	uint64 AdaptCallTime = 0;
+	uint64 AdaptProcessStartTime = 0;
+	uint64 AdaptProcessFinalizeTime = 0;
+	uint64 AdaptProcessEndTime = 0;
+
+	// Frame use timings (can happen multiple times. ie. we are consuming frames faster than producing them)
+	uint32 UseCount = 0;
+	uint64 FirstEncodeStartTime = 0;
+	uint64 LastEncodeStartTime = 0;
+	uint64 LastEncodeEndTime = 0;
+
+	// wanted this to be explicit with a name
+	FPixelStreamingFrameMetadata Copy() const { return *this; }
+};
diff --git a/Engine/Plugins/Media/PixelStreaming/Source/PixelStreaming/Public/PixelStreamingPeerConnection.h b/Engine/Plugins/Media/PixelStreaming/Source/PixelStreaming/Public/PixelStreamingPeerConnection.h
new file mode 100644
--- /dev/null
+++ b/Engine/Plugins/Media/PixelStreaming/Source/PixelStreaming/Public/PixelStreamingPeerConnection.h
@@ -0,0 +1,66 @@
+// Copyright Epic Games, Inc. All Rights Reserved.
+
+#pragma once
+
+#include "PixelStreamingProtocolDefs.h"
+#include "WebRTCIncludes.h"
+
+class PIXELSTREAMING_API FPixelStreamingPeerConnection : public webrtc::PeerConnectionObserver, public webrtc::DataChannelObserver
+{
+public:
+	using FConfig = webrtc::PeerConnectionInterface::RTCConfiguration;
+	static TUniquePtr<FPixelStreamingPeerConnection> Create(const FConfig& Config);
+	static void Shutdown();
+	~FPixelStreamingPeerConnection() = default;
+
+	void SetSuccessCallback(const TFunction<void(const webrtc::SessionDescriptionInterface*)>& Callback);
+	void SetFailureCallback(const TFunction<void(const FString&)>& Callback);
+	void SetIceCandidateCallback(const TFunction<void(const webrtc::IceCandidateInterface*)>& Callback);
+
+	void SetRemoteDescription(const FString& Sdp);
+	void AddRemoteIceCandidate(const FString& SdpMid, int SdpMLineIndex, const FString& Sdp);
+	void CreateDataChannels(int32 SendStreamId, int32 RecvStreamId);
+
+	void SetVideoSink(rtc::VideoSinkInterface<webrtc::VideoFrame>* Sink);
+
+	bool SendMessage(UE::PixelStreaming::Protocol::EToStreamerMsg Type, const FString& Descriptor) const;
+
+protected:
+	//
+	// webrtc::PeerConnectionObserver implementation.
+	//
+	virtual void OnSignalingChange(webrtc::PeerConnectionInterface::SignalingState NewState) override;
+	virtual void OnAddStream(rtc::scoped_refptr<webrtc::MediaStreamInterface> Stream) override;
+	virtual void OnRemoveStream(rtc::scoped_refptr<webrtc::MediaStreamInterface> Stream) override;
+	virtual void OnDataChannel(rtc::scoped_refptr<webrtc::DataChannelInterface> Channel) override;
+	virtual void OnRenegotiationNeeded() override;
+	virtual void OnIceConnectionChange(webrtc::PeerConnectionInterface::IceConnectionState NewState) override;
+	virtual void OnIceGatheringChange(webrtc::PeerConnectionInterface::IceGatheringState NewState) override;
+	virtual void OnIceCandidate(const webrtc::IceCandidateInterface* Candidate) override;
+	virtual void OnIceCandidatesRemoved(const std::vector<cricket::Candidate>& candidates) override;
+	virtual void OnIceConnectionReceivingChange(bool Receiving) override;
+	virtual void OnTrack(rtc::scoped_refptr<webrtc::RtpTransceiverInterface> Transceiver) override;
+	virtual void OnRemoveTrack(rtc::scoped_refptr<webrtc::RtpReceiverInterface> receiver) override;
+
+	//
+	// webrtc::DataChannelObserver implementation.
+	//
+	virtual void OnStateChange() override;
+	virtual void OnMessage(const webrtc::DataBuffer& buffer) override;
+
+private:
+	FPixelStreamingPeerConnection() = default;
+
+	TFunction<void(const webrtc::SessionDescriptionInterface*)> SuccessCallback;
+	TFunction<void(const FString&)> FailureCallback;
+	TFunction<void(const webrtc::IceCandidateInterface*)> IceCandidateCallback;
+
+	rtc::scoped_refptr<webrtc::PeerConnectionInterface> PeerConnection;
+	rtc::scoped_refptr<webrtc::DataChannelInterface> DataChannel;
+
+	rtc::VideoSinkInterface<webrtc::VideoFrame>* VideoSink;
+
+	static void CreatePeerConnectionFactory();
+	static TUniquePtr<rtc::Thread> SignallingThread;
+	static rtc::scoped_refptr<webrtc::PeerConnectionFactoryInterface> PeerConnectionFactory;
+};
diff --git a/Engine/Plugins/Media/PixelStreaming/Source/PixelStreaming/Public/PixelStreamingProtocolDefs.h b/Engine/Plugins/Media/PixelStreaming/Source/PixelStreaming/Public/PixelStreamingProtocolDefs.h
--- a/Engine/Plugins/Media/PixelStreaming/Source/PixelStreaming/Public/PixelStreamingProtocolDefs.h
+++ b/Engine/Plugins/Media/PixelStreaming/Source/PixelStreaming/Public/PixelStreamingProtocolDefs.h
@@ -2,6 +2,8 @@
 
 #pragma once
 
+#include "HAL/Platform.h"
+
 namespace UE::PixelStreaming::Protocol
 {
 	enum EPlayerFlags
@@ -11,7 +13,7 @@ namespace UE::PixelStreaming::Protocol
 		PSPFlag_IsSFU = 0x02,
 	};
 
-	enum class EToStreamerMsg : uint8_t
+	enum class EToStreamerMsg : uint8
 	{
 		/**********************************************************************/
 
@@ -32,7 +34,7 @@ namespace UE::PixelStreaming::Protocol
 		/**********************************************************************/
 
 		/*
-		 * Input Messages. Range = 50..89.
+		 * Input Messages. Range = 50+. (DO NOT ADD ANY NEW INPUT MESSAGE <50 ).
 		 */
 
 		// Generic Input Messages. Range = 50..59.
@@ -73,7 +75,7 @@ namespace UE::PixelStreaming::Protocol
 	};
 
 	//! Messages that can be sent to the webrtc players
-	enum class EToPlayerMsg : uint8_t
+	enum class EToPlayerMsg : uint8
 	{
 		QualityControlOwnership = 0,
 		Response = 1,
@@ -86,7 +88,8 @@ namespace UE::PixelStreaming::Protocol
 		FileExtension = 8,
 		FileMimeType = 9,
 		FileContents = 10,
-		TestEcho = 11
+		TestEcho = 11,
+		InputControlOwnership = 12
 	};
 
 }; // namespace UE::PixelStreaming::Protocol
diff --git a/Engine/Plugins/Media/PixelStreaming/Source/PixelStreaming/Public/PixelStreamingServers.h b/Engine/Plugins/Media/PixelStreaming/Source/PixelStreaming/Public/PixelStreamingServers.h
new file mode 100644
--- /dev/null
+++ b/Engine/Plugins/Media/PixelStreaming/Source/PixelStreaming/Public/PixelStreamingServers.h
@@ -0,0 +1,158 @@
+// Copyright Epic Games, Inc. All Rights Reserved.
+
+#pragma once
+
+#include "Misc/MonitoredProcess.h"
+#include "Templates/SharedPointer.h"
+#include "GenericPlatform/GenericPlatformMisc.h"
+#include "Containers/Map.h"
+#include "Tickable.h"
+#include "HAL/ThreadSafeBool.h"
+
+namespace UE::PixelStreaming
+{
+	/*
+	 * Utility for launching Pixel Streaming servers such as Cirrus, SFU, and Matchmaker.
+	 * This utility automatically searchs in known locations for these servers and if found, launchs them as child processes.
+	 */
+	namespace Servers
+	{
+
+		/**
+		 * Configuration to control behaviour when launching any of the Pixel Streaming servers.
+		 **/
+		struct PIXELSTREAMING_API FLaunchArgs
+		{
+			// If true, will copy to a temp directory and when finished delete itself. If false, will launch from Samples/PixelStreaming/WebServers/*.
+			bool bEphemeral = true;
+
+			// Arguments passed to the actual server when its process is started.
+			FString ProcessArgs = TEXT("");
+
+			// Reconnection timeout in seconds
+			float ReconnectionTimeoutSeconds = 30.0f;
+
+			// Reconnect interval in seconds.
+			float ReconnectionIntervalSeconds = 2.0f;
+		};
+
+		/**
+		 * Endpoints for the various Pixel Streaming servers.
+		 **/
+		enum class PIXELSTREAMING_API EEndpoint
+		{
+			// The websocket signalling url between the server and the UE streamer - e.g. ws://localhost:8888
+			Signalling_Streamer,
+
+			// The websocket signalling url between the server and the players (aka. web browsers) - e.g. ws://localhost:80
+			Signalling_Players,
+
+			// The websocket signalling url between the server and the matchmaker server - e.g. ws://localhost:9999
+			Signalling_Matchmaker,
+
+			// The websocket signalling url between the server and the SFU server - e.g. ws://localhost:8889
+			Signalling_SFU,
+
+			// The http url for the webserver hosted within the signalling server - e.g. http://localhost
+			Signalling_Webserver
+		};
+
+		// ---------------------------------------------------------------------------------------------
+		typedef TMap<EEndpoint, FString> FEndpoints;
+		DECLARE_MULTICAST_DELEGATE_OneParam(FOnReady, const FEndpoints& /* Endpoint urls */);
+		// ---------------------------------------------------------------------------------------------
+
+		/**
+		 *  Base class for all Pixel Streaming servers.
+		 **/
+		class PIXELSTREAMING_API FServerBase : public FTickableGameObject
+		{
+		public:
+			virtual ~FServerBase();
+
+			/* Immediately stops the child process that was running the server. */
+			virtual void Stop();
+
+			/**
+			 * @return	The absolute path to the root directory that the server was launched from.
+			 **/
+			virtual FString GetPathOnDisk();
+
+			/**
+			 * @return	True if the server has been launched. Note: Launched does not necessarily mean it is connectible yet. Bind to OnReady for that.
+			 **/
+			virtual bool HasLaunched();
+
+			/**
+			 * Launch the server in a child process using the supplied launch arguments.
+			 * @param LaunchArgs	The launch arguments to control how the server is launched, including what args to pass to the child process.
+			 **/
+			virtual void Launch(FLaunchArgs& InLaunchArgs);
+
+			/**
+			 * @return	True if the server has been connected to and is ready for new connections.
+			 **/
+			virtual bool IsReady();
+
+			/**
+			 * @return	True if the server has timed out while trying to establish a connection.
+			 **/
+			virtual bool IsTimedOut();
+
+			/* Begin FTickableGameObject */
+			virtual bool IsTickableWhenPaused() const override;
+			virtual bool IsTickableInEditor() const override;
+			virtual void Tick(float DeltaTime) override;
+			virtual bool IsAllowedToTick() const override;
+			TStatId GetStatId() const { RETURN_QUICK_DECLARE_CYCLE_STAT(PixelStreamingServers, STATGROUP_Tickables); }
+			/* End FTickableGameObject */
+
+		protected:
+			/**
+			 * Implemented by derived types. The actual implementation of how this specified server is launched.
+			 * @param LaunchArgs	The launch arguments to use.
+			 * @param OutServerPath	The output absolute path to the root directory containing the server we wish to launch - can change if server is ephemeral.
+			 * @param OutEndpoints	The output endpoints the user can expect to use with this server.
+			 **/
+			virtual TSharedPtr<FMonitoredProcess> LaunchImpl(FLaunchArgs& InLaunchArgs, FString& OutServerPath, TMap<EEndpoint, FString>& OutEndpoints) = 0;
+
+			/**
+			 * Implemented by derived types. Implementation specific but somehow the server has been tested to see if it ready for connections.
+			 * @return	True if the server is able to be connected to.
+			 **/
+			virtual bool TestConnection() = 0;
+
+		public:
+			// Delegate fired when the server is ready for connections, first parameter is a map of all supported endpoints and their urls.
+			FOnReady OnReady;
+
+			DECLARE_MULTICAST_DELEGATE(FOnFailedToReady);
+			/* Can fire when the server is unable to be contacted or connecting to it timed out. */
+			FOnFailedToReady OnFailedToReady;
+
+		protected:
+			FLaunchArgs LaunchArgs;
+			TSharedPtr<FMonitoredProcess> ServerProcess;
+			bool bHasLaunched = false;
+			FString ServerRootAbsPath;
+			TMap<EEndpoint, FString> Endpoints;
+			FThreadSafeBool bIsReady = false;
+			float PollingStartedSeconds = 0.0f;
+			float LastReconnectionTimeSeconds = 0.0f;
+			bool bAllowedToTick = true;
+			bool bTimedOut = false;
+		};
+
+		/* -------------- Static utility methods for working with Pixel Streaming servers. ----------------- */
+
+		/**
+		 * Creates a NodeJS Cirrus signalling server by launching it as a child process.
+		 * Note: Calling this method does not launch the server. You should call Launch() yourself
+		 * once you have bound to appropriate delegates such as OnReady.
+		 * @return	The signalling server wrapped in a server object for simple management.
+		 **/
+		PIXELSTREAMING_API TSharedPtr<FServerBase> MakeSignallingServer();
+
+	} // namespace Servers
+
+} // namespace UE::PixelStreaming
diff --git a/Engine/Plugins/Media/PixelStreaming/Source/PixelStreaming/Public/PixelStreamingSessionDescriptionObservers.h b/Engine/Plugins/Media/PixelStreaming/Source/PixelStreaming/Public/PixelStreamingSessionDescriptionObservers.h
new file mode 100644
--- /dev/null
+++ b/Engine/Plugins/Media/PixelStreaming/Source/PixelStreaming/Public/PixelStreamingSessionDescriptionObservers.h
@@ -0,0 +1,84 @@
+// Copyright Epic Games, Inc. All Rights Reserved.
+
+#pragma once
+
+#include "WebRTCIncludes.h"
+
+#include "Templates/UnrealTemplate.h"
+#include "Templates/Function.h"
+#include "Containers/UnrealString.h"
+
+#include <string>
+#include <memory>
+
+//////////////////////////////////////////////////////////////////////////
+// FPixelStreamingSetSessionDescriptionObserver
+// WebRTC requires an implementation of `webrtc::SetSessionDescriptionObserver` interface as a callback
+// for setting session description, either on receiving remote `offer` (`PeerConnection::SetRemoteDescription`)
+// of on sending `answer` (`PeerConnection::SetLocalDescription`)
+class PIXELSTREAMING_API FPixelStreamingSetSessionDescriptionObserver : public webrtc::SetSessionDescriptionObserver
+{
+public:
+	using FSuccessCallback = TUniqueFunction<void()>;
+	using FFailureCallback = TUniqueFunction<void(const FString&)>;
+
+	static FPixelStreamingSetSessionDescriptionObserver* Create(FSuccessCallback&& successCallback, FFailureCallback&& failureCallback)
+	{
+		return new rtc::RefCountedObject<FPixelStreamingSetSessionDescriptionObserver>(MoveTemp(successCallback), MoveTemp(failureCallback));
+	}
+
+	FPixelStreamingSetSessionDescriptionObserver(FSuccessCallback&& successCallback, FFailureCallback&& failureCallback)
+		: SuccessCallback(MoveTemp(successCallback))
+		, FailureCallback(MoveTemp(failureCallback))
+	{
+	}
+
+	// we don't need to do anything on success
+	void OnSuccess() override
+	{
+		SuccessCallback();
+	}
+
+	// errors usually mean incompatibility between our session configuration (often H.264, its profile and level) and
+	// player, malformed SDP or if player doesn't support UnifiedPlan (whatever was used by proxy)
+	void OnFailure(webrtc::RTCError Error) override
+	{
+		FailureCallback(FString(UTF8_TO_TCHAR(Error.message())));
+	}
+
+private:
+	FSuccessCallback SuccessCallback;
+	FFailureCallback FailureCallback;
+};
+
+class PIXELSTREAMING_API FPixelStreamingCreateSessionDescriptionObserver : public webrtc::CreateSessionDescriptionObserver
+{
+public:
+	using FSuccessCallback = TUniqueFunction<void(webrtc::SessionDescriptionInterface*)>;
+	using FFailureCallback = TUniqueFunction<void(const FString&)>;
+
+	static FPixelStreamingCreateSessionDescriptionObserver* Create(FSuccessCallback&& successCallback, FFailureCallback&& failureCallback)
+	{
+		return new rtc::RefCountedObject<FPixelStreamingCreateSessionDescriptionObserver>(MoveTemp(successCallback), MoveTemp(failureCallback));
+	}
+
+	FPixelStreamingCreateSessionDescriptionObserver(FSuccessCallback&& successCallback, FFailureCallback&& failureCallback)
+		: SuccessCallback(MoveTemp(successCallback))
+		, FailureCallback(MoveTemp(failureCallback))
+	{
+	}
+
+	void OnSuccess(webrtc::SessionDescriptionInterface* SDP) override
+	{
+		SuccessCallback(SDP);
+	}
+
+	void OnFailure(webrtc::RTCError Error) override
+	{
+		FailureCallback(ANSI_TO_TCHAR(Error.message()));
+	}
+
+private:
+	FSuccessCallback SuccessCallback;
+	FFailureCallback FailureCallback;
+};
diff --git a/Engine/Plugins/Media/PixelStreaming/Source/PixelStreaming/Public/PixelStreamingSignallingConnection.h b/Engine/Plugins/Media/PixelStreaming/Source/PixelStreaming/Public/PixelStreamingSignallingConnection.h
new file mode 100644
--- /dev/null
+++ b/Engine/Plugins/Media/PixelStreaming/Source/PixelStreaming/Public/PixelStreamingSignallingConnection.h
@@ -0,0 +1,96 @@
+// Copyright Epic Games, Inc. All Rights Reserved.
+
+#pragma once
+
+#include "WebRTCIncludes.h"
+#include "Dom/JsonObject.h"
+#include "Engine/EngineTypes.h"
+#include "PixelStreamingPlayerId.h"
+#include "IPixelStreamingSignallingConnectionObserver.h"
+
+class IWebSocket;
+
+class PIXELSTREAMING_API FPixelStreamingSignallingConnection final
+{
+public:
+	using FWebSocketFactory = TFunction<TSharedPtr<IWebSocket>(const FString&)>;
+
+	FPixelStreamingSignallingConnection(const FWebSocketFactory& InWebSocketFactory, IPixelStreamingSignallingConnectionObserver& Observer, FString InStreamerId = "");
+	~FPixelStreamingSignallingConnection();
+
+	void Connect(const FString& Url);
+	void Disconnect();
+	bool IsConnected() const;
+
+	// Streamer interface
+	void SendOffer(FPixelStreamingPlayerId PlayerId, const webrtc::SessionDescriptionInterface& SDP);
+	void SendAnswer(FPixelStreamingPlayerId PlayerId, const webrtc::SessionDescriptionInterface& SDP);
+	void SendIceCandidate(const webrtc::IceCandidateInterface& IceCandidate);
+	void SendIceCandidate(FPixelStreamingPlayerId PlayerId, const webrtc::IceCandidateInterface& IceCandidate);
+	void SendDisconnectPlayer(FPixelStreamingPlayerId PlayerId, const FString& Reason);
+
+	// Player interface
+	void SendAnswer(const webrtc::SessionDescriptionInterface& SDP);
+
+private:
+	void KeepAlive();
+
+	void OnConnected();
+	void OnConnectionError(const FString& Error);
+	void OnClosed(int32 StatusCode, const FString& Reason, bool bWasClean);
+	void OnMessage(const FString& Msg);
+
+	using FJsonObjectPtr = TSharedPtr<FJsonObject>;
+	void RegisterHandler(const FString& messageType, const TFunction<void(FJsonObjectPtr)>& handler);
+
+	void OnIdRequested();
+	void OnConfig(const FJsonObjectPtr& Json);
+	void OnSessionDescription(const FJsonObjectPtr& Json);
+	void OnIceCandidate(const FJsonObjectPtr& Json);
+	void OnPlayerCount(const FJsonObjectPtr& Json);
+	void OnPlayerConnected(const FJsonObjectPtr& Json);
+	void OnPlayerDisconnected(const FJsonObjectPtr& Json);
+	void OnSFUPeerDataChannels(const FJsonObjectPtr& Json);
+	void OnPeerDataChannels(const FJsonObjectPtr& Json);
+	void SetPlayerIdJson(FJsonObjectPtr& JsonObject, FPixelStreamingPlayerId PlayerId);
+	bool GetPlayerIdJson(const FJsonObjectPtr& Json, FPixelStreamingPlayerId& OutPlayerId, const FString& FieldId = TEXT("playerId"));
+
+	void StartKeepAliveTimer();
+	void StopKeepAliveTimer();
+
+	void SendMessage(const FString& Msg);
+
+	template <typename FmtType, typename... T>
+	void PlayerError(FPixelStreamingPlayerId PlayerId, const FmtType& Msg, T... args)
+	{
+		const FString FormattedMsg = FString::Printf(Msg, args...);
+		PlayerError(PlayerId, FormattedMsg);
+	}
+	void PlayerError(FPixelStreamingPlayerId PlayerId, const FString& Msg);
+
+	template <typename FmtType, typename... T>
+	void FatalError(const FmtType& Msg, T... args)
+	{
+		const FString FormattedMsg = FString::Printf(Msg, args...);
+		FatalError(FormattedMsg);
+	}
+	void FatalError(const FString& Msg);
+
+private:
+	FWebSocketFactory WebSocketFactory;
+	IPixelStreamingSignallingConnectionObserver& Observer;
+	FString StreamerId;
+
+	TSharedPtr<IWebSocket> WebSocket;
+
+	FDelegateHandle OnConnectedHandle;
+	FDelegateHandle OnConnectionErrorHandle;
+	FDelegateHandle OnClosedHandle;
+	FDelegateHandle OnMessageHandle;
+
+	/** Handle for efficient management of KeepAlive timer */
+	FTimerHandle TimerHandle_KeepAlive;
+	const float KEEP_ALIVE_INTERVAL = 60.0f;
+
+	TMap<FString, TFunction<void(FJsonObjectPtr)>> MessageHandlers;
+};
diff --git a/Engine/Plugins/Media/PixelStreaming/Source/PixelStreaming/Public/PixelStreamingSourceFrame.h b/Engine/Plugins/Media/PixelStreaming/Source/PixelStreaming/Public/PixelStreamingSourceFrame.h
new file mode 100644
--- /dev/null
+++ b/Engine/Plugins/Media/PixelStreaming/Source/PixelStreaming/Public/PixelStreamingSourceFrame.h
@@ -0,0 +1,28 @@
+// Copyright Epic Games, Inc. All Rights Reserved.
+
+#pragma once
+
+#include "PixelStreamingFrameMetadata.h"
+#include "RHI.h"
+
+/*
+ * The unit that FPixelStreamingVideoInput pushes around.
+ * Contains the source frame data and also some metadata
+ * that can be used for timing metrics.
+ */
+class PIXELSTREAMING_API FPixelStreamingSourceFrame
+{
+public:
+	FPixelStreamingSourceFrame(FTexture2DRHIRef InFrameTexture)
+		: FrameTexture(InFrameTexture)
+	{
+		Metadata.SourceTime = FPlatformTime::Cycles64();
+	}
+	virtual ~FPixelStreamingSourceFrame() = default;
+
+	virtual int32 GetWidth() const { return FrameTexture->GetDesc().Extent.X; }
+	virtual int32 GetHeight() const { return FrameTexture->GetDesc().Extent.Y; }
+
+	FTexture2DRHIRef FrameTexture;
+	FPixelStreamingFrameMetadata Metadata;
+};
diff --git a/Engine/Plugins/Media/PixelStreaming/Source/PixelStreaming/Public/PixelStreamingVideoInput.h b/Engine/Plugins/Media/PixelStreaming/Source/PixelStreaming/Public/PixelStreamingVideoInput.h
new file mode 100644
--- /dev/null
+++ b/Engine/Plugins/Media/PixelStreaming/Source/PixelStreaming/Public/PixelStreamingVideoInput.h
@@ -0,0 +1,21 @@
+// Copyright Epic Games, Inc. All Rights Reserved.
+
+#pragma once
+
+#include "Delegates/DelegateCombinations.h"
+#include "CoreMinimal.h"
+
+class FPixelStreamingSourceFrame;
+
+/*
+ * Feeds FStreamer with frame data.
+ * Broadcast the OnFrame delegate when a new frame is ready to be streamed.
+ */
+class PIXELSTREAMING_API FPixelStreamingVideoInput
+{
+public:
+	virtual ~FPixelStreamingVideoInput() = default;
+
+	DECLARE_MULTICAST_DELEGATE_OneParam(FOnFrame, const FPixelStreamingSourceFrame&);
+	FOnFrame OnFrame;
+};
