Fix PixelStreaming servers to launch correctly in VCam usecase and provide better logging/output.

- Added ability to turn off polling if ready on Pixel Streaming servers.
- Added lifecycle logging and child process logging to servers launched through UE::PixelStreaming::Servers.
- Fixed PixelStreamingVCam to shut down streamer before signalling server to avoid ugly log message.
- Modified PixelStreaming::Server tests to have a timeout if they fail to resolve within 60 seconds.
- Exposed HttpPort UProperty on PixelStreamingVCam so user can set webserver port other than 80 if required (this is required for Linux).

diff --git a/Engine/Plugins/Media/PixelStreaming/Source/PixelStreaming/Private/PixelStreamingServers.cpp b/Engine/Plugins/Media/PixelStreaming/Source/PixelStreaming/Private/PixelStreamingServers.cpp
--- a/Engine/Plugins/Media/PixelStreaming/Source/PixelStreaming/Private/PixelStreamingServers.cpp
+++ b/Engine/Plugins/Media/PixelStreaming/Source/PixelStreaming/Private/PixelStreamingServers.cpp
@@ -30,7 +30,7 @@ namespace UE::PixelStreaming
 			return FPaths::ConvertRelativePathToFull(FPaths::RootDir() / SamplesRelativePath);
 		}
 
-		TSharedPtr<FMonitoredProcess> LaunchScriptChildProcess(FString ScriptAbsPath, FString ScriptArgs)
+		TSharedPtr<FMonitoredProcess> LaunchScriptChildProcess(FString ScriptAbsPath, FString ScriptArgs, FString LogPrefix)
 		{
 			// Check if the binary actually exists
 			IPlatformFile& FileManager = FPlatformFileManager::Get().GetPlatformFile();
@@ -57,6 +57,11 @@ namespace UE::PixelStreaming
 
 			TSharedPtr<FMonitoredProcess> ChildProcess = MakeShared<FMonitoredProcess>(ScriptRunner, ScriptRunnerParams, true, true);
 
+			// Bind to output so we can capture the output in the log
+			ChildProcess->OnOutput().BindLambda([LogPrefix](FString Output) {
+				UE_LOG(LogPixelStreaming, Log, TEXT("%s - %s"), *LogPrefix, *Output);
+			});
+
 			// Run the child process
 			UE_LOG(LogPixelStreaming, Log, TEXT("Launch child process - %s %s"), *ScriptRunner, *ScriptRunnerParams);
 			ChildProcess->Launch();
@@ -197,12 +202,7 @@ namespace UE::PixelStreaming
 			CopyServerTo(ExistingServerAbsPath, NewServerPath);
 
 			FString ScriptAbsPath = AddPlatformScripts(NewServerPath);
-			TSharedPtr<FMonitoredProcess> ServerProcess = LaunchScriptChildProcess(ScriptAbsPath, Args);
-
-			// Bind to output so we can capture the output in the log
-			ServerProcess->OnOutput().BindLambda([ServerName](FString Output) {
-				UE_LOG(LogPixelStreaming, Log, TEXT("%s - %s"), *ServerName, *Output);
-			});
+			TSharedPtr<FMonitoredProcess> ServerProcess = LaunchScriptChildProcess(ScriptAbsPath, Args, ServerName);
 
 			// Bind to the cancel and delete the server folder
 			ServerProcess->OnCanceled().BindLambda([NewServerPath]() {
@@ -240,7 +240,7 @@ namespace UE::PixelStreaming
 			else
 			{
 				OutLaunchAbsPath = ServerAbsPath;
-				return LaunchScriptChildProcess(AddPlatformScripts(ServerAbsPath), LaunchArgs.ProcessArgs);
+				return LaunchScriptChildProcess(AddPlatformScripts(ServerAbsPath), LaunchArgs.ProcessArgs, FPaths::GetPathLeaf(ServerAbsPath));
 			}
 
 			/* Note this has not handled waiting for server to become ready, that is server specific and is therefore handled in say `LaunchCirrus` */
@@ -334,6 +334,7 @@ namespace UE::PixelStreaming
 		void FServerBase::Launch(FLaunchArgs& InLaunchArgs)
 		{
 			LaunchArgs = InLaunchArgs;
+			bPollUntilReady = false;
 			ServerProcess = LaunchImpl(InLaunchArgs, ServerRootAbsPath, Endpoints);
 			bHasLaunched = true;
 			PollingStartedSeconds = FPlatformTime::Seconds();
@@ -366,6 +367,12 @@ namespace UE::PixelStreaming
 
 		void FServerBase::Tick(float DeltaTime)
 		{
+			// No need to do polling if polling is turned off
+			if (!bPollUntilReady)
+			{
+				return;
+			}
+
 			// No need to start polling if we have not launched or we have already concluded the server is ready.
 			if (!bHasLaunched || bIsReady)
 			{
@@ -445,6 +452,11 @@ namespace UE::PixelStreaming
 
 		protected:
 			/* Begin FServerBase interface */
+			void Stop() override
+			{
+				UE_LOG(LogPixelStreaming, Log, TEXT("Stopping Cirrus signalling server."));
+				FServerBase::Stop();
+			}
 
 			TSharedPtr<FMonitoredProcess> LaunchImpl(FLaunchArgs& InLaunchArgs, FString& OutServerPath, TMap<EEndpoint, FString>& OutEndPoints) override
 			{
@@ -458,15 +470,20 @@ namespace UE::PixelStreaming
 
 				// Construct endpoint urls
 				const FString SignallingStreamerUrl = FString::Printf(TEXT("ws://localhost:%s"), *StreamerPort);
+				const FString WebserverUrl = FString::Printf(TEXT("http://localhost:%s"), *HttpPort);
 				OutEndPoints.Add(EEndpoint::Signalling_Streamer, SignallingStreamerUrl);
 				OutEndPoints.Add(EEndpoint::Signalling_Players, FString::Printf(TEXT("ws://localhost:%s"), *HttpPort));
 				OutEndPoints.Add(EEndpoint::Signalling_SFU, FString::Printf(TEXT("ws://localhost:%s"), *SFUPort));
 				OutEndPoints.Add(EEndpoint::Signalling_Matchmaker, FString::Printf(TEXT("ws://localhost:%s"), *MatchmakerPort));
-				OutEndPoints.Add(EEndpoint::Signalling_Webserver, FString::Printf(TEXT("http://localhost:%s"), *HttpPort));
+				OutEndPoints.Add(EEndpoint::Signalling_Webserver, WebserverUrl);
 
 				TSharedPtr<FMonitoredProcess> CirrusProcess = LaunchServer(CirrusDir, InLaunchArgs, OutServerPath);
+				UE_LOG(LogPixelStreaming, Log, TEXT("Cirrus signalling server running at: %s"), *WebserverUrl);
 
-				Probe = MakeUnique<FWebSocketProbe>(SignallingStreamerUrl);
+				if (bPollUntilReady)
+				{
+					Probe = MakeUnique<FWebSocketProbe>(SignallingStreamerUrl);
+				}
 
 				return CirrusProcess;
 			}
diff --git a/Engine/Plugins/Media/PixelStreaming/Source/PixelStreaming/Private/PixelStreamingSignallingConnection.cpp b/Engine/Plugins/Media/PixelStreaming/Source/PixelStreaming/Private/PixelStreamingSignallingConnection.cpp
--- a/Engine/Plugins/Media/PixelStreaming/Source/PixelStreaming/Private/PixelStreamingSignallingConnection.cpp
+++ b/Engine/Plugins/Media/PixelStreaming/Source/PixelStreaming/Private/PixelStreamingSignallingConnection.cpp
@@ -280,7 +280,7 @@ void FPixelStreamingSignallingConnection::OnConnected()
 
 void FPixelStreamingSignallingConnection::OnConnectionError(const FString& Error)
 {
-	UE_LOG(LogPixelStreamingSS, Error, TEXT("Failed to connect to SS: %s"), *Error);
+	UE_LOG(LogPixelStreamingSS, Warning, TEXT("Failed to connect to SS - signalling server may not be up yet. Message: \"%s\""), *Error);
 
 	Observer.OnSignallingError(Error);
 
diff --git a/Engine/Plugins/Media/PixelStreaming/Source/PixelStreaming/Private/Tests/ServerTests.cpp b/Engine/Plugins/Media/PixelStreaming/Source/PixelStreaming/Private/Tests/ServerTests.cpp
--- a/Engine/Plugins/Media/PixelStreaming/Source/PixelStreaming/Private/Tests/ServerTests.cpp
+++ b/Engine/Plugins/Media/PixelStreaming/Source/PixelStreaming/Private/Tests/ServerTests.cpp
@@ -42,6 +42,7 @@ namespace UE::PixelStreaming::Servers
 		TSharedPtr<FServerBase> SignallingServer = MakeSignallingServer();
 		FLaunchArgs LaunchArgs;
 		LaunchArgs.bEphemeral = true;
+		LaunchArgs.bPollUntilReady = false;
 
 		SignallingServer->Launch(LaunchArgs);
 
@@ -61,6 +62,7 @@ namespace UE::PixelStreaming::Servers
 		TSharedPtr<FServerBase> SignallingServer = MakeSignallingServer();
 		FLaunchArgs LaunchArgs;
 		LaunchArgs.bEphemeral = true;
+		LaunchArgs.bPollUntilReady = true;
 
 		SignallingServer->OnReady.AddLambda([this](TMap<EEndpoint, FString> Endpoints) {
 			TestTrue("Server was ready.", true);
diff --git a/Engine/Plugins/Media/PixelStreaming/Source/PixelStreaming/Public/PixelStreamingServers.h b/Engine/Plugins/Media/PixelStreaming/Source/PixelStreaming/Public/PixelStreamingServers.h
--- a/Engine/Plugins/Media/PixelStreaming/Source/PixelStreaming/Public/PixelStreamingServers.h
+++ b/Engine/Plugins/Media/PixelStreaming/Source/PixelStreaming/Public/PixelStreamingServers.h
@@ -23,11 +23,14 @@ namespace UE::PixelStreaming
 		 **/
 		struct PIXELSTREAMING_API FLaunchArgs
 		{
+			// Arguments passed to the actual server when its process is started.
+			FString ProcessArgs = TEXT("");
+
 			// If true, will copy to a temp directory and when finished delete itself. If false, will launch from Samples/PixelStreaming/WebServers/*.
 			bool bEphemeral = true;
 
-			// Arguments passed to the actual server when its process is started.
-			FString ProcessArgs = TEXT("");
+			// If true poll until ready
+			bool bPollUntilReady = false;
 
 			// Reconnection timeout in seconds
 			float ReconnectionTimeoutSeconds = 30.0f;
@@ -141,6 +144,7 @@ namespace UE::PixelStreaming
 			float LastReconnectionTimeSeconds = 0.0f;
 			bool bAllowedToTick = true;
 			bool bTimedOut = false;
+			bool bPollUntilReady = false;
 		};
 
 		/* -------------- Static utility methods for working with Pixel Streaming servers. ----------------- */
diff --git a/Engine/Plugins/Media/PixelStreamingVCam/Source/PixelStreamingVCam/Private/VCamPixelStreamingSession.cpp b/Engine/Plugins/Media/PixelStreamingVCam/Source/PixelStreamingVCam/Private/VCamPixelStreamingSession.cpp
--- a/Engine/Plugins/Media/PixelStreamingVCam/Source/PixelStreamingVCam/Private/VCamPixelStreamingSession.cpp
+++ b/Engine/Plugins/Media/PixelStreamingVCam/Source/PixelStreamingVCam/Private/VCamPixelStreamingSession.cpp
@@ -5,6 +5,7 @@
 #include "PixelStreamingServers.h"
 #include "Editor/EditorPerformanceSettings.h"
 #include "Containers/UnrealString.h"
+#include "Async/Async.h"
 
 DEFINE_LOG_CATEGORY(LogVCamOutputProvider);
 namespace VCamPixelStreamingSession
@@ -61,7 +62,7 @@ void UVCamPixelStreamingSession::Activate()
 		CirrusProcess = UE::PixelStreaming::Servers::MakeSignallingServer();
 		UE::PixelStreaming::Servers::FLaunchArgs LaunchArgs;
 		LaunchArgs.bEphemeral = false;
-		LaunchArgs.ProcessArgs = FString::Printf(TEXT("--StreamerPort=%s"), *FString::FromInt(PortNumber));
+		LaunchArgs.ProcessArgs = FString::Printf(TEXT("--StreamerPort=%s --HttpPort=%s --nosudo"), *FString::FromInt(PortNumber), *FString::FromInt(HttpPort));
 		CirrusProcess->Launch(LaunchArgs);
 	}
 
@@ -98,13 +99,35 @@ void UVCamPixelStreamingSession::Activate()
 	Super::Activate();
 }
 
+void UVCamPixelStreamingSession::StopSignallingServer()
+{
+	if (CirrusProcess.IsValid())
+	{
+		CirrusProcess->Stop();
+	}
+}
+
 void UVCamPixelStreamingSession::Deactivate()
 {
 	if (MediaCapture)
 	{
+
+		if (MediaOutput && MediaOutput->GetStreamer())
+		{
+			// Shutting streamer down before closing signalling server prevents an ugly websocket disconnect showing in the log
+			MediaOutput->GetStreamer()->StopStreaming();
+			StopSignallingServer();
+		}
+
 		MediaCapture->StopCapture(true);
 		MediaCapture = nullptr;
 	}
+	else
+	{
+		// There is not media capture we defensively clean up the signalling server if it exists.
+		StopSignallingServer();
+	}
+
 	Super::Deactivate();
 	if (bUsingDummyUMG)
 	{
@@ -112,11 +135,6 @@ void UVCamPixelStreamingSession::Deactivate()
 		bUsingDummyUMG = false;
 	}
 
-	if(CirrusProcess.IsValid())
-	{
-		CirrusProcess->Stop();
-	}
-	
 	UEditorPerformanceSettings* Settings = GetMutableDefault<UEditorPerformanceSettings>();
 	Settings->bThrottleCPUWhenNotForeground = bOldThrottleCPUWhenNotForeground;
 	Settings->PostEditChange();
diff --git a/Engine/Plugins/Media/PixelStreamingVCam/Source/PixelStreamingVCam/Public/VCamPixelStreamingSession.h b/Engine/Plugins/Media/PixelStreamingVCam/Source/PixelStreamingVCam/Public/VCamPixelStreamingSession.h
--- a/Engine/Plugins/Media/PixelStreamingVCam/Source/PixelStreamingVCam/Public/VCamPixelStreamingSession.h
+++ b/Engine/Plugins/Media/PixelStreamingVCam/Source/PixelStreamingVCam/Public/VCamPixelStreamingSession.h
@@ -9,7 +9,7 @@
 #include "Slate/SceneViewport.h"
 #include "Engine/TextureRenderTarget2D.h"
 #if WITH_EDITOR
-#include "LevelEditorViewport.h"
+	#include "LevelEditorViewport.h"
 #endif
 #include "VCamPixelStreamingSession.generated.h"
 
@@ -32,7 +32,11 @@ public:
 	// Network port number - change this only if connecting multiple RemoteSession devices to the same PC
 	UPROPERTY(EditAnywhere, BlueprintReadWrite, Category = "Output")
 	int32 PortNumber = 8888;
-
+	
+	// Http webserver port number - e.g. Go to http://localhost:YourPort to access the streamed VCam output. Warning Ports below 1024 require sudo under Linux so it is recommended to use a higher port on that platform.
+	UPROPERTY(EditAnywhere, BlueprintReadWrite, Category = "Output")
+	int32 HttpPort = 80;
+	
 	// If using the output from a Composure Output Provider, specify it here
 	UPROPERTY(EditAnywhere, BlueprintReadWrite, Category = "Output")
 	int32 FromComposureOutputProviderIndex = INDEX_NONE;
@@ -52,6 +56,9 @@ protected:
 	UPROPERTY(Transient)
 	UPixelStreamingMediaCapture* MediaCapture = nullptr;
 
+private:
+	void StopSignallingServer();
+
 private:
 	FHitResult LastViewportTouchResult;
 	bool bUsingDummyUMG = false;
