Fix nonunity build error in game target.

diff --git a/Engine/Plugins/Media/PixelStreaming/Source/PixelStreaming/Private/Streamer.cpp b/Engine/Plugins/Media/PixelStreaming/Source/PixelStreaming/Private/Streamer.cpp
--- a/Engine/Plugins/Media/PixelStreaming/Source/PixelStreaming/Private/Streamer.cpp
+++ b/Engine/Plugins/Media/PixelStreaming/Source/PixelStreaming/Private/Streamer.cpp
@@ -1,6 +1,13 @@
 // Copyright Epic Games, Inc. All Rights Reserved.
 
 #include "Streamer.h"
+#include "IPixelStreamingModule.h"
+#include "DataChannelObserver.h"
+#include "PixelStreamingPrivate.h"
+#include "PixelStreamingDelegates.h"
+#include "PixelStreamingSourceFrame.h"
+#include "PixelStreamingSignallingConnection.h"
+#include "PixelStreamingAudioDeviceModule.h"
 #include "WebRTCIncludes.h"
 #include "VideoEncoderFactory.h"
 #include "ToStringExtensions.h"
@@ -10,22 +17,16 @@
 #include "Settings.h"
 #include "WebRTCLogging.h"
 #include "WebSocketsModule.h"
-#include "PixelStreamingAudioDeviceModule.h"
 #include "VideoSourceP2P.h"
 #include "VideoSourceSFU.h"
-#include "PixelStreamingDelegates.h"
-#include "DataChannelObserver.h"
-#include "IPixelStreamingModule.h"
 #include "VideoSourceGroup.h"
 #include "InputDevice.h"
 #include "Misc/CoreDelegates.h"
 #include "UtilsRender.h"
 #include "IImageWrapper.h"
 #include "IImageWrapperModule.h"
-#include "PixelStreamingSignallingConnection.h"
-#include "PixelStreamingPrivate.h"
 #include "Async/Async.h"
-#include "PixelStreamingSourceFrame.h"
+#include "Framework/Application/SlateApplication.h"
 
 namespace UE::PixelStreaming
 {
