Add additional PixelStreaming tests to automation

diff --git a/Engine/Plugins/Media/PixelStreaming/Source/PixelStreaming/Private/Tests/FrameAdaptTest.cpp b/Engine/Plugins/Media/PixelStreaming/Source/PixelStreaming/Private/Tests/FrameAdaptTest.cpp
--- a/Engine/Plugins/Media/PixelStreaming/Source/PixelStreaming/Private/Tests/FrameAdaptTest.cpp
+++ b/Engine/Plugins/Media/PixelStreaming/Source/PixelStreaming/Private/Tests/FrameAdaptTest.cpp
@@ -124,7 +124,7 @@
 		}
 	};
 
-	IMPLEMENT_SIMPLE_AUTOMATION_TEST(FAdaptProcessTest, "PixelStreaming.AdaptProcess", EAutomationTestFlags::ApplicationContextMask | EAutomationTestFlags::EngineFilter)
+	IMPLEMENT_SIMPLE_AUTOMATION_TEST(FAdaptProcessTest, "System.Plugins.PixelStreaming.AdaptProcess", EAutomationTestFlags::EditorContext | EAutomationTestFlags::ClientContext | EAutomationTestFlags::ProductFilter)
 	bool FAdaptProcessTest::RunTest(const FString& Parameters)
 	{
 		const int32 MockFrame1Width = 32;
@@ -204,7 +204,7 @@
 		return true;
 	}
 
-	IMPLEMENT_SIMPLE_AUTOMATION_TEST(FFrameAdapterTest, "PixelStreaming.FrameAdapterTest", EAutomationTestFlags::ApplicationContextMask | EAutomationTestFlags::EngineFilter)
+	IMPLEMENT_SIMPLE_AUTOMATION_TEST(FFrameAdapterTest, "System.Plugins.PixelStreaming.FrameAdapterTest", EAutomationTestFlags::EditorContext | EAutomationTestFlags::ClientContext | EAutomationTestFlags::ProductFilter)
 	bool FFrameAdapterTest::RunTest(const FString& Parameters)
 	{
 		TSharedPtr<FPixelStreamingVideoInput> MockVideoInput = MakeShared<FPixelStreamingVideoInput>();
diff --git a/Engine/Plugins/Media/PixelStreaming/Source/PixelStreaming/Private/Tests/FrameMetadataTest.cpp b/Engine/Plugins/Media/PixelStreaming/Source/PixelStreaming/Private/Tests/FrameMetadataTest.cpp
--- a/Engine/Plugins/Media/PixelStreaming/Source/PixelStreaming/Private/Tests/FrameMetadataTest.cpp
+++ b/Engine/Plugins/Media/PixelStreaming/Source/PixelStreaming/Private/Tests/FrameMetadataTest.cpp
@@ -9,7 +9,7 @@
 namespace UE::PixelStreaming
 {
 
-	IMPLEMENT_SIMPLE_AUTOMATION_TEST(FFrameMetadataTest, "PixelStreaming.FrameMetadata", EAutomationTestFlags::ApplicationContextMask | EAutomationTestFlags::EngineFilter)
+	IMPLEMENT_SIMPLE_AUTOMATION_TEST(FFrameMetadataTest, "System.Plugins.PixelStreaming.FrameMetadata", EAutomationTestFlags::EditorContext | EAutomationTestFlags::ClientContext | EAutomationTestFlags::ProductFilter)
 	bool FFrameMetadataTest::RunTest(const FString& Parameters)
 	{
 		FPixelStreamingFrameMetadata Metadata;
diff --git a/Engine/Plugins/Media/PixelStreaming/Source/PixelStreaming/Private/Tests/ServerTests.cpp b/Engine/Plugins/Media/PixelStreaming/Source/PixelStreaming/Private/Tests/ServerTests.cpp
--- a/Engine/Plugins/Media/PixelStreaming/Source/PixelStreaming/Private/Tests/ServerTests.cpp
+++ b/Engine/Plugins/Media/PixelStreaming/Source/PixelStreaming/Private/Tests/ServerTests.cpp
@@ -35,7 +35,7 @@
 		}
 	}
 
-	IMPLEMENT_SIMPLE_AUTOMATION_TEST(FLaunchTransientCirrusTest, "PixelStreaming.LaunchTransientCirrus", EAutomationTestFlags::ApplicationContextMask | EAutomationTestFlags::EngineFilter)
+	IMPLEMENT_SIMPLE_AUTOMATION_TEST(FLaunchTransientCirrusTest, "System.Plugins.PixelStreaming.LaunchTransientCirrus", EAutomationTestFlags::EditorContext | EAutomationTestFlags::ClientContext | EAutomationTestFlags::ProductFilter)
 	bool FLaunchTransientCirrusTest::RunTest(const FString& Parameters)
 	{
 
@@ -55,7 +55,7 @@
 		return true;
 	}
 
-	IMPLEMENT_SIMPLE_AUTOMATION_TEST(FLaunchTransientCirrusWithParams, "PixelStreaming.LaunchTransientCirrusWithParams", EAutomationTestFlags::ApplicationContextMask | EAutomationTestFlags::EngineFilter)
+	IMPLEMENT_SIMPLE_AUTOMATION_TEST(FLaunchTransientCirrusWithParams, "System.Plugins.PixelStreaming.LaunchTransientCirrusWithParams", EAutomationTestFlags::EditorContext | EAutomationTestFlags::ClientContext | EAutomationTestFlags::ProductFilter)
 	bool FLaunchTransientCirrusWithParams::RunTest(const FString& Parameters)
 	{
 
diff --git a/Engine/Plugins/Media/PixelStreaming/Source/PixelStreaming/Private/Tests/SourceFrameTest.cpp b/Engine/Plugins/Media/PixelStreaming/Source/PixelStreaming/Private/Tests/SourceFrameTest.cpp
--- a/Engine/Plugins/Media/PixelStreaming/Source/PixelStreaming/Private/Tests/SourceFrameTest.cpp
+++ b/Engine/Plugins/Media/PixelStreaming/Source/PixelStreaming/Private/Tests/SourceFrameTest.cpp
@@ -9,7 +9,7 @@
 namespace UE::PixelStreaming
 {
 
-	IMPLEMENT_SIMPLE_AUTOMATION_TEST(FSourceFrameTest, "PixelStreaming.SourceFrame", EAutomationTestFlags::ApplicationContextMask | EAutomationTestFlags::EngineFilter)
+	IMPLEMENT_SIMPLE_AUTOMATION_TEST(FSourceFrameTest, "System.Plugins.PixelStreaming.SourceFrame", EAutomationTestFlags::EditorContext | EAutomationTestFlags::ClientContext | EAutomationTestFlags::ProductFilter)
 	bool FSourceFrameTest::RunTest(const FString& Parameters)
 	{
 		FRHITextureCreateDesc TextureDesc =
