Disabling some PVS-Studio warnings in test code.

diff --git a/Engine/Plugins/Media/PixelStreaming/Source/PixelStreaming/Private/Tests/FrameMetadataTest.cpp b/Engine/Plugins/Media/PixelStreaming/Source/PixelStreaming/Private/Tests/FrameMetadataTest.cpp
--- a/Engine/Plugins/Media/PixelStreaming/Source/PixelStreaming/Private/Tests/FrameMetadataTest.cpp
+++ b/Engine/Plugins/Media/PixelStreaming/Source/PixelStreaming/Private/Tests/FrameMetadataTest.cpp
@@ -9,48 +9,49 @@
 namespace UE::PixelStreaming
 {
 
-	IMPLEMENT_SIMPLE_AUTOMATION_TEST(FFrameMetadataTest, "System.Plugins.PixelStreaming.FrameMetadata", EAutomationTestFlags::EditorContext | EAutomationTestFlags::ClientContext | EAutomationTestFlags::ProductFilter)
-	bool FFrameMetadataTest::RunTest(const FString& Parameters)
-	{
-		FPixelStreamingFrameMetadata Metadata;
-
-		TestTrue("Layer initialized", Metadata.Layer == 0);
-		TestTrue("SourceTime initialized", Metadata.SourceTime == 0);
-		TestTrue("AdaptCallTime initialized", Metadata.AdaptCallTime == 0);
-		TestTrue("AdaptProcessStartTime initialized", Metadata.AdaptProcessStartTime == 0);
-		TestTrue("AdaptProcessFinalizeTime initialized", Metadata.AdaptProcessFinalizeTime == 0);
-		TestTrue("AdaptProcessEndTime initialized", Metadata.AdaptProcessEndTime == 0);
-		TestTrue("UseCount initialized", Metadata.UseCount == 0);
-		TestTrue("FirstEncodeStartTime initialized", Metadata.FirstEncodeStartTime == 0);
-		TestTrue("LastEncodeStartTime initialized", Metadata.LastEncodeStartTime == 0);
-		TestTrue("LastEncodedEndTime initialized", Metadata.LastEncodeEndTime == 0);
-
-		Metadata.Layer = 1;
-		Metadata.SourceTime = 2;
-		Metadata.AdaptCallTime = 3;
-		Metadata.AdaptProcessStartTime = 4;
-		Metadata.AdaptProcessFinalizeTime = 5;
-		Metadata.AdaptProcessEndTime = 6;
-		Metadata.UseCount = 7;
-		Metadata.FirstEncodeStartTime = 8;
-		Metadata.LastEncodeStartTime = 9;
-		Metadata.LastEncodeEndTime = 10;
-
-		FPixelStreamingFrameMetadata MetadataCopy = Metadata.Copy();
-
-		TestTrue("Layer copied", MetadataCopy.Layer == 1);
-		TestTrue("SourceTime copied", MetadataCopy.SourceTime == 2);
-		TestTrue("AdaptCallTime copied", MetadataCopy.AdaptCallTime == 3);
-		TestTrue("AdaptProcessStartTime copied", MetadataCopy.AdaptProcessStartTime == 4);
-		TestTrue("AdaptProcessFinalizeTime copied", MetadataCopy.AdaptProcessFinalizeTime == 5);
-		TestTrue("AdaptProcessEndTime copied", MetadataCopy.AdaptProcessEndTime == 6);
-		TestTrue("UseCount copied", MetadataCopy.UseCount == 7);
-		TestTrue("FirstEncodeStartTime copied", MetadataCopy.FirstEncodeStartTime == 8);
-		TestTrue("LastEncodeStartTime copied", MetadataCopy.LastEncodeStartTime == 9);
-		TestTrue("LastEncodedEndTime copied", MetadataCopy.LastEncodeEndTime == 10);
-
-		return true;
-	}
+    IMPLEMENT_SIMPLE_AUTOMATION_TEST(FFrameMetadataTest, "System.Plugins.PixelStreaming.FrameMetadata", EAutomationTestFlags::EditorContext | EAutomationTestFlags::ClientContext | EAutomationTestFlags::ProductFilter)
+    bool FFrameMetadataTest::RunTest(const FString& Parameters)
+    {
+        FPixelStreamingFrameMetadata Metadata;
+
+        // Disabling PVS-Studio "always true" for these since thats the assumption we're testing
+        TestTrue("Layer initialized", Metadata.Layer == 0);                                       //-V547
+        TestTrue("SourceTime initialized", Metadata.SourceTime == 0);                             //-V547
+        TestTrue("AdaptCallTime initialized", Metadata.AdaptCallTime == 0);                       //-V547
+        TestTrue("AdaptProcessStartTime initialized", Metadata.AdaptProcessStartTime == 0);       //-V547
+        TestTrue("AdaptProcessFinalizeTime initialized", Metadata.AdaptProcessFinalizeTime == 0); //-V547
+        TestTrue("AdaptProcessEndTime initialized", Metadata.AdaptProcessEndTime == 0);           //-V547
+        TestTrue("UseCount initialized", Metadata.UseCount == 0);                                 //-V547
+        TestTrue("FirstEncodeStartTime initialized", Metadata.FirstEncodeStartTime == 0);         //-V547
+        TestTrue("LastEncodeStartTime initialized", Metadata.LastEncodeStartTime == 0);           //-V547
+        TestTrue("LastEncodedEndTime initialized", Metadata.LastEncodeEndTime == 0);              //-V547
+
+        Metadata.Layer = 1;
+        Metadata.SourceTime = 2;
+        Metadata.AdaptCallTime = 3;
+        Metadata.AdaptProcessStartTime = 4;
+        Metadata.AdaptProcessFinalizeTime = 5;
+        Metadata.AdaptProcessEndTime = 6;
+        Metadata.UseCount = 7;
+        Metadata.FirstEncodeStartTime = 8;
+        Metadata.LastEncodeStartTime = 9;
+        Metadata.LastEncodeEndTime = 10;
+
+        FPixelStreamingFrameMetadata MetadataCopy = Metadata.Copy();
+
+        TestTrue("Layer copied", MetadataCopy.Layer == 1);
+        TestTrue("SourceTime copied", MetadataCopy.SourceTime == 2);
+        TestTrue("AdaptCallTime copied", MetadataCopy.AdaptCallTime == 3);
+        TestTrue("AdaptProcessStartTime copied", MetadataCopy.AdaptProcessStartTime == 4);
+        TestTrue("AdaptProcessFinalizeTime copied", MetadataCopy.AdaptProcessFinalizeTime == 5);
+        TestTrue("AdaptProcessEndTime copied", MetadataCopy.AdaptProcessEndTime == 6);
+        TestTrue("UseCount copied", MetadataCopy.UseCount == 7);
+        TestTrue("FirstEncodeStartTime copied", MetadataCopy.FirstEncodeStartTime == 8);
+        TestTrue("LastEncodeStartTime copied", MetadataCopy.LastEncodeStartTime == 9);
+        TestTrue("LastEncodedEndTime copied", MetadataCopy.LastEncodeEndTime == 10);
+
+        return true;
+    }
 
 } // namespace UE::PixelStreaming
 
