Disable PixelStreaming test that doesn't work in a Horde environment

diff --git a/Engine/Plugins/Media/PixelStreaming/Source/PixelStreaming/Private/Tests/ServerTests.cpp b/Engine/Plugins/Media/PixelStreaming/Source/PixelStreaming/Private/Tests/ServerTests.cpp
--- a/Engine/Plugins/Media/PixelStreaming/Source/PixelStreaming/Private/Tests/ServerTests.cpp
+++ b/Engine/Plugins/Media/PixelStreaming/Source/PixelStreaming/Private/Tests/ServerTests.cpp
@@ -35,7 +35,7 @@ namespace UE::PixelStreaming::Servers
 		}
 	}
 
-	IMPLEMENT_SIMPLE_AUTOMATION_TEST(FLaunchTransientCirrusTest, "System.Plugins.PixelStreaming.LaunchTransientCirrus", EAutomationTestFlags::EditorContext | EAutomationTestFlags::ClientContext | EAutomationTestFlags::ProductFilter)
+	IMPLEMENT_SIMPLE_AUTOMATION_TEST(FLaunchTransientCirrusTest, "System.Plugins.PixelStreaming.LaunchTransientCirrus", EAutomationTestFlags::EditorContext | EAutomationTestFlags::ClientContext | EAutomationTestFlags::ProductFilter | EAutomationTestFlags::Disabled)
 	bool FLaunchTransientCirrusTest::RunTest(const FString& Parameters)
 	{
 
@@ -55,7 +55,7 @@ namespace UE::PixelStreaming::Servers
 		return true;
 	}
 
-	IMPLEMENT_SIMPLE_AUTOMATION_TEST(FLaunchTransientCirrusWithParams, "System.Plugins.PixelStreaming.LaunchTransientCirrusWithParams", EAutomationTestFlags::EditorContext | EAutomationTestFlags::ClientContext | EAutomationTestFlags::ProductFilter)
+	IMPLEMENT_SIMPLE_AUTOMATION_TEST(FLaunchTransientCirrusWithParams, "System.Plugins.PixelStreaming.LaunchTransientCirrusWithParams", EAutomationTestFlags::EditorContext | EAutomationTestFlags::ClientContext | EAutomationTestFlags::ProductFilter | EAutomationTestFlags::Disabled)
 	bool FLaunchTransientCirrusWithParams::RunTest(const FString& Parameters)
 	{
 
