Remove NDA platform references

#rnx

diff --git a/Engine/Plugins/Media/PixelStreaming/Source/PixelStreaming/Private/PCH.h b/Engine/Plugins/Media/PixelStreaming/Source/PixelStreaming/Private/PCH.h
--- a/Engine/Plugins/Media/PixelStreaming/Source/PixelStreaming/Private/PCH.h
+++ b/Engine/Plugins/Media/PixelStreaming/Source/PixelStreaming/Private/PCH.h
@@ -72,7 +72,7 @@
 #include "Engine/Engine.h"
 #include "Framework/Application/SlateUser.h"
 
-#if PLATFORM_WINDOWS || PLATFORM_XBOXONE
+#if PLATFORM_WINDOWS
 	#include "Windows/WindowsPlatformMisc.h"
 #elif PLATFORM_LINUX
 	#include "Linux/LinuxPlatformMisc.h"
diff --git a/Engine/Plugins/Media/PixelStreaming/Source/PixelStreaming/Private/PixelStreamingModule.cpp b/Engine/Plugins/Media/PixelStreaming/Source/PixelStreaming/Private/PixelStreamingModule.cpp
--- a/Engine/Plugins/Media/PixelStreaming/Source/PixelStreaming/Private/PixelStreamingModule.cpp
+++ b/Engine/Plugins/Media/PixelStreaming/Source/PixelStreaming/Private/PixelStreamingModule.cpp
@@ -18,7 +18,7 @@
 #include "Utils.h"
 #include "UtilsRender.h"
 
-#if PLATFORM_WINDOWS || PLATFORM_XBOXONE
+#if PLATFORM_WINDOWS
 #include "Windows/WindowsHWrapper.h"
 #elif PLATFORM_LINUX
 #include "CudaModule.h"
@@ -361,7 +361,7 @@
 	{
 		bool bCompatible = true;
 
-#if PLATFORM_WINDOWS || PLATFORM_XBOXONE
+#if PLATFORM_WINDOWS
 		bool bWin8OrHigher = IsWindows8OrGreater();
 		if (!bWin8OrHigher)
 		{
diff --git a/Engine/Source/Runtime/AVEncoder/Private/MicrosoftCommon.cpp b/Engine/Source/Runtime/AVEncoder/Private/MicrosoftCommon.cpp
--- a/Engine/Source/Runtime/AVEncoder/Private/MicrosoftCommon.cpp
+++ b/Engine/Source/Runtime/AVEncoder/Private/MicrosoftCommon.cpp
@@ -25,18 +25,6 @@ namespace AVEncoder
 //}
 //#endif
 //
-////
-//// XboxOne only code
-//// 
-//#if PLATFORM_XBOXONE 
-//ID3D12Device* GetUE4DxDevice()
-//{
-//	auto Device = static_cast<ID3D12Device*>(GDynamicRHI->RHIGetNativeDevice());
-//	checkf(Device != nullptr, TEXT("Failed to get UE4's ID3D12Device"));
-//	return Device;
-//}
-//
-//#endif
 
 
 } // namespace AVEncoder
diff --git a/Engine/Source/Runtime/AVEncoder/Private/MicrosoftCommon.h b/Engine/Source/Runtime/AVEncoder/Private/MicrosoftCommon.h
--- a/Engine/Source/Runtime/AVEncoder/Private/MicrosoftCommon.h
+++ b/Engine/Source/Runtime/AVEncoder/Private/MicrosoftCommon.h
@@ -100,10 +100,6 @@ const GUID CLSID_VideoProcessorMFT = { 0x88753b26, 0x5b24, 0x49bd, { 0xb2, 0xe7,
 //
 //ID3D11Device* GetUE4DxDevice();
 //
-//#elif PLATFORM_XBOXONE
-//
-//ID3D12Device* GetUE4DxDevice();
-//
 //#endif
 
 //// scope-disable particular DX11 Debug Layer errors
diff --git a/Engine/Source/Runtime/GameplayMediaEncoder/Private/GameplayMediaEncoderCommon.h b/Engine/Source/Runtime/GameplayMediaEncoder/Private/GameplayMediaEncoderCommon.h
--- a/Engine/Source/Runtime/GameplayMediaEncoder/Private/GameplayMediaEncoderCommon.h
+++ b/Engine/Source/Runtime/GameplayMediaEncoder/Private/GameplayMediaEncoderCommon.h
@@ -24,7 +24,7 @@
 //
 // Includes common to Windows and XboxOne
 //
-#if PLATFORM_WINDOWS || PLATFORM_XBOXONE
+#if PLATFORM_WINDOWS
 	#include "Templates/RefCounting.h"
 #endif
 
