PixelStreaming: Fix access to FSlateApplication when it is not initialized

diff --git a/Engine/Plugins/Media/PixelStreaming/Source/PixelStreaming/Private/PixelStreamingModule.cpp b/Engine/Plugins/Media/PixelStreaming/Source/PixelStreaming/Private/PixelStreamingModule.cpp
--- a/Engine/Plugins/Media/PixelStreaming/Source/PixelStreaming/Private/PixelStreamingModule.cpp
+++ b/Engine/Plugins/Media/PixelStreaming/Source/PixelStreaming/Private/PixelStreamingModule.cpp
@@ -78,6 +78,11 @@ namespace UE::PixelStreaming
 			return;
 		}
 
+		if (!FSlateApplication::IsInitialized())
+		{
+			return;
+		}
+
 		const ERHIInterfaceType RHIType = GDynamicRHI ? RHIGetInterfaceType() : ERHIInterfaceType::Hidden;
 
 		StreamerInputDevices = MakeShared<FStreamerInputDevices>(FSlateApplication::Get().GetPlatformApplication()->GetMessageHandler());
@@ -108,12 +113,13 @@ namespace UE::PixelStreaming
 		ExternalVideoSourceGroup = MakeUnique<FVideoSourceGroup>();
 		ExternalVideoSourceGroup->SetVideoInput(MakeShared<FVideoInputBackBuffer>());
 		ExternalVideoSourceGroup->Start();
+
+		bStartupCompleted = true;
 	}
 
 	void FPixelStreamingModule::ShutdownModule()
 	{
-		// Pixel Streaming does not make sense without an RHI so we don't run in commandlets without one.
-		if (IsRunningCommandlet() && !IsAllowCommandletRendering())
+		if (!bStartupCompleted)
 		{
 			return;
 		}
@@ -127,6 +133,8 @@ namespace UE::PixelStreaming
 		rtc::CleanupSSL();
 
 		IModularFeatures::Get().UnregisterModularFeature(GetModularFeatureName(), this);
+
+		bStartupCompleted = false;
 	}
 	/** 
 	 * End IModuleInterface implementation 
diff --git a/Engine/Plugins/Media/PixelStreaming/Source/PixelStreaming/Private/PixelStreamingModule.h b/Engine/Plugins/Media/PixelStreaming/Source/PixelStreaming/Private/PixelStreamingModule.h
--- a/Engine/Plugins/Media/PixelStreaming/Source/PixelStreaming/Private/PixelStreamingModule.h
+++ b/Engine/Plugins/Media/PixelStreaming/Source/PixelStreaming/Private/PixelStreamingModule.h
@@ -62,6 +62,7 @@
 
 	private:
 		bool bModuleReady = false;
+		bool bStartupCompleted = false;
 		static IPixelStreamingModule* PixelStreamingModule;
 
 		FReadyEvent ReadyEvent;
