Added mechanism to easily fetch and package PixelStreaming infra from its new location

diff --git a/Engine/Plugins/Media/PixelStreaming/Resources/get_ps_servers.bat b/Engine/Plugins/Media/PixelStreaming/Resources/get_ps_servers.bat
new file mode 100755
--- /dev/null
+++ b/Engine/Plugins/Media/PixelStreaming/Resources/get_ps_servers.bat
@@ -0,0 +1,83 @@
+@Rem Copyright Epic Games, Inc. All Rights Reserved.
+
+@echo off
+
+@Rem Set script location as working directory for commands.
+pushd "%~dp0"
+
+:arg_loop_start
+SET ARG=%1
+if DEFINED ARG (
+    if "%ARG%"=="/h" (
+        goto print_help
+    )
+    if "%ARG%"=="/p" (
+        SET GitHubAccessToken=%2
+        SHIFT
+    )
+    if "%ARG%"=="/b" (
+        SET PSInfraTagOrBranch=%2
+        SET IsTag=0
+        SHIFT
+    )
+    if "%ARG%"=="/t" (
+        SET PSInfraTagOrBranch=%2
+        SET IsTag=1
+        SHIFT
+    )
+    SHIFT
+    goto arg_loop_start
+)
+
+@Rem Name and version of ps-infra that we are downloading
+SET PSInfraOrg=EpicGames
+SET PSInfraRepo=PixelStreamingInfrastructure
+
+if NOT DEFINED PSInfraTagOrBranch (
+    SET PSInfraTagOrBranch=v0.1.0-prerelease
+    SET IsTag=1
+)
+
+if %IsTag%==1 (
+  SET RefType=tags
+) else (
+  SET RefType=heads
+)
+
+@Rem Look for a ps-infra directory next to this script
+if exist ps-infra\ (
+  echo ps-infra directory found...skipping install.
+) else (
+  echo ps-infra directory not found...beginning ps-infra download.
+
+  if DEFINED GitHubAccessToken (
+    @Rem Download ps-infra with authentication and follow redirects.
+    curl -H "Accept: application/vnd.github.v3+json" -H "Authorization: token %GitHubAccessToken%" -L https://api.github.com/repos/%PSInfraOrg%/%PSInfraRepo%/zipball/%PSInfraTagOrBranch% > ps-infra.zip
+  ) else (
+    @Rem Download ps-infra and follow redirects.
+    curl -L https://github.com/%PSInfraOrg%/%PSInfraRepo%/archive/refs/%RefType%/%PSInfraTagOrBranch%.zip > ps-infra.zip
+  )
+  
+  @Rem Unarchive the .zip
+  tar -xmf ps-infra.zip || echo bad archive, contents: && type ps-infra.zip && exit 0
+
+  @Rem Rename the extracted, versioned, directory
+  move "EpicGames-PixelStreamingInfrastructure-*" "ps-infra"
+
+  @Rem Delete the downloaded zip
+  del ps-infra.zip
+)
+
+exit 0
+
+:print_help
+echo "Tool for fetching PixelStreaming Infrastructure"
+echo ""
+echo "Usage:"
+echo "  %0 [/h] [/b <branch>] [/t <tag>] [/p <personal access token>]"
+echo "Where:"
+echo "  /b      Specify a specific branch for the tool to download from repo (If this and /t are not set will default to recommended version)"
+echo "  /t      Specify a specific tag for the tool to download from repo (If this and /b are not set will default to recommended version)"
+echo "  /p      Specify a GitHub Personal Access Token to use to authorize requests (only necessary if repo is private)"
+echo "  /h      Display this help message"
+exit 1
\ No newline at end of file
diff --git a/Engine/Plugins/Media/PixelStreaming/Resources/get_ps_servers.sh b/Engine/Plugins/Media/PixelStreaming/Resources/get_ps_servers.sh
new file mode 100644
--- /dev/null
+++ b/Engine/Plugins/Media/PixelStreaming/Resources/get_ps_servers.sh
@@ -0,0 +1,74 @@
+#!/bin/bash
+# Copyright Epic Games, Inc. All Rights Reserved.
+
+BASH_LOCATION=$( cd -- "$( dirname -- "${BASH_SOURCE[0]}" )" &> /dev/null && pwd )
+
+pushd "${BASH_LOCATION}" > /dev/null
+
+print_help() {
+ echo "
+ Tool for fetching PixelStreaming Infrastructure.
+
+ Usage:
+   ${0} [-h] [-b <branch>] [-t <tag>] [-p <personal access token>]
+ Where:
+   -b      Specify a specific branch for the tool to download from repo (If this and -t are not set will default to recommended version)
+   -t      Specify a specific tag for the tool to download from repo (If this and -b are not set will default to recommended version)
+   -p      Specify a GitHub Personal Access Token to use to authorize requests (only necessary if repo is private)
+   -h      Display this help message
+"
+ exit 1
+}
+
+while(($#)) ; do
+  case "$1" in
+   -h ) print_help;;
+   -b ) PSInfraTagOrBranch="$2"; IsTag=0; shift 2;;
+   -t ) PSInfraTagOrBranch="$2"; IsTag=1; shift 2;;
+   -p ) GitHubAccessToken="$2"; shift 2;;
+   * ) echo "Unknown command: $1"; shift;;
+  esac
+ done
+
+# Name and version of ps-infra that we are downloading
+PSInfraOrg=EpicGames
+PSInfraRepo=PixelStreamingInfrastructure
+
+if [ -z "$PSInfraTagOrBranch" ]
+then
+    PSInfraTagOrBranch=v0.1.0-prerelease
+    IsTag=0
+fi
+
+if [ "$IsTag" -eq 1 ]
+then
+    RefType=tags
+else
+    RefType=heads
+fi
+
+# Look for a ps-infra directory next to this script
+if [ -d ps-infra ]
+then
+  echo "ps-infra directory found...skipping install."
+else
+  echo "ps-infra directory not found...beginning ps-infra download."
+
+  if [ ! -z "$GitHubAccessToken" ]
+  then
+    # Download ps-infra with authentication and follow redirects.
+    curl -H "Accept: application/vnd.github.v3+json" -H "Authorization: token $GitHubAccessToken" -L https://api.github.com/repos/$PSInfraOrg/$PSInfraRepo/tarball/$PSInfraTagOrBranch > ps-infra.tar.gz
+  else
+    # Download ps-infra and follow redirects.
+    curl -L https://github.com/$PSInfraOrg/$PSInfraRepo/archive/refs/$RefType/$PSInfraTagOrBranch.tar.gz > ps-infra.tar.gz
+  fi
+  
+  # Unarchive the .tar
+  tar -xmf ps-infra.tar.gz || $(echo "bad archive, contents:" && head --lines=20 ps-infra.tar.gz && exit 0)
+
+  # Rename the extracted, versioned, directory
+  mv EpicGames-PixelStreamingInfrastructure-* "ps-infra"
+
+  # Delete the downloaded tar
+  rm ps-infra.tar.gz
+fi
\ No newline at end of file
diff --git a/Engine/Plugins/Media/PixelStreaming/Source/PixelStreaming/PixelStreaming.Build.cs b/Engine/Plugins/Media/PixelStreaming/Source/PixelStreaming/PixelStreaming.Build.cs
--- a/Engine/Plugins/Media/PixelStreaming/Source/PixelStreaming/PixelStreaming.Build.cs
+++ b/Engine/Plugins/Media/PixelStreaming/Source/PixelStreaming/PixelStreaming.Build.cs
@@ -10,11 +10,11 @@ namespace UnrealBuildTool.Rules
 {
 	public class PixelStreaming : ModuleRules
 	{
-		const string PixelStreamingProgramsDirectory = "../../Samples/PixelStreaming";
+		const string PixelStreamingProgramsDirectory = "../../Engine/Plugins/Media/PixelStreaming/Resources";
 
 		private void AddFolder(string Folder)
 		{
-			string DirectoryToAdd = new DirectoryInfo(PixelStreamingProgramsDirectory + "/WebServers/" + Folder).FullName;
+			string DirectoryToAdd = new DirectoryInfo(PixelStreamingProgramsDirectory + "/ps-infra/" + Folder).FullName;
 
 			if (!Directory.Exists(DirectoryToAdd))
 			{
@@ -118,9 +118,12 @@ namespace UnrealBuildTool.Rules
 			// When we build a Game target we also package the servers with it as runtime dependencies
 			if(Target.Type == TargetType.Game)
 			{
-				//AddFolder("SignallingWebServer");
-				//AddFolder("Matchmaker");
-				//AddFolder("SFU");
+				AddFolder("SignallingWebServer");
+				AddFolder("Matchmaker");
+				AddFolder("SFU");
+
+				RuntimeDependencies.Add("$(PluginDir)/Resources/get_ps_servers.bat", StagedFileType.NonUFS);
+				RuntimeDependencies.Add("$(PluginDir)/Resources/get_ps_servers.sh", StagedFileType.NonUFS);
 			}
 		}
 	}
