Suppress error dialog when hardware encoders fail and fallback to software encoder instead.

diff --git a/Engine/Plugins/Media/PixelStreaming/Source/PixelStreaming/Private/PixelStreamingModule.cpp b/Engine/Plugins/Media/PixelStreaming/Source/PixelStreaming/Private/PixelStreamingModule.cpp
--- a/Engine/Plugins/Media/PixelStreaming/Source/PixelStreaming/Private/PixelStreamingModule.cpp
+++ b/Engine/Plugins/Media/PixelStreaming/Source/PixelStreaming/Private/PixelStreamingModule.cpp
@@ -385,12 +385,9 @@ namespace UE::PixelStreaming
 		if (Settings::CVarPixelStreamingEncoderCodec.GetValueOnAnyThread() == "H264"
 			&& !AVEncoder::FVideoEncoderFactory::Get().HasEncoderForCodec(AVEncoder::ECodecType::H264))
 		{
-			FText TitleText = FText::FromString(TEXT("Pixel Streaming Plugin"));
-			FString ErrorString = TEXT("No compatible GPU found, or failed to load their respective encoder libraries");
-			FText ErrorText = FText::FromString(ErrorString);
-			FMessageDialog::Open(EAppMsgType::Ok, ErrorText, &TitleText);
-			UE_LOG(LogPixelStreaming, Error, TEXT("%s"), *ErrorString);
-			bCompatible = false;
+			UE_LOG(LogPixelStreaming, Warning, TEXT("Could not setup hardware encoder for H.264. This is usually a driver issue, try reinstalling your drivers."));
+			UE_LOG(LogPixelStreaming, Warning, TEXT("Falling back to VP8 software video encoding."));
+			Settings::CVarPixelStreamingEncoderCodec.AsVariable()->Set(TEXT("VP8"), ECVF_SetByCommandline);
 		}
 
 		return bCompatible;
