Fix: Static analysis issues in PixelStreaming

diff --git a/Engine/Plugins/Media/PixelStreaming/Source/PixelStreaming/Private/WebRTCIncludes.h b/Engine/Plugins/Media/PixelStreaming/Source/PixelStreaming/Private/WebRTCIncludes.h
--- a/Engine/Plugins/Media/PixelStreaming/Source/PixelStreaming/Private/WebRTCIncludes.h
+++ b/Engine/Plugins/Media/PixelStreaming/Source/PixelStreaming/Private/WebRTCIncludes.h
@@ -3,6 +3,7 @@
 #pragma once
 
 #include "CoreMinimal.h"
+#include "PixelStreamingWebRTCIncludes.h"
 
 THIRD_PARTY_INCLUDES_START
 
@@ -21,9 +22,7 @@ THIRD_PARTY_INCLUDES_START
 
 #include "api/rtp_receiver_interface.h"
 #include "api/media_types.h"
-#include "api/media_stream_interface.h"
-#include "api/data_channel_interface.h"
-#include "api/peer_connection_interface.h"
+
 #include "api/create_peerconnection_factory.h"
 #include "api/task_queue/default_task_queue_factory.h"
 #include "api/audio_codecs/audio_format.h"
@@ -37,8 +36,7 @@ THIRD_PARTY_INCLUDES_START
 #include "api/video_codecs/video_encoder_software_fallback_wrapper.h"
 #include "api/video/video_frame.h"
 #include "api/video/video_rotation.h"
-#include "api/video/video_frame_buffer.h"
-#include "api/video/i420_buffer.h"
+
 #include "api/video/video_sink_interface.h"
 
 #include "rtc_base/thread.h"
@@ -98,13 +96,6 @@ THIRD_PARTY_INCLUDES_START
 
 	#include "Windows/PostWindowsApi.h"
 	#include "Windows/HideWindowsPlatformTypes.h"
-
-#else
-
-	#ifdef PF_MAX
-		#undef PF_MAX
-	#endif
-
 #endif //PLATFORM_WINDOWS
 
 THIRD_PARTY_INCLUDES_END
diff --git a/Engine/Plugins/Media/PixelStreaming/Source/PixelStreamingEditor/Private/PixelStreamingToolbar.cpp b/Engine/Plugins/Media/PixelStreaming/Source/PixelStreamingEditor/Private/PixelStreamingToolbar.cpp
--- a/Engine/Plugins/Media/PixelStreaming/Source/PixelStreamingEditor/Private/PixelStreamingToolbar.cpp
+++ b/Engine/Plugins/Media/PixelStreaming/Source/PixelStreamingEditor/Private/PixelStreamingToolbar.cpp
@@ -206,7 +206,7 @@ namespace UE::PixelStreaming
                                 }
                                 MenuBuilder.EndSection();
 
-                                MenuBuilder.BeginSection("Advanced", LOCTEXT("PixelStreamingSettings", "Advanced"));
+                                MenuBuilder.BeginSection("Advanced", LOCTEXT("PixelStreamingAdvancedSettings", "Advanced"));
                                 {
                                     MenuBuilder.AddSubMenu(
                                         LOCTEXT("PixelStreamingCodecSettings", "Select Codec"),
