Change logging category for some PixelStreaming messages that shouldn't show as warnings

diff --git a/Engine/Plugins/Media/PixelStreaming/Source/PixelStreaming/Private/PixelStreamingModule.cpp b/Engine/Plugins/Media/PixelStreaming/Source/PixelStreaming/Private/PixelStreamingModule.cpp
--- a/Engine/Plugins/Media/PixelStreaming/Source/PixelStreaming/Private/PixelStreamingModule.cpp
+++ b/Engine/Plugins/Media/PixelStreaming/Source/PixelStreaming/Private/PixelStreamingModule.cpp
@@ -373,7 +373,7 @@ namespace UE::PixelStreaming
 			else
 			{
 				SignallingServerURL = Settings::GetDefaultSignallingURL();
-				UE_LOG(LogPixelStreaming, Warning, TEXT("-PixelStreamingURL was not specified on the command line, using the default connection url: %s"), *SignallingServerURL);
+				UE_LOG(LogPixelStreaming, Log, TEXT("-PixelStreamingURL was not specified on the command line, using the default connection url: %s"), *SignallingServerURL);
 			}
 		}
 
diff --git a/Engine/Plugins/Media/PixelStreaming/Source/PixelStreaming/Private/PixelStreamingSignallingConnection.cpp b/Engine/Plugins/Media/PixelStreaming/Source/PixelStreaming/Private/PixelStreamingSignallingConnection.cpp
--- a/Engine/Plugins/Media/PixelStreaming/Source/PixelStreaming/Private/PixelStreamingSignallingConnection.cpp
+++ b/Engine/Plugins/Media/PixelStreaming/Source/PixelStreaming/Private/PixelStreamingSignallingConnection.cpp
@@ -231,7 +231,7 @@ void FPixelStreamingSignallingConnection::OnConnected()
 
 void FPixelStreamingSignallingConnection::OnConnectionError(const FString& Error)
 {
-	UE_LOG(LogPixelStreamingSS, Warning, TEXT("Failed to connect to SS - signalling server may not be up yet. Message: \"%s\""), *Error);
+	UE_LOG(LogPixelStreamingSS, Log, TEXT("Failed to connect to SS - signalling server may not be up yet. Message: \"%s\""), *Error);
 
 	Observer.OnSignallingError(Error);
 
