Disable PixelStreaming infra packaging temporarily to unblock automation.

diff --git a/Engine/Plugins/Media/PixelStreaming/Source/PixelStreaming/PixelStreaming.Build.cs b/Engine/Plugins/Media/PixelStreaming/Source/PixelStreaming/PixelStreaming.Build.cs
--- a/Engine/Plugins/Media/PixelStreaming/Source/PixelStreaming/PixelStreaming.Build.cs
+++ b/Engine/Plugins/Media/PixelStreaming/Source/PixelStreaming/PixelStreaming.Build.cs
@@ -114,15 +114,15 @@ namespace UnrealBuildTool.Rules
 			}
 
 			// When we build a Game target we also package the servers with it as runtime dependencies
-			if(Target.Type == TargetType.Game)
-			{
-				AddFolder("SignallingWebServer");
-				AddFolder("Matchmaker");
-				AddFolder("SFU");
-
-				RuntimeDependencies.Add("$(EngineDir)/../get_ps_servers.bat", "$(PluginDir)/Resources/get_ps_servers.bat", StagedFileType.SystemNonUFS);
-				RuntimeDependencies.Add("$(EngineDir)/../get_ps_servers.sh", "$(PluginDir)/Resources/get_ps_servers.sh", StagedFileType.SystemNonUFS);
-			}
+			// if(Target.Type == TargetType.Game)
+			// {
+			// 	AddFolder("SignallingWebServer");
+			// 	AddFolder("Matchmaker");
+			// 	AddFolder("SFU");
+
+			// 	RuntimeDependencies.Add("$(EngineDir)/../get_ps_servers.bat", "$(PluginDir)/Resources/get_ps_servers.bat", StagedFileType.SystemNonUFS);
+			// 	RuntimeDependencies.Add("$(EngineDir)/../get_ps_servers.sh", "$(PluginDir)/Resources/get_ps_servers.sh", StagedFileType.SystemNonUFS);
+			// }
 		}
 	}
 }
