Fix build warnings due to ps infra packaging location

diff --git a/Engine/Plugins/Media/PixelStreaming/Resources/WebServers/get_ps_servers.bat b/Engine/Plugins/Media/PixelStreaming/Resources/WebServers/get_ps_servers.bat
new file mode 100755
--- /dev/null
+++ b/Engine/Plugins/Media/PixelStreaming/Resources/WebServers/get_ps_servers.bat
@@ -0,0 +1,95 @@
+@Rem Copyright Epic Games, Inc. All Rights Reserved.
+
+@echo off
+
+@Rem Set script location as working directory for commands.
+pushd "%~dp0"
+
+:arg_loop_start
+SET ARG=%1
+if DEFINED ARG (
+    if "%ARG%"=="/h" (
+        goto print_help
+    )
+    if "%ARG%"=="/p" (
+        SET GitHubAccessToken=%2
+        SHIFT
+    )
+    if "%ARG%"=="/b" (
+        SET PSInfraTagOrBranch=%2
+        SET IsTag=0
+        SHIFT
+    )
+    if "%ARG%"=="/t" (
+        SET PSInfraTagOrBranch=%2
+        SET IsTag=1
+        SHIFT
+    )
+    SHIFT
+    goto arg_loop_start
+)
+
+@Rem Name and version of ps-infra that we are downloading
+SET PSInfraOrg=EpicGames
+SET PSInfraRepo=PixelStreamingInfrastructure
+
+if NOT DEFINED PSInfraTagOrBranch (
+    SET PSInfraTagOrBranch=v0.1.0-prerelease
+    SET IsTag=1
+)
+
+if %IsTag%==1 (
+  SET RefType=tags
+) else (
+  SET RefType=heads
+)
+
+@Rem Look for a SignallingWebServer directory next to this script
+if exist SignallingWebServer\ (
+  echo SignallingWebServer directory found...skipping install.
+) else (
+  echo SignallingWebServer directory not found...beginning ps-infra download.
+
+  if DEFINED GitHubAccessToken (
+    @Rem Download ps-infra with authentication and follow redirects.
+    curl -H "Accept: application/vnd.github.v3+json" -H "Authorization: token %GitHubAccessToken%" -L https://api.github.com/repos/%PSInfraOrg%/%PSInfraRepo%/zipball/%PSInfraTagOrBranch% > ps-infra.zip
+  ) else (
+    @Rem Download ps-infra and follow redirects.
+    curl -L https://github.com/%PSInfraOrg%/%PSInfraRepo%/archive/refs/%RefType%/%PSInfraTagOrBranch%.zip > ps-infra.zip
+  )
+  
+  @Rem Unarchive the .zip
+  tar -xmf ps-infra.zip || echo bad archive, contents: && type ps-infra.zip && exit 0
+
+  @Rem Rename the extracted, versioned, directory
+  for /d %%i in ("PixelStreamingInfrastructure-*") do (
+    for /d %%j in ("%%i/*") do (
+      echo "%%i\%%j"
+      move "%%i\%%j" .
+    )
+    for %%j in ("%%i/*") do (
+      echo "%%i\%%j"
+      move "%%i\%%j" .
+    )
+
+    echo "%%i"
+    rmdir /s /q "%%i"
+  )
+
+  @Rem Delete the downloaded zip
+  del ps-infra.zip
+)
+
+exit 0
+
+:print_help
+echo "Tool for fetching PixelStreaming Infrastructure"
+echo ""
+echo "Usage:"
+echo "  %0 [/h] [/b <branch>] [/t <tag>] [/p <personal access token>]"
+echo "Where:"
+echo "  /b      Specify a specific branch for the tool to download from repo (If this and /t are not set will default to recommended version)"
+echo "  /t      Specify a specific tag for the tool to download from repo (If this and /b are not set will default to recommended version)"
+echo "  /p      Specify a GitHub Personal Access Token to use to authorize requests (only necessary if repo is private)"
+echo "  /h      Display this help message"
+exit 1
\ No newline at end of file
diff --git a/Engine/Plugins/Media/PixelStreaming/Resources/WebServers/get_ps_servers.sh b/Engine/Plugins/Media/PixelStreaming/Resources/WebServers/get_ps_servers.sh
new file mode 100644
--- /dev/null
+++ b/Engine/Plugins/Media/PixelStreaming/Resources/WebServers/get_ps_servers.sh
@@ -0,0 +1,75 @@
+#!/bin/bash
+# Copyright Epic Games, Inc. All Rights Reserved.
+
+BASH_LOCATION=$( cd -- "$( dirname -- "${BASH_SOURCE[0]}" )" &> /dev/null && pwd )
+
+pushd "${BASH_LOCATION}" > /dev/null
+
+print_help() {
+ echo "
+ Tool for fetching PixelStreaming Infrastructure.
+
+ Usage:
+   ${0} [-h] [-b <branch>] [-t <tag>] [-p <personal access token>]
+ Where:
+   -b      Specify a specific branch for the tool to download from repo (If this and -t are not set will default to recommended version)
+   -t      Specify a specific tag for the tool to download from repo (If this and -b are not set will default to recommended version)
+   -p      Specify a GitHub Personal Access Token to use to authorize requests (only necessary if repo is private)
+   -h      Display this help message
+"
+ exit 1
+}
+
+while(($#)) ; do
+  case "$1" in
+   -h ) print_help;;
+   -b ) PSInfraTagOrBranch="$2"; IsTag=0; shift 2;;
+   -t ) PSInfraTagOrBranch="$2"; IsTag=1; shift 2;;
+   -p ) GitHubAccessToken="$2"; shift 2;;
+   * ) echo "Unknown command: $1"; shift;;
+  esac
+ done
+
+# Name and version of ps-infra that we are downloading
+PSInfraOrg=EpicGames
+PSInfraRepo=PixelStreamingInfrastructure
+
+if [ -z "$PSInfraTagOrBranch" ]
+then
+    PSInfraTagOrBranch=v0.1.0-prerelease
+    IsTag=0
+fi
+
+if [ "$IsTag" -eq 1 ]
+then
+    RefType=tags
+else
+    RefType=heads
+fi
+
+# Look for a SignallingWebServer directory next to this script
+if [ -d SignallingWebServer ]
+then
+  echo "SignallingWebServer directory found...skipping install."
+else
+  echo "SignallingWebServer directory not found...beginning ps-infra download."
+
+  if [ ! -z "$GitHubAccessToken" ]
+  then
+    # Download ps-infra with authentication and follow redirects.
+    curl -H "Accept: application/vnd.github.v3+json" -H "Authorization: token $GitHubAccessToken" -L https://api.github.com/repos/$PSInfraOrg/$PSInfraRepo/tarball/$PSInfraTagOrBranch > ps-infra.tar.gz
+  else
+    # Download ps-infra and follow redirects.
+    curl -L https://github.com/$PSInfraOrg/$PSInfraRepo/archive/refs/$RefType/$PSInfraTagOrBranch.tar.gz > ps-infra.tar.gz
+  fi
+  
+  # Unarchive the .tar
+  tar -xmf ps-infra.tar.gz || $(echo "bad archive, contents:" && head --lines=20 ps-infra.tar.gz && exit 0)
+
+  # Move the server folders into the current directory (WebServers) and delete the original directory
+  mv EpicGames-PixelStreamingInfrastructure-*/* .
+  rm -rf EpicGames-PixelStreamingInfrastructure-*
+
+  # Delete the downloaded tar
+  rm ps-infra.tar.gz
+fi
\ No newline at end of file
diff --git a/Engine/Plugins/Media/PixelStreaming/Resources/get_ps_servers.bat b/Engine/Plugins/Media/PixelStreaming/Resources/get_ps_servers.bat
deleted file mode 100755
diff --git a/Engine/Plugins/Media/PixelStreaming/Resources/get_ps_servers.sh b/Engine/Plugins/Media/PixelStreaming/Resources/get_ps_servers.sh
deleted file mode 100644
diff --git a/Engine/Plugins/Media/PixelStreaming/Source/PixelStreaming/PixelStreaming.Build.cs b/Engine/Plugins/Media/PixelStreaming/Source/PixelStreaming/PixelStreaming.Build.cs
--- a/Engine/Plugins/Media/PixelStreaming/Source/PixelStreaming/PixelStreaming.Build.cs
+++ b/Engine/Plugins/Media/PixelStreaming/Source/PixelStreaming/PixelStreaming.Build.cs
@@ -10,7 +10,7 @@ namespace UnrealBuildTool.Rules
 {
 	public class PixelStreaming : ModuleRules
 	{
-		string PixelStreamingProgramsOutputDirectory = "$(EngineDir)/../WebServers";
+		string PixelStreamingProgramsOutputDirectory = "$(ProjectDir)/Samples/PixelStreaming/WebServers";
 		const string PixelStreamingProgramsDirectory = "../../Engine/Plugins/Media/PixelStreaming/Resources/WebServers";
 
 		private void AddFolder(string Folder)
@@ -114,15 +114,15 @@ namespace UnrealBuildTool.Rules
 			}
 
 			// When we build a Game target we also package the servers with it as runtime dependencies
-			// if(Target.Type == TargetType.Game)
-			// {
-			// 	AddFolder("SignallingWebServer");
-			// 	AddFolder("Matchmaker");
-			// 	AddFolder("SFU");
-
-			// 	RuntimeDependencies.Add("$(EngineDir)/../get_ps_servers.bat", "$(PluginDir)/Resources/get_ps_servers.bat", StagedFileType.SystemNonUFS);
-			// 	RuntimeDependencies.Add("$(EngineDir)/../get_ps_servers.sh", "$(PluginDir)/Resources/get_ps_servers.sh", StagedFileType.SystemNonUFS);
-			// }
+			if(Target.Type == TargetType.Game)
+			{
+				AddFolder("SignallingWebServer");
+				AddFolder("Matchmaker");
+				AddFolder("SFU");
+
+				RuntimeDependencies.Add("$(ProjectDir)/Samples/PixelStreaming/WebServers/get_ps_servers.bat", "$(PluginDir)/Resources/WebServers/get_ps_servers.bat", StagedFileType.NonUFS);
+				RuntimeDependencies.Add("$(ProjectDir)/Samples/PixelStreaming/WebServers/get_ps_servers.sh", "$(PluginDir)/Resources/WebServers/get_ps_servers.sh", StagedFileType.NonUFS);
+			}
 		}
 	}
 }
