Temporarily prevent packaging of ps infra files until a solution for them packaging their way into places they don't belong is found

diff --git a/Engine/Plugins/Media/PixelStreaming/Source/PixelStreaming/PixelStreaming.Build.cs b/Engine/Plugins/Media/PixelStreaming/Source/PixelStreaming/PixelStreaming.Build.cs
--- a/Engine/Plugins/Media/PixelStreaming/Source/PixelStreaming/PixelStreaming.Build.cs
+++ b/Engine/Plugins/Media/PixelStreaming/Source/PixelStreaming/PixelStreaming.Build.cs
@@ -114,15 +114,15 @@ namespace UnrealBuildTool.Rules
 			}
 
 			// When we build a Game target we also package the servers with it as runtime dependencies
-			if(Target.Type == TargetType.Game)
-			{
-				AddFolder("SignallingWebServer");
-				AddFolder("Matchmaker");
-				AddFolder("SFU");
-
-				RuntimeDependencies.Add("$(ProjectDir)/Samples/PixelStreaming/WebServers/get_ps_servers.bat", "$(PluginDir)/Resources/WebServers/get_ps_servers.bat", StagedFileType.NonUFS);
-				RuntimeDependencies.Add("$(ProjectDir)/Samples/PixelStreaming/WebServers/get_ps_servers.sh", "$(PluginDir)/Resources/WebServers/get_ps_servers.sh", StagedFileType.NonUFS);
-			}
+			// if(Target.Type == TargetType.Game)
+			// {
+			// 	AddFolder("SignallingWebServer");
+			// 	AddFolder("Matchmaker");
+			// 	AddFolder("SFU");
+
+			// 	RuntimeDependencies.Add("$(ProjectDir)/Samples/PixelStreaming/WebServers/get_ps_servers.bat", "$(PluginDir)/Resources/WebServers/get_ps_servers.bat", StagedFileType.NonUFS);
+			// 	RuntimeDependencies.Add("$(ProjectDir)/Samples/PixelStreaming/WebServers/get_ps_servers.sh", "$(PluginDir)/Resources/WebServers/get_ps_servers.sh", StagedFileType.NonUFS);
+			// }
 		}
 	}
 }
