Cherrypicked from Dev-TensorWorks:

Fix: multipeer VPX not working introduced in CL20588532.
Main cause for this was introduced in FVideoSource::InputFrame(***) where we check if this frame adapter should generate frames. This shouldgenerate frames function did not account for the different codecs and instead assumed that H264 was used all the time.

diff --git a/Engine/Plugins/Media/PixelStreaming/Source/PixelStreaming/Private/Streamer.cpp b/Engine/Plugins/Media/PixelStreaming/Source/PixelStreaming/Private/Streamer.cpp
--- a/Engine/Plugins/Media/PixelStreaming/Source/PixelStreaming/Private/Streamer.cpp
+++ b/Engine/Plugins/Media/PixelStreaming/Source/PixelStreaming/Private/Streamer.cpp
@@ -21,6 +21,7 @@
 #include "IImageWrapperModule.h"
 #include "Framework/Application/SlateApplication.h"
 #include "UtilsRender.h"
+#include "PixelStreamingCodec.h"
 
 namespace UE::PixelStreaming
 {
@@ -858,7 +859,22 @@ namespace UE::PixelStreaming
 
 	bool FStreamer::ShouldPeerGenerateFrames(FPixelStreamingPlayerId PlayerId) const
 	{
-		return PlayerId != INVALID_PLAYER_ID && (PlayerId == QualityControllingId || PlayerId == SFUPlayerId);
+		EPixelStreamingCodec Codec = Settings::GetSelectedCodec();
+		switch(Codec)
+		{
+			case EPixelStreamingCodec::H264:
+				return PlayerId != INVALID_PLAYER_ID && (PlayerId == QualityControllingId || PlayerId == SFUPlayerId);
+			break;
+			case EPixelStreamingCodec::VP8:
+			case EPixelStreamingCodec::VP9:
+				return PlayerId != INVALID_PLAYER_ID;
+			break;
+			default:
+				// There should be a case for every Codec type, so this should never happen.
+				checkNoEntry();
+				break;
+		}
+		return false;
 	}
 
 	void FStreamer::SetQualityController(FPixelStreamingPlayerId PlayerId)
