Update version of PS Infra that `get_ps` scripts retrieve, and add flag to allow retrieving version PS Infra recommended for a specific version of UE

diff --git a/Engine/Plugins/Media/PixelStreaming/Resources/WebServers/get_ps_servers.bat b/Engine/Plugins/Media/PixelStreaming/Resources/WebServers/get_ps_servers.bat
--- a/Engine/Plugins/Media/PixelStreaming/Resources/WebServers/get_ps_servers.bat
+++ b/Engine/Plugins/Media/PixelStreaming/Resources/WebServers/get_ps_servers.bat
@@ -11,8 +11,8 @@ if DEFINED ARG (
     if "%ARG%"=="/h" (
         goto print_help
     )
-    if "%ARG%"=="/p" (
-        SET GitHubAccessToken=%2
+    if "%ARG%"=="/v" (
+        SET UEVersion=%2
         SHIFT
     )
     if "%ARG%"=="/b" (
@@ -33,11 +33,29 @@ if DEFINED ARG (
 SET PSInfraOrg=EpicGames
 SET PSInfraRepo=PixelStreamingInfrastructure
 
+@Rem If a UE version is supplied set the right branch or tag to fetch for that version of UE
+if DEFINED UEVersion (
+  if "%UEVersion%"=="4.26" (
+    SET PSInfraTagOrBranch=UE4.26
+    SET IsTag=0
+  )
+  if "%UEVersion%"=="4.27" (
+    SET PSInfraTagOrBranch=UE4.27
+    SET IsTag=0
+  )
+  if "%UEVersion%"=="5.0" (
+    SET PSInfraTagOrBranch=UE5.0
+    SET IsTag=0
+  )
+)
+
+@Rem If no arguments select a specific version, fetch the appropriate default
 if NOT DEFINED PSInfraTagOrBranch (
-    SET PSInfraTagOrBranch=v0.1.0-prerelease
-    SET IsTag=1
+    SET PSInfraTagOrBranch=master
+    SET IsTag=0
 )
 
+@Rem Whether the named reference is a tag or a branch affects the URL we fetch it on
 if %IsTag%==1 (
   SET RefType=tags
 ) else (
@@ -50,13 +68,8 @@ if exist SignallingWebServer\ (
 ) else (
   echo SignallingWebServer directory not found...beginning ps-infra download.
 
-  if DEFINED GitHubAccessToken (
-    @Rem Download ps-infra with authentication and follow redirects.
-    curl -H "Accept: application/vnd.github.v3+json" -H "Authorization: token %GitHubAccessToken%" -L https://api.github.com/repos/%PSInfraOrg%/%PSInfraRepo%/zipball/%PSInfraTagOrBranch% > ps-infra.zip
-  ) else (
-    @Rem Download ps-infra and follow redirects.
-    curl -L https://github.com/%PSInfraOrg%/%PSInfraRepo%/archive/refs/%RefType%/%PSInfraTagOrBranch%.zip > ps-infra.zip
-  )
+  @Rem Download ps-infra and follow redirects.
+  curl -L https://github.com/%PSInfraOrg%/%PSInfraRepo%/archive/refs/%RefType%/%PSInfraTagOrBranch%.zip > ps-infra.zip
   
   @Rem Unarchive the .zip
   tar -xmf ps-infra.zip || echo bad archive, contents: && type ps-infra.zip && exit 0
@@ -83,13 +96,15 @@ if exist SignallingWebServer\ (
 exit 0
 
 :print_help
-echo "Tool for fetching PixelStreaming Infrastructure"
-echo ""
-echo "Usage:"
-echo "  %0 [/h] [/b <branch>] [/t <tag>] [/p <personal access token>]"
-echo "Where:"
-echo "  /b      Specify a specific branch for the tool to download from repo (If this and /t are not set will default to recommended version)"
-echo "  /t      Specify a specific tag for the tool to download from repo (If this and /b are not set will default to recommended version)"
-echo "  /p      Specify a GitHub Personal Access Token to use to authorize requests (only necessary if repo is private)"
-echo "  /h      Display this help message"
+echo.
+echo  Tool for fetching PixelStreaming Infrastructure. If no flags are set specifying a version to fetch,
+echo  the recommended version will be chosen as a default.
+echo.
+echo  Usage:
+echo    %~n0%~x0 [^/h] [^/v ^<UE version^>] [^/b ^<branch^>] [^/t ^<tag^>]
+echo  Where:
+echo    /v      Specify a version of Unreal Engine to download the recommended release for
+echo    /b      Specify a specific branch for the tool to download from repo
+echo    /t      Specify a specific tag for the tool to download from repo
+echo    /h      Display this help message
 exit 1
\ No newline at end of file
diff --git a/Engine/Plugins/Media/PixelStreaming/Resources/WebServers/get_ps_servers.sh b/Engine/Plugins/Media/PixelStreaming/Resources/WebServers/get_ps_servers.sh
--- a/Engine/Plugins/Media/PixelStreaming/Resources/WebServers/get_ps_servers.sh
+++ b/Engine/Plugins/Media/PixelStreaming/Resources/WebServers/get_ps_servers.sh
@@ -7,14 +7,16 @@ pushd "${BASH_LOCATION}" > /dev/null
 
 print_help() {
  echo "
- Tool for fetching PixelStreaming Infrastructure.
+ Tool for fetching PixelStreaming Infrastructure. If no flags are set specifying a version to fetch, 
+ the recommended version will be chosen as a default.
 
  Usage:
-   ${0} [-h] [-b <branch>] [-t <tag>] [-p <personal access token>]
+   ${0} [-h] [-v <UE version>] [-b <branch>] [-t <tag>]
  Where:
-   -b      Specify a specific branch for the tool to download from repo (If this and -t are not set will default to recommended version)
-   -t      Specify a specific tag for the tool to download from repo (If this and -b are not set will default to recommended version)
-   -p      Specify a GitHub Personal Access Token to use to authorize requests (only necessary if repo is private)
+   -v      Specify a version of Unreal Engine to download the recommended 
+           release for
+   -b      Specify a specific branch for the tool to download from repo
+   -t      Specify a specific tag for the tool to download from repo
    -h      Display this help message
 "
  exit 1
@@ -23,9 +25,9 @@ print_help() {
 while(($#)) ; do
   case "$1" in
    -h ) print_help;;
+   -v ) UEVersion="$2"; shift 2;;
    -b ) PSInfraTagOrBranch="$2"; IsTag=0; shift 2;;
    -t ) PSInfraTagOrBranch="$2"; IsTag=1; shift 2;;
-   -p ) GitHubAccessToken="$2"; shift 2;;
    * ) echo "Unknown command: $1"; shift;;
   esac
  done
@@ -34,17 +36,39 @@ while(($#)) ; do
 PSInfraOrg=EpicGames
 PSInfraRepo=PixelStreamingInfrastructure
 
-if [ -z "$PSInfraTagOrBranch" ]
+# If a UE version is supplied set the right branch or tag to fetch for that version of UE
+if [ ! -z "$UEVersion" ]
 then
-    PSInfraTagOrBranch=v0.1.0-prerelease
+  if [ "$UEVersion" = "4.26" ]
+  then
+    PSInfraTagOrBranch=UE4.26
+    IsTag=0
+  fi
+  if [ "$UEVersion" = "4.27" ]
+  then
+    PSInfraTagOrBranch=UE4.27
+    IsTag=0
+  fi
+  if [ "$UEVersion" = "5.0" ]
+  then
+    PSInfraTagOrBranch=UE5.0
     IsTag=0
+  fi
 fi
 
+# If no arguments select a specific version, fetch the appropriate default
+if [ -z "$PSInfraTagOrBranch" ]
+then
+  PSInfraTagOrBranch=master
+  IsTag=0
+fi
+
+# Whether the named reference is a tag or a branch affects the URL we fetch it on
 if [ "$IsTag" -eq 1 ]
 then
-    RefType=tags
+  RefType=tags
 else
-    RefType=heads
+  RefType=heads
 fi
 
 # Look for a SignallingWebServer directory next to this script
@@ -54,21 +78,15 @@ then
 else
   echo "SignallingWebServer directory not found...beginning ps-infra download."
 
-  if [ ! -z "$GitHubAccessToken" ]
-  then
-    # Download ps-infra with authentication and follow redirects.
-    curl -H "Accept: application/vnd.github.v3+json" -H "Authorization: token $GitHubAccessToken" -L https://api.github.com/repos/$PSInfraOrg/$PSInfraRepo/tarball/$PSInfraTagOrBranch > ps-infra.tar.gz
-  else
-    # Download ps-infra and follow redirects.
-    curl -L https://github.com/$PSInfraOrg/$PSInfraRepo/archive/refs/$RefType/$PSInfraTagOrBranch.tar.gz > ps-infra.tar.gz
-  fi
+  # Download ps-infra and follow redirects.
+  curl -L https://github.com/$PSInfraOrg/$PSInfraRepo/archive/refs/$RefType/$PSInfraTagOrBranch.tar.gz > ps-infra.tar.gz
   
   # Unarchive the .tar
   tar -xmf ps-infra.tar.gz || $(echo "bad archive, contents:" && head --lines=20 ps-infra.tar.gz && exit 0)
 
   # Move the server folders into the current directory (WebServers) and delete the original directory
-  mv EpicGames-PixelStreamingInfrastructure-*/* .
-  rm -rf EpicGames-PixelStreamingInfrastructure-*
+  mv PixelStreamingInfrastructure-*/* .
+  rm -rf PixelStreamingInfrastructure-*
 
   # Delete the downloaded tar
   rm ps-infra.tar.gz
