Update the PixelStreamingMessageHandler to use the newer version of OnController functions that take an FInputDeviceId instead of the older soon to be deprecated versions

#rnx

diff --git a/Engine/Plugins/Media/PixelStreaming/Source/PixelStreaming/Private/PixelStreamingMessageHandler.cpp b/Engine/Plugins/Media/PixelStreaming/Source/PixelStreaming/Private/PixelStreamingMessageHandler.cpp
--- a/Engine/Plugins/Media/PixelStreaming/Source/PixelStreaming/Private/PixelStreamingMessageHandler.cpp
+++ b/Engine/Plugins/Media/PixelStreaming/Source/PixelStreaming/Private/PixelStreamingMessageHandler.cpp
@@ -48,9 +48,15 @@ namespace UE::PixelStreaming
         BIND_PLAYBACK_HANDLER((uint8) EToStreamerMsg::TouchMove, &FPixelStreamingMessageHandler::HandleOnTouchMoved);
         BIND_PLAYBACK_HANDLER((uint8) EToStreamerMsg::TouchEnd, &FPixelStreamingMessageHandler::HandleOnTouchEnded);
 
+    	PRAGMA_DISABLE_DEPRECATION_WARNINGS
         BIND_PLAYBACK_HANDLER((uint8) EToStreamerMsg::GamepadAnalog, &FPixelStreamingMessageHandler::HandleOnControllerAnalog);
         BIND_PLAYBACK_HANDLER((uint8) EToStreamerMsg::GamepadButtonPressed, &FPixelStreamingMessageHandler::HandleOnControllerButtonPressed);
         BIND_PLAYBACK_HANDLER((uint8) EToStreamerMsg::GamepadButtonReleased, &FPixelStreamingMessageHandler::HandleOnControllerButtonReleased);
+		PRAGMA_ENABLE_DEPRECATION_WARNINGS
+    	
+    	BIND_PLAYBACK_HANDLER((uint8) EToStreamerMsg::GamepadAnalog, &FPixelStreamingMessageHandler::HandleOnControllerAnalogWithPlatformUser);
+    	BIND_PLAYBACK_HANDLER((uint8) EToStreamerMsg::GamepadButtonPressed, &FPixelStreamingMessageHandler::HandleOnControllerButtonPressedWithPlatformUser);
+    	BIND_PLAYBACK_HANDLER((uint8) EToStreamerMsg::GamepadButtonReleased, &FPixelStreamingMessageHandler::HandleOnControllerButtonReleasedWithPlatformUser);
 
         BIND_PLAYBACK_HANDLER((uint8) EToStreamerMsg::MouseEnter, &FPixelStreamingMessageHandler::HandleOnMouseEnter);
         BIND_PLAYBACK_HANDLER((uint8) EToStreamerMsg::MouseLeave, &FPixelStreamingMessageHandler::HandleOnMouseLeave);
@@ -281,7 +287,8 @@ namespace UE::PixelStreaming
             FSlateApplication::Get().OverridePlatformApplication(PixelStreamerApplicationWrapper->WrappedApplication);
         }
     }
-    
+
+	PRAGMA_DISABLE_DEPRECATION_WARNINGS
     void FPixelStreamingMessageHandler::HandleOnControllerAnalog(FMemoryReader& Ar)
     {
         TPayloadThreeParam<uint8, uint8, double> Payload(Ar);
@@ -316,6 +323,45 @@ namespace UE::PixelStreaming
         UE_LOG(LogPixelStreamingMessageHandler, Verbose, TEXT("GAMEPAD_RELEASED: ControllerId = %d; KeyName = %s;"), ControllerId, *Button.ToString());
         MessageHandler->OnControllerButtonReleased(Button, ControllerId, false);
     }
+	PRAGMA_ENABLE_DEPRECATION_WARNINGS
+
+	void FPixelStreamingMessageHandler::HandleOnControllerAnalogWithPlatformUser(FMemoryReader& Ar)
+	{
+    	TPayloadFourParam<uint8, int32, int32, float> Payload(Ar);
+
+    	FGamepadKeyNames::Type Button = ConvertAxisIndexToGamepadAxis(Payload.Param1);
+    	FPlatformUserId UserId = FPlatformUserId::CreateFromInternalId(Payload.Param2);
+    	FInputDeviceId DeviceId = FInputDeviceId::CreateFromInternalId(Payload.Param3);
+    	float AnalogValue = Payload.Param4;
+
+    	UE_LOG(LogPixelStreamingMessageHandler, Verbose, TEXT("GAMEPAD_ANALOG: UserId = %d; DeviceId = %d; KeyName = %s; AnalogValue = %.4f;"), UserId.GetInternalId(), DeviceId.GetId(), *Button.ToString(), AnalogValue);
+    	MessageHandler->OnControllerAnalog(Button, UserId, DeviceId, AnalogValue);
+	}
+	
+	void FPixelStreamingMessageHandler::HandleOnControllerButtonPressedWithPlatformUser(FMemoryReader& Ar)
+    {
+    	TPayloadFourParam<uint8, int32, int32, bool> Payload(Ar);
+
+    	FGamepadKeyNames::Type Button = ConvertAxisIndexToGamepadAxis(Payload.Param1);
+    	FPlatformUserId UserId = FPlatformUserId::CreateFromInternalId(Payload.Param2);
+    	FInputDeviceId DeviceId = FInputDeviceId::CreateFromInternalId(Payload.Param3);
+    	bool bIsRepeat = Payload.Param4;
+
+    	UE_LOG(LogPixelStreamingMessageHandler, Verbose, TEXT("GAMEPAD_PRESSED: UserId = %d; DeviceId = %d; KeyName = %s; IsRepeat = %s;"), UserId.GetInternalId(), DeviceId.GetId(), *Button.ToString(), bIsRepeat ? TEXT("True") : TEXT("False"));
+    	MessageHandler->OnControllerButtonPressed(Button, UserId, DeviceId, bIsRepeat);
+    }
+	
+	void FPixelStreamingMessageHandler::HandleOnControllerButtonReleasedWithPlatformUser(FMemoryReader& Ar)
+	{
+    	TPayloadThreeParam<uint8, int32, int32> Payload(Ar);
+
+    	FGamepadKeyNames::Type Button = ConvertAxisIndexToGamepadAxis(Payload.Param1);
+    	FPlatformUserId UserId = FPlatformUserId::CreateFromInternalId(Payload.Param2);
+    	FInputDeviceId DeviceId = FInputDeviceId::CreateFromInternalId(Payload.Param3);
+
+    	UE_LOG(LogPixelStreamingMessageHandler, Verbose, TEXT("GAMEPAD_RELEASED: UserId = %d; DeviceId = %d; KeyName = %s;"), UserId.GetInternalId(), DeviceId.GetId(), *Button.ToString());
+    	MessageHandler->OnControllerButtonReleased(Button, UserId, DeviceId, false);
+	}
 
     /**
      * Mouse events
diff --git a/Engine/Plugins/Media/PixelStreaming/Source/PixelStreaming/Private/PixelStreamingMessageHandler.h b/Engine/Plugins/Media/PixelStreaming/Source/PixelStreaming/Private/PixelStreamingMessageHandler.h
--- a/Engine/Plugins/Media/PixelStreaming/Source/PixelStreaming/Private/PixelStreamingMessageHandler.h
+++ b/Engine/Plugins/Media/PixelStreaming/Source/PixelStreaming/Private/PixelStreamingMessageHandler.h
@@ -51,9 +51,17 @@ namespace UE::PixelStreaming
         /**
          * Controller handling
          */
+    	UE_DEPRECATED(5.1, "This version of HandleOnControllerAnalog is deprecated, use HandleOnControllerAnalogWithPlatformUser instead.")
         virtual void HandleOnControllerAnalog(FMemoryReader& Ar);
+    	UE_DEPRECATED(5.1, "This version of HandleOnControllerButtonPressed is deprecated, use HandleOnControllerButtonPressedWithPlatformUser instead.")
         virtual void HandleOnControllerButtonPressed(FMemoryReader& Ar);
+    	UE_DEPRECATED(5.1, "This version of HandleOnControllerButtonReleased is deprecated, use HandleOnControllerButtonReleasedWithPlatformUser instead.")
         virtual void HandleOnControllerButtonReleased(FMemoryReader& Ar);
+
+    	virtual void HandleOnControllerAnalogWithPlatformUser(FMemoryReader& Ar);
+    	virtual void HandleOnControllerButtonPressedWithPlatformUser(FMemoryReader& Ar);
+    	virtual void HandleOnControllerButtonReleasedWithPlatformUser(FMemoryReader& Ar);
+    	
         /**
          * Mouse handling
          */
