Fix PixelStreaming controller input. PixelStreaming is now compatible with the new FGenericApplicationMessageHandler controller interface.

diff --git a/Engine/Plugins/Media/PixelStreaming/Source/PixelStreaming/Private/PixelStreamingMessageHandler.cpp b/Engine/Plugins/Media/PixelStreaming/Source/PixelStreaming/Private/PixelStreamingMessageHandler.cpp
--- a/Engine/Plugins/Media/PixelStreaming/Source/PixelStreaming/Private/PixelStreamingMessageHandler.cpp
+++ b/Engine/Plugins/Media/PixelStreaming/Source/PixelStreaming/Private/PixelStreamingMessageHandler.cpp
@@ -17,6 +17,7 @@
 #include "Engine/Engine.h"
 #include "Framework/Application/SlateUser.h"
 #include "Widgets/Input/SEditableTextBox.h"
+#include "Misc/CoreMiscDefines.h"
 
 DECLARE_LOG_CATEGORY_EXTERN(LogPixelStreamingMessageHandler, Log, VeryVerbose);
 DEFINE_LOG_CATEGORY(LogPixelStreamingMessageHandler);
@@ -48,15 +49,9 @@ namespace UE::PixelStreaming
         BIND_PLAYBACK_HANDLER((uint8) EToStreamerMsg::TouchMove, &FPixelStreamingMessageHandler::HandleOnTouchMoved);
         BIND_PLAYBACK_HANDLER((uint8) EToStreamerMsg::TouchEnd, &FPixelStreamingMessageHandler::HandleOnTouchEnded);
 
-    	PRAGMA_DISABLE_DEPRECATION_WARNINGS
         BIND_PLAYBACK_HANDLER((uint8) EToStreamerMsg::GamepadAnalog, &FPixelStreamingMessageHandler::HandleOnControllerAnalog);
         BIND_PLAYBACK_HANDLER((uint8) EToStreamerMsg::GamepadButtonPressed, &FPixelStreamingMessageHandler::HandleOnControllerButtonPressed);
         BIND_PLAYBACK_HANDLER((uint8) EToStreamerMsg::GamepadButtonReleased, &FPixelStreamingMessageHandler::HandleOnControllerButtonReleased);
-		PRAGMA_ENABLE_DEPRECATION_WARNINGS
-    	
-    	BIND_PLAYBACK_HANDLER((uint8) EToStreamerMsg::GamepadAnalog, &FPixelStreamingMessageHandler::HandleOnControllerAnalogWithPlatformUser);
-    	BIND_PLAYBACK_HANDLER((uint8) EToStreamerMsg::GamepadButtonPressed, &FPixelStreamingMessageHandler::HandleOnControllerButtonPressedWithPlatformUser);
-    	BIND_PLAYBACK_HANDLER((uint8) EToStreamerMsg::GamepadButtonReleased, &FPixelStreamingMessageHandler::HandleOnControllerButtonReleasedWithPlatformUser);
 
         BIND_PLAYBACK_HANDLER((uint8) EToStreamerMsg::MouseEnter, &FPixelStreamingMessageHandler::HandleOnMouseEnter);
         BIND_PLAYBACK_HANDLER((uint8) EToStreamerMsg::MouseLeave, &FPixelStreamingMessageHandler::HandleOnMouseLeave);
@@ -288,80 +283,43 @@ namespace UE::PixelStreaming
         }
     }
 
-	PRAGMA_DISABLE_DEPRECATION_WARNINGS
     void FPixelStreamingMessageHandler::HandleOnControllerAnalog(FMemoryReader& Ar)
     {
         TPayloadThreeParam<uint8, uint8, double> Payload(Ar);
 
-        FGamepadKeyNames::Type Button = ConvertAxisIndexToGamepadAxis(Payload.Param1);
-        int32 ControllerId = (int32) Payload.Param2;
+        FInputDeviceId ControllerId = FInputDeviceId::CreateFromInternalId((int32) Payload.Param1);
+        FGamepadKeyNames::Type Button = ConvertAxisIndexToGamepadAxis(Payload.Param2);
         float AnalogValue = (float) Payload.Param3;
+        FPlatformUserId UserId = IPlatformInputDeviceMapper::Get().GetPrimaryPlatformUser();
 
-        UE_LOG(LogPixelStreamingMessageHandler, Verbose, TEXT("GAMEPAD_ANALOG: ControllerId = %d; KeyName = %s; AnalogValue = %.4f;"), ControllerId, *Button.ToString(), AnalogValue);
-        MessageHandler->OnControllerAnalog(Button, ControllerId, (float) AnalogValue);
+        UE_LOG(LogPixelStreamingMessageHandler, Verbose, TEXT("GAMEPAD_ANALOG: ControllerId = %d; KeyName = %s; AnalogValue = %.4f;"), ControllerId.GetId(), *Button.ToString(), AnalogValue);
+        MessageHandler->OnControllerAnalog(Button, UserId, ControllerId, AnalogValue);
     }
 
     void FPixelStreamingMessageHandler::HandleOnControllerButtonPressed(FMemoryReader& Ar)
     {
         TPayloadThreeParam<uint8, uint8, uint8> Payload(Ar);
 
-        FGamepadKeyNames::Type Button = ConvertAxisIndexToGamepadAxis(Payload.Param1);
-        int32 ControllerId = (int32) Payload.Param2;
+        FInputDeviceId ControllerId = FInputDeviceId::CreateFromInternalId((int32) Payload.Param1);
+        FGamepadKeyNames::Type Button = ConvertButtonIndexToGamepadButton(Payload.Param2);
         bool bIsRepeat = Payload.Param3 != 0;
+        FPlatformUserId UserId = IPlatformInputDeviceMapper::Get().GetPrimaryPlatformUser();
         
-        UE_LOG(LogPixelStreamingMessageHandler, Verbose, TEXT("GAMEPAD_PRESSED: ControllerId = %d; KeyName = %s; IsRepeat = %s;"), ControllerId, *Button.ToString(), bIsRepeat ? TEXT("True") : TEXT("False"));
-        MessageHandler->OnControllerButtonPressed(Button, ControllerId, bIsRepeat);
+        UE_LOG(LogPixelStreamingMessageHandler, Verbose, TEXT("GAMEPAD_PRESSED: ControllerId = %d; KeyName = %s; IsRepeat = %s;"), ControllerId.GetId(), *Button.ToString(), bIsRepeat ? TEXT("True") : TEXT("False"));
+        MessageHandler->OnControllerButtonPressed(Button, UserId, ControllerId, bIsRepeat);
     }
 
     void FPixelStreamingMessageHandler::HandleOnControllerButtonReleased(FMemoryReader& Ar)
     {
         TPayloadTwoParam<uint8, uint8> Payload(Ar);
 
-        FGamepadKeyNames::Type Button = ConvertAxisIndexToGamepadAxis(Payload.Param1);
-        int32 ControllerId = (int32) Payload.Param2;
+        FInputDeviceId ControllerId = FInputDeviceId::CreateFromInternalId((int32) Payload.Param1);
+        FGamepadKeyNames::Type Button = ConvertButtonIndexToGamepadButton(Payload.Param2);
+        FPlatformUserId UserId = IPlatformInputDeviceMapper::Get().GetPrimaryPlatformUser();
 
-        UE_LOG(LogPixelStreamingMessageHandler, Verbose, TEXT("GAMEPAD_RELEASED: ControllerId = %d; KeyName = %s;"), ControllerId, *Button.ToString());
-        MessageHandler->OnControllerButtonReleased(Button, ControllerId, false);
+        UE_LOG(LogPixelStreamingMessageHandler, Verbose, TEXT("GAMEPAD_RELEASED: ControllerId = %d; KeyName = %s;"), ControllerId.GetId(), *Button.ToString());
+        MessageHandler->OnControllerButtonReleased(Button, UserId, ControllerId, false);
     }
-	PRAGMA_ENABLE_DEPRECATION_WARNINGS
-
-	void FPixelStreamingMessageHandler::HandleOnControllerAnalogWithPlatformUser(FMemoryReader& Ar)
-	{
-    	TPayloadFourParam<uint8, int32, int32, float> Payload(Ar);
-
-    	FGamepadKeyNames::Type Button = ConvertAxisIndexToGamepadAxis(Payload.Param1);
-    	FPlatformUserId UserId = FPlatformUserId::CreateFromInternalId(Payload.Param2);
-    	FInputDeviceId DeviceId = FInputDeviceId::CreateFromInternalId(Payload.Param3);
-    	float AnalogValue = Payload.Param4;
-
-    	UE_LOG(LogPixelStreamingMessageHandler, Verbose, TEXT("GAMEPAD_ANALOG: UserId = %d; DeviceId = %d; KeyName = %s; AnalogValue = %.4f;"), UserId.GetInternalId(), DeviceId.GetId(), *Button.ToString(), AnalogValue);
-    	MessageHandler->OnControllerAnalog(Button, UserId, DeviceId, AnalogValue);
-	}
-	
-	void FPixelStreamingMessageHandler::HandleOnControllerButtonPressedWithPlatformUser(FMemoryReader& Ar)
-    {
-    	TPayloadFourParam<uint8, int32, int32, bool> Payload(Ar);
-
-    	FGamepadKeyNames::Type Button = ConvertAxisIndexToGamepadAxis(Payload.Param1);
-    	FPlatformUserId UserId = FPlatformUserId::CreateFromInternalId(Payload.Param2);
-    	FInputDeviceId DeviceId = FInputDeviceId::CreateFromInternalId(Payload.Param3);
-    	bool bIsRepeat = Payload.Param4;
-
-    	UE_LOG(LogPixelStreamingMessageHandler, Verbose, TEXT("GAMEPAD_PRESSED: UserId = %d; DeviceId = %d; KeyName = %s; IsRepeat = %s;"), UserId.GetInternalId(), DeviceId.GetId(), *Button.ToString(), bIsRepeat ? TEXT("True") : TEXT("False"));
-    	MessageHandler->OnControllerButtonPressed(Button, UserId, DeviceId, bIsRepeat);
-    }
-	
-	void FPixelStreamingMessageHandler::HandleOnControllerButtonReleasedWithPlatformUser(FMemoryReader& Ar)
-	{
-    	TPayloadThreeParam<uint8, int32, int32> Payload(Ar);
-
-    	FGamepadKeyNames::Type Button = ConvertAxisIndexToGamepadAxis(Payload.Param1);
-    	FPlatformUserId UserId = FPlatformUserId::CreateFromInternalId(Payload.Param2);
-    	FInputDeviceId DeviceId = FInputDeviceId::CreateFromInternalId(Payload.Param3);
-
-    	UE_LOG(LogPixelStreamingMessageHandler, Verbose, TEXT("GAMEPAD_RELEASED: UserId = %d; DeviceId = %d; KeyName = %s;"), UserId.GetInternalId(), DeviceId.GetId(), *Button.ToString());
-    	MessageHandler->OnControllerButtonReleased(Button, UserId, DeviceId, false);
-	}
 
     /**
      * Mouse events
diff --git a/Engine/Plugins/Media/PixelStreaming/Source/PixelStreaming/Private/PixelStreamingMessageHandler.h b/Engine/Plugins/Media/PixelStreaming/Source/PixelStreaming/Private/PixelStreamingMessageHandler.h
--- a/Engine/Plugins/Media/PixelStreaming/Source/PixelStreaming/Private/PixelStreamingMessageHandler.h
+++ b/Engine/Plugins/Media/PixelStreaming/Source/PixelStreaming/Private/PixelStreamingMessageHandler.h
@@ -51,16 +51,9 @@ namespace UE::PixelStreaming
         /**
          * Controller handling
          */
-    	UE_DEPRECATED(5.1, "This version of HandleOnControllerAnalog is deprecated, use HandleOnControllerAnalogWithPlatformUser instead.")
         virtual void HandleOnControllerAnalog(FMemoryReader& Ar);
-    	UE_DEPRECATED(5.1, "This version of HandleOnControllerButtonPressed is deprecated, use HandleOnControllerButtonPressedWithPlatformUser instead.")
         virtual void HandleOnControllerButtonPressed(FMemoryReader& Ar);
-    	UE_DEPRECATED(5.1, "This version of HandleOnControllerButtonReleased is deprecated, use HandleOnControllerButtonReleasedWithPlatformUser instead.")
         virtual void HandleOnControllerButtonReleased(FMemoryReader& Ar);
-
-    	virtual void HandleOnControllerAnalogWithPlatformUser(FMemoryReader& Ar);
-    	virtual void HandleOnControllerButtonPressedWithPlatformUser(FMemoryReader& Ar);
-    	virtual void HandleOnControllerButtonReleasedWithPlatformUser(FMemoryReader& Ar);
     	
         /**
          * Mouse handling
