Cherrypicked from UE5/Dev-TensorWorks

Fix PixelStreaming NonUnity Linux

diff --git a/Engine/Plugins/Media/PixelStreaming/Source/PixelStreaming/Private/FrameAdapterProcessRHIToH264.cpp b/Engine/Plugins/Media/PixelStreaming/Source/PixelStreaming/Private/FrameAdapterProcessRHIToH264.cpp
--- a/Engine/Plugins/Media/PixelStreaming/Source/PixelStreaming/Private/FrameAdapterProcessRHIToH264.cpp
+++ b/Engine/Plugins/Media/PixelStreaming/Source/PixelStreaming/Private/FrameAdapterProcessRHIToH264.cpp
@@ -40,7 +40,7 @@ namespace UE::PixelStreaming
 
 	void FFrameAdapterProcessRHIToH264::BeginProcess(const IPixelStreamingInputFrame& InputFrame, TSharedPtr<IPixelStreamingAdaptedOutputFrame> OutputBuffer)
 	{
-		checkf(InputFrame.GetType() == EPixelStreamingInputFrameType::RHI, TEXT("Incorrect source frame coming into frame adapter process."));
+		checkf(InputFrame.GetType() == static_cast<int32>(EPixelStreamingInputFrameType::RHI), TEXT("Incorrect source frame coming into frame adapter process."));
 
 		FRHICommandListImmediate& RHICmdList = FRHICommandListExecutor::GetImmediateCommandList();
 		RHICmdList.EnqueueLambda([this](FRHICommandListImmediate&) { MarkAdaptProcessStarted(); });
diff --git a/Engine/Plugins/Media/PixelStreaming/Source/PixelStreaming/Private/FrameAdapterProcessRHIToI420CPU.cpp b/Engine/Plugins/Media/PixelStreaming/Source/PixelStreaming/Private/FrameAdapterProcessRHIToI420CPU.cpp
--- a/Engine/Plugins/Media/PixelStreaming/Source/PixelStreaming/Private/FrameAdapterProcessRHIToI420CPU.cpp
+++ b/Engine/Plugins/Media/PixelStreaming/Source/PixelStreaming/Private/FrameAdapterProcessRHIToI420CPU.cpp
@@ -66,7 +66,7 @@ namespace UE::PixelStreaming
 
 	void FFrameAdapterProcessRHIToI420CPU::BeginProcess(const IPixelStreamingInputFrame& InputFrame, TSharedPtr<IPixelStreamingAdaptedOutputFrame> OutputBuffer)
 	{
-		checkf(InputFrame.GetType() == EPixelStreamingInputFrameType::RHI, TEXT("Incorrect source frame coming into frame adapter process."));
+		checkf(InputFrame.GetType() == static_cast<int32>(EPixelStreamingInputFrameType::RHI), TEXT("Incorrect source frame coming into frame adapter process."));
 
 		const FPixelStreamingInputFrameRHI& RHISourceFrame = StaticCast<const FPixelStreamingInputFrameRHI&>(InputFrame);
 		FTexture2DRHIRef SourceTexture = RHISourceFrame.FrameTexture;
diff --git a/Engine/Plugins/Media/PixelStreaming/Source/PixelStreaming/Private/FrameAdapterProcessRHIToI420Compute.cpp b/Engine/Plugins/Media/PixelStreaming/Source/PixelStreaming/Private/FrameAdapterProcessRHIToI420Compute.cpp
--- a/Engine/Plugins/Media/PixelStreaming/Source/PixelStreaming/Private/FrameAdapterProcessRHIToI420Compute.cpp
+++ b/Engine/Plugins/Media/PixelStreaming/Source/PixelStreaming/Private/FrameAdapterProcessRHIToI420Compute.cpp
@@ -87,7 +87,7 @@ namespace UE::PixelStreaming
 
 	void FFrameAdapterProcessRHIToI420Compute::BeginProcess(const IPixelStreamingInputFrame& InputFrame, TSharedPtr<IPixelStreamingAdaptedOutputFrame> OutputBuffer)
 	{
-		checkf(InputFrame.GetType() == EPixelStreamingInputFrameType::RHI, TEXT("Incorrect source frame coming into frame adapter process."));
+		checkf(InputFrame.GetType() == static_cast<int32>(EPixelStreamingInputFrameType::RHI), TEXT("Incorrect source frame coming into frame adapter process."));
 
 		const FPixelStreamingInputFrameRHI& RHISourceFrame = StaticCast<const FPixelStreamingInputFrameRHI&>(InputFrame);
 		FTexture2DRHIRef SourceTexture = RHISourceFrame.FrameTexture;
diff --git a/Engine/Plugins/Media/PixelStreaming/Source/PixelStreaming/Private/Streamer.cpp b/Engine/Plugins/Media/PixelStreaming/Source/PixelStreaming/Private/Streamer.cpp
--- a/Engine/Plugins/Media/PixelStreaming/Source/PixelStreaming/Private/Streamer.cpp
+++ b/Engine/Plugins/Media/PixelStreaming/Source/PixelStreaming/Private/Streamer.cpp
@@ -57,7 +57,7 @@ namespace UE::PixelStreaming
 	{
 		VideoSourceGroup->SetVideoInput(Input);
 		Input->OnFrameReady.AddLambda([this](const IPixelStreamingInputFrame& SourceFrame) {
-			if (SourceFrame.GetType() == EPixelStreamingInputFrameType::RHI)
+			if (SourceFrame.GetType() == static_cast<int32>(EPixelStreamingInputFrameType::RHI))
 			{
 				const FPixelStreamingInputFrameRHI& RHISourceFrame = StaticCast<const FPixelStreamingInputFrameRHI&>(SourceFrame);
 
diff --git a/Engine/Plugins/Media/PixelStreaming/Source/PixelStreaming/Private/Tests/FrameAdaptTest.cpp b/Engine/Plugins/Media/PixelStreaming/Source/PixelStreaming/Private/Tests/FrameAdaptTest.cpp
--- a/Engine/Plugins/Media/PixelStreaming/Source/PixelStreaming/Private/Tests/FrameAdaptTest.cpp
+++ b/Engine/Plugins/Media/PixelStreaming/Source/PixelStreaming/Private/Tests/FrameAdaptTest.cpp
@@ -20,7 +20,7 @@ namespace UE::PixelStreaming
 		{
 		}
 
-		virtual int32 GetType() const override { return EPixelStreamingInputFrameType::User; }
+		virtual int32 GetType() const override { return static_cast<int32>(EPixelStreamingInputFrameType::User); }
 		virtual int32 GetWidth() const override { return Width; }
 		virtual int32 GetHeight() const override { return Height; }
 
diff --git a/Engine/Plugins/Media/PixelStreaming/Source/PixelStreaming/Public/IPixelStreamingInputFrame.h b/Engine/Plugins/Media/PixelStreaming/Source/PixelStreaming/Public/IPixelStreamingInputFrame.h
--- a/Engine/Plugins/Media/PixelStreaming/Source/PixelStreaming/Public/IPixelStreamingInputFrame.h
+++ b/Engine/Plugins/Media/PixelStreaming/Source/PixelStreaming/Public/IPixelStreamingInputFrame.h
@@ -15,7 +15,7 @@
  *     NewType2
  * }
  */
-enum EPixelStreamingInputFrameType : int32
+enum class EPixelStreamingInputFrameType : int32
 {
 	Unknown,
 	RHI,
diff --git a/Engine/Plugins/Media/PixelStreaming/Source/PixelStreaming/Public/PixelStreamingInputFrameRHI.h b/Engine/Plugins/Media/PixelStreaming/Source/PixelStreaming/Public/PixelStreamingInputFrameRHI.h
--- a/Engine/Plugins/Media/PixelStreaming/Source/PixelStreaming/Public/PixelStreamingInputFrameRHI.h
+++ b/Engine/Plugins/Media/PixelStreaming/Source/PixelStreaming/Public/PixelStreamingInputFrameRHI.h
@@ -18,7 +18,7 @@ public:
 	}
 	virtual ~FPixelStreamingInputFrameRHI() = default;
 
-	virtual int32 GetType() const override { return EPixelStreamingInputFrameType::RHI; }
+	virtual int32 GetType() const override { return static_cast<int32>(EPixelStreamingInputFrameType::RHI); }
 	virtual int32 GetWidth() const override { return FrameTexture->GetDesc().Extent.X; }
 	virtual int32 GetHeight() const override { return FrameTexture->GetDesc().Extent.Y; }
 
