Create embedded C++ signalling server for the purposes of unit testing signalling where we cannot download servers and we need some signalling mechanism in tree.

Create embedded C++ signalling server for the purposes of unit testing signalling where we cannot download servers and we need some signalling mechanism in UE itself.

- Created new module in Pixel Streaming plugin called "PixelStreamingServers" that is the C++ API for launching signalling server (embedded or downloaded).
- Moved server tests from PixelStreaming module to new PixelStreamingServers module.
- Fixed broken handshake tests to use the new embedded C++ signalling server.
- Updated PixelStreamingVCam plugin to use the new embedded C++ signalling server.
- Added new C++ function to download the Pixel Streaming servers using platform dependant get_ps_servers.

diff --git a/Engine/Plugins/Media/PixelStreaming/PixelStreaming.uplugin b/Engine/Plugins/Media/PixelStreaming/PixelStreaming.uplugin
--- a/Engine/Plugins/Media/PixelStreaming/PixelStreaming.uplugin
+++ b/Engine/Plugins/Media/PixelStreaming/PixelStreaming.uplugin
@@ -22,6 +22,13 @@
 			"PlatformAllowList": [ "Win64", "Linux" ],
 			"TargetDenyList": [ "Server" ]
 		},
+		{
+			"Name": "PixelStreamingServers",
+			"Type": "Runtime",
+			"LoadingPhase": "PostConfigInit",
+			"PlatformAllowList": [ "Win64", "Linux" ],
+			"TargetDenyList": [ ]
+		},
 		{
             "Name": "PixelStreamingShaders",
             "Type": "Runtime",
@@ -46,6 +53,10 @@
 		{
 			"Name": "HardwareEncoders",
 			"Enabled": true
+		},
+		{
+			"Name": "WebSocketNetworking",
+			"Enabled": true
 		}
 	]
 }
diff --git a/Engine/Plugins/Media/PixelStreaming/Resources/NotFound/not_found.html b/Engine/Plugins/Media/PixelStreaming/Resources/NotFound/not_found.html
new file mode 100644
--- /dev/null
+++ b/Engine/Plugins/Media/PixelStreaming/Resources/NotFound/not_found.html
@@ -0,0 +1,11 @@
+<html>
+	<body>
+		<div>If you are seeing this page you have not downloaded the Pixel Streaming servers. 
+			<ol>
+				<li>Go to <i>Engine\Plugins\Media\PixelStreaming\Resources\WebServers</i></li>
+				<li>run <b>get_ps_servers</b></li>
+				<li>Restart the webserver or UE</li>
+			</ol>
+		</div>
+	</body>
+</html>
\ No newline at end of file
diff --git a/Engine/Plugins/Media/PixelStreaming/Source/PixelStreaming/PixelStreaming.Build.cs b/Engine/Plugins/Media/PixelStreaming/Source/PixelStreaming/PixelStreaming.Build.cs
--- a/Engine/Plugins/Media/PixelStreaming/Source/PixelStreaming/PixelStreaming.Build.cs
+++ b/Engine/Plugins/Media/PixelStreaming/Source/PixelStreaming/PixelStreaming.Build.cs
@@ -100,6 +100,7 @@ namespace UnrealBuildTool.Rules
 				"DeveloperSettings",
 				"AVEncoder",
 				"PixelStreamingShaders",
+				"PixelStreamingServers"
 			});
 
 			PrivateDependencyModuleNames.Add("VulkanRHI");
diff --git a/Engine/Plugins/Media/PixelStreaming/Source/PixelStreaming/Private/PixelStreamingModule.cpp b/Engine/Plugins/Media/PixelStreaming/Source/PixelStreaming/Private/PixelStreamingModule.cpp
--- a/Engine/Plugins/Media/PixelStreaming/Source/PixelStreaming/Private/PixelStreamingModule.cpp
+++ b/Engine/Plugins/Media/PixelStreaming/Source/PixelStreaming/Private/PixelStreamingModule.cpp
@@ -370,11 +370,6 @@ namespace UE::PixelStreaming
 				UE_LOG(LogPixelStreaming, Warning, TEXT("PixelStreamingIP and PixelStreamingPort are deprecated flags. Use PixelStreamingURL instead. eg. -PixelStreamingURL=ws://%s:%d"), *SignallingServerIP, SignallingServerPort);
 				SignallingServerURL = FString::Printf(TEXT("ws://%s:%d"), *SignallingServerIP, SignallingServerPort);
 			}
-			else
-			{
-				SignallingServerURL = Settings::GetDefaultSignallingURL();
-				UE_LOG(LogPixelStreaming, Log, TEXT("-PixelStreamingURL was not specified on the command line, using the default connection url: %s"), *SignallingServerURL);
-			}
 		}
 
 		TSharedPtr<IPixelStreamingStreamer> Streamer = CreateStreamer(Settings::GetDefaultStreamerID());
diff --git a/Engine/Plugins/Media/PixelStreaming/Source/PixelStreaming/Private/PixelStreamingServers.cpp b/Engine/Plugins/Media/PixelStreaming/Source/PixelStreaming/Private/PixelStreamingServers.cpp
deleted file mode 100644
diff --git a/Engine/Plugins/Media/PixelStreaming/Source/PixelStreaming/Private/PixelStreamingSignallingConnection.cpp b/Engine/Plugins/Media/PixelStreaming/Source/PixelStreaming/Private/PixelStreamingSignallingConnection.cpp
--- a/Engine/Plugins/Media/PixelStreaming/Source/PixelStreaming/Private/PixelStreamingSignallingConnection.cpp
+++ b/Engine/Plugins/Media/PixelStreaming/Source/PixelStreaming/Private/PixelStreamingSignallingConnection.cpp
@@ -40,7 +40,7 @@ FPixelStreamingSignallingConnection::~FPixelStreamingSignallingConnection()
 	Disconnect();
 }
 
-void FPixelStreamingSignallingConnection::Connect(const FString& Url)
+void FPixelStreamingSignallingConnection::Connect(const FString& InUrl)
 {
 	// Already have a websocket connection, no need to make another one
 	if (WebSocket)
@@ -48,6 +48,7 @@ void FPixelStreamingSignallingConnection::Connect(const FString& Url)
 		return;
 	}
 
+	Url = InUrl;
 	WebSocket = WebSocketFactory(Url);
 	checkf(WebSocket, TEXT("Web Socket Factory failed to return a valid Web Socket."));
 
@@ -79,6 +80,7 @@ void FPixelStreamingSignallingConnection::Disconnect()
 
 	WebSocket->Close();
 	WebSocket = nullptr;
+	UE_LOG(LogPixelStreamingSS, Log, TEXT("Closing websocket to SS %s"), *Url);
 }
 
 bool FPixelStreamingSignallingConnection::IsConnected() const
@@ -97,7 +99,7 @@ void FPixelStreamingSignallingConnection::SendOffer(FPixelStreamingPlayerId Play
 	FString SdpStr = UE::PixelStreaming::ToString(SdpAnsi);
 	OfferJson->SetStringField(TEXT("sdp"), SdpStr);
 
-	UE_LOG(LogPixelStreamingSS, Log, TEXT("Sending offer to player id=%s"), *PlayerId);
+	UE_LOG(LogPixelStreamingSS, Log, TEXT("Sending player=%s \"offer\" to SS %s"), *PlayerId, *Url);
 	UE_LOG(LogPixelStreamingSS, Verbose, TEXT("SDP offer\n%s"), *SdpStr);
 
 	SendMessage(UE::PixelStreaming::ToString(OfferJson, false));
@@ -110,14 +112,21 @@ void FPixelStreamingSignallingConnection::SendAnswer(FPixelStreamingPlayerId Pla
 	SetPlayerIdJson(AnswerJson, PlayerId);
 
 	std::string SdpAnsi;
-	verifyf(SDP.ToString(&SdpAnsi), TEXT("Failed to serialise local SDP"));
-	FString SdpStr = UE::PixelStreaming::ToString(SdpAnsi);
-	AnswerJson->SetStringField(TEXT("sdp"), SdpStr);
 
-	UE_LOG(LogPixelStreamingSS, Log, TEXT("Sending answer to player id=%s"), *PlayerId);
-	UE_LOG(LogPixelStreamingSS, Verbose, TEXT("SDP answer\n%s"), *SdpStr);
+	if(SDP.ToString(&SdpAnsi))
+	{
+		FString SdpStr = UE::PixelStreaming::ToString(SdpAnsi);
+		AnswerJson->SetStringField(TEXT("sdp"), SdpStr);
+
+		UE_LOG(LogPixelStreamingSS, Log, TEXT("Sending player=%s \"answer\" to SS %s"), *PlayerId, *Url);
+		UE_LOG(LogPixelStreamingSS, Verbose, TEXT("SDP answer\n%s"), *SdpStr);
 
-	SendMessage(UE::PixelStreaming::ToString(AnswerJson, false));
+		SendMessage(UE::PixelStreaming::ToString(AnswerJson, false));
+	}
+	else
+	{
+		UE_LOG(LogPixelStreamingSS, Error, TEXT("Failed to serialise local SDP."));
+	}
 }
 
 void FPixelStreamingSignallingConnection::SendIceCandidate(FPixelStreamingPlayerId PlayerId, const webrtc::IceCandidateInterface& IceCandidate)
@@ -132,15 +141,18 @@ void FPixelStreamingSignallingConnection::SendIceCandidate(FPixelStreamingPlayer
 	CandidateJson->SetNumberField(TEXT("sdpMLineIndex"), static_cast<double>(IceCandidate.sdp_mline_index()));
 
 	std::string CandidateAnsi;
-	verifyf(IceCandidate.ToString(&CandidateAnsi), TEXT("Failed to serialize IceCandidate"));
-	FString CandidateStr = UE::PixelStreaming::ToString(CandidateAnsi);
-	CandidateJson->SetStringField(TEXT("candidate"), CandidateStr);
-
-	IceCandidateJson->SetObjectField(TEXT("candidate"), CandidateJson);
-
-	UE_LOG(LogPixelStreamingSS, Verbose, TEXT("-> SS: iceCandidate\n%s"), *UE::PixelStreaming::ToString(IceCandidateJson));
-
-	SendMessage(UE::PixelStreaming::ToString(IceCandidateJson, false));
+	if(IceCandidate.ToString(&CandidateAnsi))
+	{
+		FString CandidateStr = UE::PixelStreaming::ToString(CandidateAnsi);
+		CandidateJson->SetStringField(TEXT("candidate"), CandidateStr);
+		IceCandidateJson->SetObjectField(TEXT("candidate"), CandidateJson);
+		UE_LOG(LogPixelStreamingSS, Verbose, TEXT("-> SS: iceCandidate\n%s"), *UE::PixelStreaming::ToString(IceCandidateJson));
+		SendMessage(UE::PixelStreaming::ToString(IceCandidateJson, false));
+	}
+	else
+	{
+		UE_LOG(LogPixelStreamingSS, Error, TEXT("Failed to serialize IceCandidate"));
+	}
 }
 
 void FPixelStreamingSignallingConnection::SendDisconnectPlayer(FPixelStreamingPlayerId PlayerId, const FString& Reason)
@@ -152,7 +164,7 @@ void FPixelStreamingSignallingConnection::SendDisconnectPlayer(FPixelStreamingPl
 	Json->SetStringField(TEXT("reason"), Reason);
 
 	FString Msg = UE::PixelStreaming::ToString(Json, false);
-	UE_LOG(LogPixelStreamingSS, Verbose, TEXT("-> SS: iceCandidate\n%s"), *Msg);
+	UE_LOG(LogPixelStreamingSS, Verbose, TEXT("Sending player=%s \"disconnectPlayer\" to SS %s"), *PlayerId, *Url);
 
 	SendMessage(Msg);
 }
@@ -216,7 +228,7 @@ void FPixelStreamingSignallingConnection::KeepAlive()
 
 void FPixelStreamingSignallingConnection::OnConnected()
 {
-	UE_LOG(LogPixelStreamingSS, Log, TEXT("Connected to SS"));
+	UE_LOG(LogPixelStreamingSS, Log, TEXT("Connected to SS %s"), *Url);
 
 	Observer.OnSignallingConnected();
 
@@ -231,7 +243,7 @@ void FPixelStreamingSignallingConnection::OnConnected()
 
 void FPixelStreamingSignallingConnection::OnConnectionError(const FString& Error)
 {
-	UE_LOG(LogPixelStreamingSS, Log, TEXT("Failed to connect to SS - signalling server may not be up yet. Message: \"%s\""), *Error);
+	UE_LOG(LogPixelStreamingSS, Log, TEXT("Failed to connect to SS %s - signalling server may not be up yet. Message: \"%s\""), *Url, *Error);
 
 	Observer.OnSignallingError(Error);
 
@@ -246,7 +258,7 @@ void FPixelStreamingSignallingConnection::OnConnectionError(const FString& Error
 
 void FPixelStreamingSignallingConnection::OnClosed(int32 StatusCode, const FString& Reason, bool bWasClean)
 {
-	UE_LOG(LogPixelStreamingSS, Log, TEXT("Connection to SS closed: \n\tstatus %d\n\treason: %s\n\twas clean: %s"), StatusCode, *Reason, bWasClean ? TEXT("true") : TEXT("false"));
+	UE_LOG(LogPixelStreamingSS, Log, TEXT("Closed connection to SS %s - \n\tstatus %d\n\treason: %s\n\twas clean: %s"), *Url, StatusCode, *Reason, bWasClean ? TEXT("true") : TEXT("false"));
 
 	Observer.OnSignallingDisconnected(StatusCode, Reason, bWasClean);
 
@@ -266,14 +278,14 @@ void FPixelStreamingSignallingConnection::OnMessage(const FString& Msg)
 
 	if (!FJsonSerializer::Deserialize(JsonReader, JsonMsg))
 	{
-		FatalError(TEXT("Failed to parse SS message:\n%s"), *Msg);
+		UE_LOG(LogPixelStreamingSS, Error, TEXT("Failed to parse SS message:\n%s"), *Msg);
 		return;
 	}
 
 	FString MsgType;
 	if (!JsonMsg->TryGetStringField(TEXT("type"), MsgType))
 	{
-		FatalError(TEXT("Cannot find `type` field in SS message:\n%s"), *Msg);
+		UE_LOG(LogPixelStreamingSS, Error, TEXT("Cannot find `type` field in SS message:\n%s"), *Msg);
 		return;
 	}
 
@@ -284,7 +296,7 @@ void FPixelStreamingSignallingConnection::OnMessage(const FString& Msg)
 	}
 	else
 	{
-		FatalError(TEXT("Unsupported message `%s` received from SS"), *MsgType);
+		UE_LOG(LogPixelStreamingSS, Error, TEXT("Unsupported message `%s` received from SS"), *MsgType);
 	}
 }
 
@@ -315,7 +327,7 @@ void FPixelStreamingSignallingConnection::OnConfig(const FJsonObjectPtr& Json)
 	const TSharedPtr<FJsonObject>* PeerConnectionOptions = nullptr;
 	if (!Json->TryGetObjectField(TEXT("peerConnectionOptions"), PeerConnectionOptions))
 	{
-		FatalError(TEXT("Cannot find `peerConnectionOptions` field in SS config\n%s"), *UE::PixelStreaming::ToString(Json));
+		UE_LOG(LogPixelStreamingSS, Error, TEXT("Cannot find `peerConnectionOptions` field in SS config\n%s"), *UE::PixelStreaming::ToString(Json));
 		return;
 	}
 
@@ -329,7 +341,8 @@ void FPixelStreamingSignallingConnection::OnConfig(const FJsonObjectPtr& Json)
 			const TSharedPtr<FJsonObject>* IceServerJson;
 			if (!IceServerVal->TryGetObject(IceServerJson))
 			{
-				FatalError(TEXT("Failed to parse SS config: `iceServer` - not an object\n%s"), *UE::PixelStreaming::ToString(*PeerConnectionOptions));
+				UE_LOG(LogPixelStreamingSS, Error, TEXT("Failed to parse SS config: `iceServer` - not an object\n%s"), *UE::PixelStreaming::ToString(*PeerConnectionOptions));
+				return;
 			}
 
 			RTCConfig.servers.push_back(webrtc::PeerConnectionInterface::IceServer{});
@@ -338,9 +351,9 @@ void FPixelStreamingSignallingConnection::OnConfig(const FJsonObjectPtr& Json)
 			TArray<FString> Urls;
 			if ((*IceServerJson)->TryGetStringArrayField(TEXT("urls"), Urls))
 			{
-				for (const FString& Url : Urls)
+				for (const FString& UrlElem : Urls)
 				{
-					IceServer.urls.push_back(UE::PixelStreaming::ToString(Url));
+					IceServer.urls.push_back(UE::PixelStreaming::ToString(UrlElem));
 				}
 			}
 			else
@@ -381,7 +394,8 @@ void FPixelStreamingSignallingConnection::OnSessionDescription(const FJsonObject
 	FString Sdp;
 	if (!Json->TryGetStringField(TEXT("sdp"), Sdp))
 	{
-		FatalError(TEXT("Cannot find `sdp` in Streamer's answer\n%s"), *UE::PixelStreaming::ToString(Json));
+		UE_LOG(LogPixelStreamingSS, Error, TEXT("Cannot find `sdp` in Streamer's answer\n%s"), *UE::PixelStreaming::ToString(Json));
+		return;
 	}
 
 	FPixelStreamingPlayerId PlayerId;
@@ -440,7 +454,8 @@ void FPixelStreamingSignallingConnection::OnPlayerCount(const FJsonObjectPtr& Js
 	uint32 Count;
 	if (!Json->TryGetNumberField(TEXT("count"), Count))
 	{
-		FatalError(TEXT("Failed to get `count` from `playerCount` message\n%s"), *UE::PixelStreaming::ToString(Json));
+		UE_LOG(LogPixelStreamingSS, Error, TEXT("Failed to get `count` from `playerCount` message\n%s"), *UE::PixelStreaming::ToString(Json));
+		return;
 	}
 
 	Observer.OnSignallingPlayerCount(Count);
@@ -452,7 +467,8 @@ void FPixelStreamingSignallingConnection::OnPlayerConnected(const FJsonObjectPtr
 	bool bGotPlayerId = GetPlayerIdJson(Json, PlayerId);
 	if (!bGotPlayerId)
 	{
-		FatalError(TEXT("Failed to get `playerId` from `join` message\n%s"), *UE::PixelStreaming::ToString(Json));
+		UE_LOG(LogPixelStreamingSS, Error, TEXT("Failed to get `playerId` from `join` message\n%s"), *UE::PixelStreaming::ToString(Json));
+		return;
 	}
 
 	UE_LOG(LogPixelStreamingSS, Log, TEXT("Got player connected, player id=%s"), *PlayerId);
@@ -476,7 +492,8 @@ void FPixelStreamingSignallingConnection::OnPlayerDisconnected(const FJsonObject
 	bool bGotPlayerId = GetPlayerIdJson(Json, PlayerId);
 	if (!bGotPlayerId)
 	{
-		FatalError(TEXT("Failed to get `playerId` from `playerDisconnected` message\n%s"), *UE::PixelStreaming::ToString(Json));
+		UE_LOG(LogPixelStreamingSS, Error, TEXT("Failed to get `playerId` from `playerDisconnected` message\n%s"), *UE::PixelStreaming::ToString(Json));
+		return;
 	}
 
 	Observer.OnSignallingPlayerDisconnected(PlayerId);
@@ -488,28 +505,32 @@ void FPixelStreamingSignallingConnection::OnSFUPeerDataChannels(const FJsonObjec
 	bool bSuccess = GetPlayerIdJson(Json, SFUId, TEXT("sfuId"));
 	if (!bSuccess)
 	{
-		FatalError(TEXT("Failed to get `sfuId` from `streamerDataChannels` message\n%s"), *UE::PixelStreaming::ToString(Json));
+		UE_LOG(LogPixelStreamingSS, Error, TEXT("Failed to get `sfuId` from `streamerDataChannels` message\n%s"), *UE::PixelStreaming::ToString(Json));
+		return;
 	}
 
 	FPixelStreamingPlayerId PlayerId;
 	bSuccess = GetPlayerIdJson(Json, PlayerId, TEXT("playerId"));
 	if (!bSuccess)
 	{
-		FatalError(TEXT("Failed to get `playerId` from `streamerDataChannels` message\n%s"), *UE::PixelStreaming::ToString(Json));
+		UE_LOG(LogPixelStreamingSS, Error, TEXT("Failed to get `playerId` from `streamerDataChannels` message\n%s"), *UE::PixelStreaming::ToString(Json));
+		return;
 	}
 
 	int32 SendStreamId;
 	bSuccess = Json->TryGetNumberField(TEXT("sendStreamId"), SendStreamId);
 	if (!bSuccess)
 	{
-		FatalError(TEXT("Failed to get `sendStreamId` from `streamerDataChannels` message\n%s"), *UE::PixelStreaming::ToString(Json));
+		UE_LOG(LogPixelStreamingSS, Error, TEXT("Failed to get `sendStreamId` from `streamerDataChannels` message\n%s"), *UE::PixelStreaming::ToString(Json));
+		return;
 	}
 
 	int32 RecvStreamId;
 	bSuccess = Json->TryGetNumberField(TEXT("recvStreamId"), RecvStreamId);
 	if (!bSuccess)
 	{
-		FatalError(TEXT("Failed to get `recvStreamId` from `streamerDataChannels` message\n%s"), *UE::PixelStreaming::ToString(Json));
+		UE_LOG(LogPixelStreamingSS, Error, TEXT("Failed to get `recvStreamId` from `streamerDataChannels` message\n%s"), *UE::PixelStreaming::ToString(Json));
+		return;
 	}
 
 	Observer.OnSignallingSFUPeerDataChannels(SFUId, PlayerId, SendStreamId, RecvStreamId);
@@ -520,13 +541,13 @@ void FPixelStreamingSignallingConnection::OnPeerDataChannels(const FJsonObjectPt
 	int32 SendStreamId = 0;
 	if (!Json->TryGetNumberField(TEXT("sendStreamId"), SendStreamId))
 	{
-		FatalError(TEXT("Failed to get `sendStreamId` from remote `peerDataChannels` message\n%s"), *UE::PixelStreaming::ToString(Json));
+		UE_LOG(LogPixelStreamingSS, Error, TEXT("Failed to get `sendStreamId` from remote `peerDataChannels` message\n%s"), *UE::PixelStreaming::ToString(Json));
 		return;
 	}
 	int32 RecvStreamId = 0;
 	if (!Json->TryGetNumberField(TEXT("recvStreamId"), RecvStreamId))
 	{
-		FatalError(TEXT("Failed to get `recvStreamId` from remote `peerDataChannels` message\n%s"), *UE::PixelStreaming::ToString(Json));
+		UE_LOG(LogPixelStreamingSS, Error, TEXT("Failed to get `recvStreamId` from remote `peerDataChannels` message\n%s"), *UE::PixelStreaming::ToString(Json));
 		return;
 	}
 	Observer.OnSignallingPeerDataChannels(SendStreamId, RecvStreamId);
@@ -603,9 +624,3 @@ void FPixelStreamingSignallingConnection::PlayerError(FPixelStreamingPlayerId Pl
 	UE_LOG(LogPixelStreamingSS, Error, TEXT("player %s: %s"), *PlayerId, *Msg);
 	SendDisconnectPlayer(PlayerId, Msg);
 }
-
-void FPixelStreamingSignallingConnection::FatalError(const FString& Msg)
-{
-	UE_LOG(LogPixelStreamingSS, Error, TEXT("%s"), *Msg);
-	WebSocket->Close(4000, Msg);
-}
diff --git a/Engine/Plugins/Media/PixelStreaming/Source/PixelStreaming/Private/Streamer.cpp b/Engine/Plugins/Media/PixelStreaming/Source/PixelStreaming/Private/Streamer.cpp
--- a/Engine/Plugins/Media/PixelStreaming/Source/PixelStreaming/Private/Streamer.cpp
+++ b/Engine/Plugins/Media/PixelStreaming/Source/PixelStreaming/Private/Streamer.cpp
@@ -112,11 +112,16 @@ namespace UE::PixelStreaming
 		return CurrentSignallingServerURL;
 	}
 
+	bool FStreamer::IsSignallingConnected()
+	{
+		return SignallingServerConnection && SignallingServerConnection->IsConnected();
+	}
+
 	void FStreamer::StartStreaming()
 	{
 		if (CurrentSignallingServerURL.IsEmpty())
 		{
-			UE_LOG(LogPixelStreaming, Error, TEXT("Attempted to start streamer (%s) but no signalling server URL has been set. Use Streamer->SetSignallingServerURL(URL)"), *StreamerId);
+			UE_LOG(LogPixelStreaming, Error, TEXT("Attempted to start streamer (%s) but no signalling server URL has been set. Use Streamer->SetSignallingServerURL(URL) or -PixelStreamingURL="), *StreamerId);
 			return;
 		}
 
diff --git a/Engine/Plugins/Media/PixelStreaming/Source/PixelStreaming/Private/Streamer.h b/Engine/Plugins/Media/PixelStreaming/Source/PixelStreaming/Private/Streamer.h
--- a/Engine/Plugins/Media/PixelStreaming/Source/PixelStreaming/Private/Streamer.h
+++ b/Engine/Plugins/Media/PixelStreaming/Source/PixelStreaming/Private/Streamer.h
@@ -33,6 +33,7 @@ namespace UE::PixelStreaming
 
 		virtual void SetSignallingServerURL(const FString& InSignallingServerURL) override;
 		virtual FString GetSignallingServerURL() override;
+		virtual bool IsSignallingConnected() override;
 		virtual void StartStreaming() override;
 		virtual void StopStreaming() override;
 		virtual bool IsStreaming() const override { return bStreamingStarted; }
diff --git a/Engine/Plugins/Media/PixelStreaming/Source/PixelStreaming/Private/Tests/HandshakeTest.cpp b/Engine/Plugins/Media/PixelStreaming/Source/PixelStreaming/Private/Tests/HandshakeTest.cpp
--- a/Engine/Plugins/Media/PixelStreaming/Source/PixelStreaming/Private/Tests/HandshakeTest.cpp
+++ b/Engine/Plugins/Media/PixelStreaming/Source/PixelStreaming/Private/Tests/HandshakeTest.cpp
@@ -4,7 +4,6 @@
 #include "Tests/AutomationCommon.h"
 #include "HAL/Runnable.h"
 #include "HAL/RunnableThread.h"
-#include "PixelStreamingServers.h"
 #include "Streamer.h"
 #include "IPixelStreamingSignallingConnectionObserver.h"
 #include "PixelStreamingPrivate.h"
@@ -13,6 +12,8 @@
 #include "PixelStreamingVideoInputBackBuffer.h"
 #include "Utils.h"
 #include "Settings.h"
+#include "PixelStreamingServers.h"
+#include "GenericPlatform/GenericPlatformTime.h"
 
 #if WITH_DEV_AUTOMATION_TESTS
 
@@ -37,15 +38,26 @@ namespace UE::PixelStreaming
 		{
 			FPixelStreamingSignallingConnection::FWebSocketFactory WebSocketFactory = [](const FString& Url) { return FWebSocketsModule::Get().CreateWebSocket(Url, TEXT("")); };
 			SignallingServerConnection = MakeUnique<FPixelStreamingSignallingConnection>(WebSocketFactory, *this, TEXT("FMockPlayer"));
+
+			UE::PixelStreaming::DoOnGameThreadAndWait(MAX_uint32, []() {
+				Settings::CVarPixelStreamingSuppressICECandidateErrors->Set(true, ECVF_SetByCode);
+			});
 		}
 
-		virtual ~FMockPlayer() = default;
+		virtual ~FMockPlayer()
+		{
+			Disconnect();
+
+			UE::PixelStreaming::DoOnGameThread([]() {
+				Settings::CVarPixelStreamingSuppressICECandidateErrors->Set(false, ECVF_SetByCode);
+			});
+		}
 
 		enum class EMode
 		{
 			Unknown,
 			AcceptOffers,
-			OfferToReceive,
+			CreateOffers,
 		};
 
 		void SetMode(EMode InMode) { Mode = InMode; }
@@ -55,6 +67,16 @@ namespace UE::PixelStreaming
 			SignallingServerConnection->Connect(Url);
 		}
 
+		void Disconnect()
+		{
+			SignallingServerConnection->Disconnect();
+		}
+
+		bool IsSignallingConnected()
+		{
+			return SignallingServerConnection->IsConnected();
+		}
+
 		virtual void OnSignallingConnected() override
 		{
 		}
@@ -69,6 +91,7 @@ namespace UE::PixelStreaming
 
 		virtual void OnSignallingConfig(const webrtc::PeerConnectionInterface::RTCConfiguration& Config) override
 		{
+
 			PeerConnection = FPixelStreamingPeerConnection::Create(Config);
 
 			PeerConnection->OnEmitIceCandidate.AddLambda([this](const webrtc::IceCandidateInterface* Candidate) {
@@ -80,12 +103,13 @@ namespace UE::PixelStreaming
 				if (NewState == webrtc::PeerConnectionInterface::IceConnectionState::kIceConnectionConnected)
 				{
 					Completed = true;
+					OnConnectionEstablished.Broadcast();
 				}
 			});
 
 			//PeerConnection->SetVideoSink(new FMockVideoSink());
 
-			if (Mode == EMode::OfferToReceive)
+			if (Mode == EMode::CreateOffers)
 			{
 				PeerConnection->CreateOffer(
 					FPixelStreamingPeerConnection::EReceiveMediaOption::All,
@@ -112,7 +136,7 @@ namespace UE::PixelStreaming
 				};
 				PeerConnection->ReceiveOffer(Sdp, OnSuccess, OnFailure);
 			}
-			else if (Type == webrtc::SdpType::kAnswer && Mode == EMode::OfferToReceive)
+			else if (Type == webrtc::SdpType::kAnswer && Mode == EMode::CreateOffers)
 			{
 				const auto OnFailure = [](const FString& Error) {
 					// fail
@@ -132,136 +156,174 @@ namespace UE::PixelStreaming
 		virtual void OnSignallingPlayerCount(uint32 Count) override {}
 		virtual void OnSignallingPeerDataChannels(int32 SendStreamId, int32 RecvStreamId) override {}
 
+		DECLARE_MULTICAST_DELEGATE(FOnConnectionEstablished);
+		FOnConnectionEstablished OnConnectionEstablished;
+
 		EMode Mode = EMode::Unknown;
 		TUniquePtr<FPixelStreamingSignallingConnection> SignallingServerConnection;
 		TUniquePtr<FPixelStreamingPeerConnection> PeerConnection;
 		bool Completed = false;
 	};
 
-	class FHandshakeTest;
+	TSharedPtr<FStreamer> CreateStreamer(int StreamerPort)
+	{
+		TSharedPtr<FStreamer> OutStreamer = FStreamer::Create("Mock Streamer");
+		OutStreamer->SetVideoInput(FPixelStreamingVideoInputBackBuffer::Create());
+		OutStreamer->SetSignallingServerURL(FString::Printf(TEXT("ws://127.0.0.1:%d"), StreamerPort));
+		OutStreamer->StartStreaming();
+		return OutStreamer;
+	}
 
-	class FMockStreamRun : public FRunnable
+	TSharedPtr<FMockPlayer> CreatePlayer(FMockPlayer::EMode OfferMode, int PlayerPort)
 	{
-	public:
-		FMockStreamRun(FHandshakeTest& InParentTest)
-			: ParentTest(InParentTest) {}
-		virtual ~FMockStreamRun() = default;
+		TSharedPtr<FMockPlayer> OutPlayer = MakeShared<FMockPlayer>();
+		OutPlayer->SetMode(OfferMode);
+		return OutPlayer;
+	}
 
-		virtual bool Init() override
+	TSharedPtr<UE::PixelStreamingServers::IServer> CreateSignallingServer(int StreamerPort, int PlayerPort)
+	{
+		// Make signalling server
+		TSharedPtr<UE::PixelStreamingServers::IServer> OutSignallingServer = UE::PixelStreamingServers::MakeSignallingServer();
+		UE::PixelStreamingServers::FLaunchArgs LaunchArgs;
+		LaunchArgs.ProcessArgs = FString::Printf(TEXT("--StreamerPort=%d --HttpPort=%d"), StreamerPort, PlayerPort);
+		bool bLaunchedSignallingServer = OutSignallingServer->Launch(LaunchArgs);
+		if(!bLaunchedSignallingServer)
 		{
-			bRunning = true;
-			return true;
+			UE_LOG(LogPixelStreaming, Error, TEXT("Failed to launch signalling server."));
 		}
+		UE_LOG(LogPixelStreaming, Log, TEXT("Signalling server launched=%s"), bLaunchedSignallingServer ? TEXT("true") : TEXT("false"));
+		return OutSignallingServer;
+	}
 
-		virtual void Stop() override
+	DEFINE_LATENT_AUTOMATION_COMMAND_TWO_PARAMETER(FWaitForICEConnectedOrTimeout, double, TimeoutSeconds, TSharedPtr<FMockPlayer>, OutPlayer);
+	bool FWaitForICEConnectedOrTimeout::Update()
+	{
+		// If no signalling we can early exit.
+		if(!OutPlayer->IsSignallingConnected())
 		{
+			UE_LOG(LogPixelStreaming, Error, TEXT("Early exiting waiting for ICE connected as player is not connected to signalling server."));
+			return true;
 		}
 
-		virtual void Exit() override
+		if (OutPlayer)
 		{
+			double DeltaTime = FPlatformTime::Seconds() - StartTime;
+			if(DeltaTime > TimeoutSeconds)
+			{
+				UE_LOG(LogPixelStreaming, Error, TEXT("Timed out waiting for RTC connection between streamer and player."));
+				return true;
+			}
+			return OutPlayer->Completed;
 		}
+		return true;
+	}
 
-		virtual uint32 Run() override;
-
-		FHandshakeTest& ParentTest;
-		bool bRunning = false;
-	};
-
-	class FHandshakeWaitCommand : public IAutomationLatentCommand
+	DEFINE_LATENT_AUTOMATION_COMMAND_TWO_PARAMETER(FWaitForStreamerConnectedOrTimeout, double, TimeoutSeconds, TSharedPtr<FStreamer>, OutStreamer);
+	bool FWaitForStreamerConnectedOrTimeout::Update()
 	{
-	public:
-		FHandshakeWaitCommand(FHandshakeTest& ParentTest)
+		if(OutStreamer->IsSignallingConnected())
 		{
-			StreamRunnable = MakeUnique<FMockStreamRun>(ParentTest);
-			StreamThread = FRunnableThread::Create(StreamRunnable.Get(), TEXT("FMockStreamRun Thread"));
+			UE_LOG(LogPixelStreaming, Log, TEXT("Streamer connected to signalling server."));
+			return true;
 		}
 
-		virtual ~FHandshakeWaitCommand()
+		double DeltaTime = FPlatformTime::Seconds() - StartTime;
+		if(DeltaTime > TimeoutSeconds)
 		{
-			StreamThread->Kill(true);
+			UE_LOG(LogPixelStreaming, Error, TEXT("Timed out waiting for streamer to connect to signalling server."));
+			return true;
 		}
+		return false;
+	}
 
-		virtual bool Update() override
+	DEFINE_LATENT_AUTOMATION_COMMAND_THREE_PARAMETER(FWaitForPlayerConnectedOrTimeout, double, TimeoutSeconds, TSharedPtr<FMockPlayer>, OutPlayer, int, PlayerPort);
+	bool FWaitForPlayerConnectedOrTimeout::Update()
+	{
+		if(OutPlayer->IsSignallingConnected())
 		{
-			return !StreamRunnable->bRunning;
+			UE_LOG(LogPixelStreaming, Log, TEXT("Player connected to signalling server."));
+			return true;
+		}
+		else
+		{
+			OutPlayer->Connect(FString::Printf(TEXT("ws://127.0.0.1:%d"), PlayerPort));
 		}
 
-		FRunnableThread* StreamThread = nullptr;
-		TUniquePtr<FMockStreamRun> StreamRunnable;
-	};
-
-	IMPLEMENT_SIMPLE_AUTOMATION_TEST(FHandshakeTest, "System.Plugins.PixelStreaming.Handshake", EAutomationTestFlags::EditorContext | EAutomationTestFlags::ClientContext | EAutomationTestFlags::ProductFilter | EAutomationTestFlags::Disabled)
-	bool FHandshakeTest::RunTest(const FString& Parameters)
-	{
-		ADD_LATENT_AUTOMATION_COMMAND(FHandshakeWaitCommand(*this));
-		return true;
+		double DeltaTime = FPlatformTime::Seconds() - StartTime;
+		if(DeltaTime > TimeoutSeconds)
+		{
+			UE_LOG(LogPixelStreaming, Error, TEXT("Timed out waiting for player to connect to signalling server."));
+			return true;
+		}
+		return false;
 	}
 
-	uint32 FMockStreamRun::Run()
+	DEFINE_LATENT_AUTOMATION_COMMAND_THREE_PARAMETER(FCleanupAll, TSharedPtr<UE::PixelStreamingServers::IServer>, OutSignallingServer, TSharedPtr<FStreamer>, OutStreamer, TSharedPtr<FMockPlayer>, OutPlayer);
+	bool FCleanupAll::Update()
 	{
-		// if we want to actually check we receive frames we need to use VP8 for now
-		//UE::PixelStreaming::DoOnGameThreadAndWait(MAX_uint32, []() {
-		//	Settings::CVarPixelStreamingEncoderCodec->Set(TEXT("VP8"), ECVF_SetByCode);
-		//});
-
-		TSharedPtr<FStreamer> Streamer;
-		TUniquePtr<FMockPlayer> Player;
-
-		Streamer = FStreamer::Create("Mock Streamer");
-		Streamer->SetVideoInput(FPixelStreamingVideoInputBackBuffer::Create());
-		Streamer->SetSignallingServerURL("ws://localhost:8888");
-		Streamer->StartStreaming();
-
-		FPlatformProcess::Sleep(2.0f);
-
-		Player = MakeUnique<FMockPlayer>();
-		Player->SetMode(FMockPlayer::EMode::AcceptOffers);
-		Player->Connect("ws://localhost");
-
-		FPlatformProcess::Sleep(3.0f);
-
-		if (!Player->Completed)
+		if(OutPlayer)
 		{
-			UE_LOG(PixelStreamingHandshakeLog, Error, TEXT("Failed Streamer Offer: ICE state did not change to connected."));
+			OutPlayer->Disconnect();
+			OutPlayer.Reset();
 		}
 
-		Player = nullptr;
-		Streamer = nullptr;
+		if(OutStreamer)
+		{
+			OutStreamer->StopStreaming();
+			OutStreamer.Reset();
+		}
 
-		FPlatformProcess::Sleep(5.0f);
+		if(OutSignallingServer)
+		{
+			OutSignallingServer->Stop();
+			OutSignallingServer.Reset();
+		}
+		return true;
+	}
 
-		UE::PixelStreaming::DoOnGameThreadAndWait(MAX_uint32, []() {
-			Settings::CVarPixelStreamingSuppressICECandidateErrors->Set(true, ECVF_SetByCode);
+	IMPLEMENT_SIMPLE_AUTOMATION_TEST(FHandshakeTestStreamerOffer, "System.Plugins.PixelStreaming.HandshakeStreamerOffer", EAutomationTestFlags::EditorContext | EAutomationTestFlags::ClientContext | EAutomationTestFlags::ProductFilter)
+	bool FHandshakeTestStreamerOffer::RunTest(const FString& Parameters)
+	{
+		int32 StreamerPort = 8866;
+		int32 PlayerPort = 86;
+		TSharedPtr<UE::PixelStreamingServers::IServer> OutSignallingServer = CreateSignallingServer(StreamerPort, PlayerPort);
+		TSharedPtr<FStreamer> OutStreamer = CreateStreamer(StreamerPort);
+		TSharedPtr<FMockPlayer> OutPlayer = CreatePlayer(FMockPlayer::EMode::AcceptOffers, PlayerPort);
+
+		OutPlayer->OnConnectionEstablished.AddLambda([this, OutPlayer](){
+			TestTrue(TEXT("Expected streamer and peer to establish RTC connection when streamer offered first."), OutPlayer->Completed);
 		});
 
-		Streamer = FStreamer::Create("Mock Streamer");
-		Streamer->SetVideoInput(FPixelStreamingVideoInputBackBuffer::Create());
-		Streamer->SetSignallingServerURL("ws://localhost:8888");
-		Streamer->StartStreaming();
-
-		FPlatformProcess::Sleep(2.0f);
-
-		Player = MakeUnique<FMockPlayer>();
-		Player->SetMode(FMockPlayer::EMode::OfferToReceive);
-		Player->Connect("ws://localhost");
-
-		FPlatformProcess::Sleep(3.0f);
+		ADD_LATENT_AUTOMATION_COMMAND(FWaitForStreamerConnectedOrTimeout(5.0, OutStreamer))
+		ADD_LATENT_AUTOMATION_COMMAND(FWaitForPlayerConnectedOrTimeout(5.0, OutPlayer, PlayerPort))
+		ADD_LATENT_AUTOMATION_COMMAND(FWaitForICEConnectedOrTimeout(5.0, OutPlayer))
+		ADD_LATENT_AUTOMATION_COMMAND(FCleanupAll(OutSignallingServer, OutStreamer, OutPlayer))
 
-		if (!Player->Completed)
-		{
-			UE_LOG(PixelStreamingHandshakeLog, Error, TEXT("Failed Player Offer: ICE state did not change to connected."));
-		}
-
-		Player = nullptr;
-		Streamer = nullptr;
+		return true;
+	}
 
-		UE::PixelStreaming::DoOnGameThread([]() {
-			Settings::CVarPixelStreamingSuppressICECandidateErrors->Set(false, ECVF_SetByCode);
+	IMPLEMENT_SIMPLE_AUTOMATION_TEST(FHandshakeTestPlayerOffer, "System.Plugins.PixelStreaming.HandshakePlayerOffer", EAutomationTestFlags::EditorContext | EAutomationTestFlags::ClientContext | EAutomationTestFlags::ProductFilter)
+	bool FHandshakeTestPlayerOffer::RunTest(const FString& Parameters)
+	{
+		int32 StreamerPort = 6446;
+		int32 PlayerPort = 68;
+		TSharedPtr<UE::PixelStreamingServers::IServer> OutSignallingServer = CreateSignallingServer(StreamerPort, PlayerPort);
+		TSharedPtr<FStreamer> OutStreamer = CreateStreamer(StreamerPort);
+		TSharedPtr<FMockPlayer> OutPlayer = CreatePlayer(FMockPlayer::EMode::CreateOffers, PlayerPort);
+
+		OutPlayer->OnConnectionEstablished.AddLambda([this, OutPlayer](){
+			TestTrue(TEXT("Expected streamer and peer to establish RTC connection when streamer offered first."), OutPlayer->Completed);
 		});
+		
+		ADD_LATENT_AUTOMATION_COMMAND(FWaitForStreamerConnectedOrTimeout(5.0, OutStreamer))
+		ADD_LATENT_AUTOMATION_COMMAND(FWaitForPlayerConnectedOrTimeout(5.0, OutPlayer, PlayerPort))
+		ADD_LATENT_AUTOMATION_COMMAND(FWaitForICEConnectedOrTimeout(5.0, OutPlayer))
+		ADD_LATENT_AUTOMATION_COMMAND(FCleanupAll(OutSignallingServer, OutStreamer, OutPlayer))
 
-		bRunning = false;
-		return 0;
+		return true;
 	}
+
 } // namespace UE::PixelStreaming
 
 #endif // WITH_DEV_AUTOMATION_TESTS
diff --git a/Engine/Plugins/Media/PixelStreaming/Source/PixelStreaming/Private/Tests/ServerTests.cpp b/Engine/Plugins/Media/PixelStreaming/Source/PixelStreaming/Private/Tests/ServerTests.cpp
deleted file mode 100644
diff --git a/Engine/Plugins/Media/PixelStreaming/Source/PixelStreaming/Public/IPixelStreamingStreamer.h b/Engine/Plugins/Media/PixelStreaming/Source/PixelStreaming/Public/IPixelStreamingStreamer.h
--- a/Engine/Plugins/Media/PixelStreaming/Source/PixelStreaming/Public/IPixelStreamingStreamer.h
+++ b/Engine/Plugins/Media/PixelStreaming/Source/PixelStreaming/Public/IPixelStreamingStreamer.h
@@ -86,6 +86,11 @@ public:
 	 */
 	virtual FString GetSignallingServerURL() = 0;
 
+	/**
+	 * @return True if the streamer is connected to the signalling mechanism (e.g. websocket for signalling server).
+	 */
+	virtual bool IsSignallingConnected() = 0;
+
 	/**
 	 * @brief Start streaming this streamer
 	 */
diff --git a/Engine/Plugins/Media/PixelStreaming/Source/PixelStreaming/Public/PixelStreamingServers.h b/Engine/Plugins/Media/PixelStreaming/Source/PixelStreaming/Public/PixelStreamingServers.h
deleted file mode 100644
diff --git a/Engine/Plugins/Media/PixelStreaming/Source/PixelStreaming/Public/PixelStreamingSignallingConnection.h b/Engine/Plugins/Media/PixelStreaming/Source/PixelStreaming/Public/PixelStreamingSignallingConnection.h
--- a/Engine/Plugins/Media/PixelStreaming/Source/PixelStreaming/Public/PixelStreamingSignallingConnection.h
+++ b/Engine/Plugins/Media/PixelStreaming/Source/PixelStreaming/Public/PixelStreamingSignallingConnection.h
@@ -120,18 +120,11 @@ private:
 	}
 	void PlayerError(FPixelStreamingPlayerId PlayerId, const FString& Msg);
 
-	template <typename FmtType, typename... T>
-	void FatalError(const FmtType& Msg, T... args)
-	{
-		const FString FormattedMsg = FString::Printf(Msg, args...);
-		FatalError(FormattedMsg);
-	}
-	void FatalError(const FString& Msg);
-
 private:
 	FWebSocketFactory WebSocketFactory;
 	IPixelStreamingSignallingConnectionObserver& Observer;
 	FString StreamerId;
+	FString Url;
 
 	TSharedPtr<IWebSocket> WebSocket;
 
diff --git a/Engine/Plugins/Media/PixelStreaming/Source/PixelStreamingServers/PixelStreamingServers.build.cs b/Engine/Plugins/Media/PixelStreaming/Source/PixelStreamingServers/PixelStreamingServers.build.cs
new file mode 100644
--- /dev/null
+++ b/Engine/Plugins/Media/PixelStreaming/Source/PixelStreamingServers/PixelStreamingServers.build.cs
@@ -0,0 +1,32 @@
+// Copyright Epic Games, Inc. All Rights Reserved.
+
+using System;
+using System.IO;
+using System.Collections.Generic;
+using EpicGames.Core;
+
+namespace UnrealBuildTool.Rules
+{
+	public class PixelStreamingServers : ModuleRules
+	{
+		public PixelStreamingServers(ReadOnlyTargetRules Target) : base(Target)
+		{
+			// This is so for game projects using our public headers don't have to include extra modules they might not know about.
+			PublicDependencyModuleNames.AddRange(new string[] {});
+
+			// NOTE: General rule is not to access the private folder of another module
+			PrivateIncludePaths.AddRange(new string[] {});
+
+			PrivateDependencyModuleNames.AddRange(new string[]
+			{
+				"Core",
+				"Engine",
+				"Json",
+				"WebSocketNetworking",
+				"WebSockets",
+			});
+
+			PrivateDependencyModuleNames.AddRange(new string[] {});
+		}
+	}
+}
\ No newline at end of file
diff --git a/Engine/Plugins/Media/PixelStreaming/Source/PixelStreamingServers/Private/CirrusWrapper.cpp b/Engine/Plugins/Media/PixelStreaming/Source/PixelStreamingServers/Private/CirrusWrapper.cpp
new file mode 100644
--- /dev/null
+++ b/Engine/Plugins/Media/PixelStreaming/Source/PixelStreamingServers/Private/CirrusWrapper.cpp
@@ -0,0 +1,72 @@
+// Copyright Epic Games, Inc. All Rights Reserved.
+
+#include "CirrusWrapper.h"
+#include "PixelStreamingServersModule.h"
+#include "ServerUtils.h"
+#include "PixelStreamingServersLog.h"
+#include "PixelStreamingServers.h"
+#include "Misc/Paths.h"
+
+namespace UE::PixelStreamingServers
+{
+
+	FString FCirrusWrapper::GetServerResourceDirectory()
+	{
+		return FString(TEXT("SignallingWebServer"));
+	}
+
+	void FCirrusWrapper::Stop()
+	{
+		UE_LOG(LogPixelStreamingServers, Log, TEXT("Stopping Cirrus signalling server."));
+		FMonitoredServerBase::Stop();
+	}
+	
+	TSharedPtr<FMonitoredProcess> FCirrusWrapper::LaunchServerProcess(FLaunchArgs& InLaunchArgs, FString ServerAbsPath, TMap<EEndpoint, FURL>& OutEndPoints)
+	{
+		Utils::PopulateCirrusEndPoints(LaunchArgs, OutEndPoints);
+
+		/* Launch the server process in-place. 
+		* Note this has not handled waiting for server to become ready, 
+		* that is server specific, and is therefore handled in `LaunchImpl` methods.
+		*/
+		bool bUseServerBinary = InLaunchArgs.ServerBinaryOverridePath.IsSet();
+
+		TSharedPtr<FMonitoredProcess> CirrusProcess = Utils::LaunchChildProcess(
+			ServerAbsPath, 
+			LaunchArgs.ProcessArgs, 
+			FPaths::GetPathLeaf(ServerAbsPath), 
+			!bUseServerBinary);
+
+		UE_LOG(LogPixelStreamingServers, Log, TEXT("Cirrus signalling server running at: %s"), *OutEndPoints[EEndpoint::Signalling_Webserver].ToString(true));
+
+		if (bPollUntilReady)
+		{
+			Probe = MakeUnique<FWebSocketProbe>(OutEndPoints[EEndpoint::Signalling_Streamer]);
+		}
+
+		return CirrusProcess;
+	}
+
+	bool FCirrusWrapper::TestConnection()
+	{
+		if (bIsReady)
+		{
+			return true;
+		}
+		else
+		{
+			bool bConnected = Probe->Probe();
+			if (bConnected)
+			{
+				// Close the websocket connection so others can use it
+				Probe.Reset();
+				return true;
+			}
+			else
+			{
+				return false;
+			}
+		}
+	}
+
+} // UE::PixelStreamingServers
\ No newline at end of file
diff --git a/Engine/Plugins/Media/PixelStreaming/Source/PixelStreamingServers/Private/CirrusWrapper.h b/Engine/Plugins/Media/PixelStreaming/Source/PixelStreamingServers/Private/CirrusWrapper.h
new file mode 100644
--- /dev/null
+++ b/Engine/Plugins/Media/PixelStreaming/Source/PixelStreamingServers/Private/CirrusWrapper.h
@@ -0,0 +1,31 @@
+// Copyright Epic Games, Inc. All Rights Reserved.
+
+#pragma once
+
+#include "PixelStreamingServers.h"
+#include "WebSocketProbe.h"
+#include "MonitoredServerBase.h"
+
+namespace UE::PixelStreamingServers
+{
+	/* The NodeJS Cirrus signalling server launched as a child process. */
+	class FCirrusWrapper : public FMonitoredServerBase
+	{
+	public:
+		FCirrusWrapper() = default;
+		virtual ~FCirrusWrapper() = default;
+
+	private:
+		TUniquePtr<FWebSocketProbe> Probe;
+
+	protected:
+		/* Begin FMonitoredServerBase interface */
+		void Stop() override;
+		TSharedPtr<FMonitoredProcess> LaunchServerProcess(FLaunchArgs& InLaunchArgs, FString ServerAbsPath, TMap<EEndpoint, FURL>& OutEndPoints) override;
+		bool TestConnection() override;
+		FString GetServerResourceDirectory() override;
+		/* End FMonitoredServerBase interface */
+	};
+
+} // UE::PixelStreamingServers
+		
\ No newline at end of file
diff --git a/Engine/Plugins/Media/PixelStreaming/Source/PixelStreamingServers/Private/MonitoredServerBase.cpp b/Engine/Plugins/Media/PixelStreaming/Source/PixelStreamingServers/Private/MonitoredServerBase.cpp
new file mode 100644
--- /dev/null
+++ b/Engine/Plugins/Media/PixelStreaming/Source/PixelStreamingServers/Private/MonitoredServerBase.cpp
@@ -0,0 +1,74 @@
+// Copyright Epic Games, Inc. All Rights Reserved.
+
+#include "MonitoredServerBase.h"
+#include "PixelStreamingServersLog.h"
+#include "HAL/PlatformFileManager.h"
+#include "ServerUtils.h"
+#include "GenericPlatform/GenericPlatformFile.h"
+#include "Misc/CoreDelegates.h"
+
+namespace UE::PixelStreamingServers
+{
+
+	FMonitoredServerBase::~FMonitoredServerBase()
+	{
+		Stop();
+	}
+
+	bool FMonitoredServerBase::FindServerAbsPath(FLaunchArgs& InLaunchArgs, FString& OutServerAbsPath)
+	{
+		bool bUseServerBinary = InLaunchArgs.ServerBinaryOverridePath.IsSet();
+
+		if(!bUseServerBinary)
+		{
+			bool bFoundServer = Utils::GetDownloadedServer(OutServerAbsPath, GetServerResourceDirectory());
+			if(bFoundServer)
+			{
+				return true;
+			}
+			UE_LOG(LogPixelStreamingServers, Log, TEXT("Downloaded server not found at - %s"), *OutServerAbsPath);
+		}
+		else
+		{
+			OutServerAbsPath = InLaunchArgs.ServerBinaryOverridePath.GetValue();
+			IPlatformFile& PlatformFile = FPlatformFileManager::Get().GetPlatformFile();
+			if(PlatformFile.FileExists(*OutServerAbsPath))
+			{
+				return true;
+			}
+			UE_LOG(LogPixelStreamingServers, Log, TEXT("Server binary not found at - %s"), *OutServerAbsPath);
+		}
+		return false;
+	}
+
+	bool FMonitoredServerBase::LaunchImpl(FLaunchArgs& InLaunchArgs, TMap<EEndpoint, FURL>& OutEndpoints)
+	{
+		// Bind to OnEnginePreExit in case our child process is still running when engine is shutting down
+		EngineShutdownHandle = FCoreDelegates::OnEnginePreExit.AddSP(this, &FMonitoredServerBase::Stop);
+
+		bool bFoundServerPath = FindServerAbsPath(InLaunchArgs, ServerRootAbsPath);
+		if(!bFoundServerPath)
+		{
+			return false;
+		}
+
+		ServerProcess = LaunchServerProcess(InLaunchArgs, ServerRootAbsPath, Endpoints);
+		return ServerProcess.IsValid();
+	}
+
+	FString FMonitoredServerBase::GetPathOnDisk()
+	{
+		return ServerRootAbsPath;
+	}
+
+	void FMonitoredServerBase::Stop()
+	{
+		if (ServerProcess)
+		{
+			ServerProcess->Cancel();
+			ServerProcess.Reset();
+		}
+		FCoreDelegates::OnEnginePreExit.Remove(EngineShutdownHandle);
+	}
+
+} // UE::PixelStreamingServers
diff --git a/Engine/Plugins/Media/PixelStreaming/Source/PixelStreamingServers/Private/MonitoredServerBase.h b/Engine/Plugins/Media/PixelStreaming/Source/PixelStreamingServers/Private/MonitoredServerBase.h
new file mode 100644
--- /dev/null
+++ b/Engine/Plugins/Media/PixelStreaming/Source/PixelStreamingServers/Private/MonitoredServerBase.h
@@ -0,0 +1,49 @@
+// Copyright Epic Games, Inc. All Rights Reserved.
+
+#pragma once
+
+#include "ServerBase.h"
+#include "Templates/SharedPointer.h"
+#include "Misc/MonitoredProcess.h"
+
+namespace UE::PixelStreamingServers
+{
+	
+	/**
+	 *  Base class for all Pixel Streaming servers that are launched as child processes, e.g. calling some scripts or executables.
+	 **/
+	class FMonitoredServerBase : public FServerBase, public TSharedFromThis<FMonitoredServerBase>
+	{
+	public:
+		virtual ~FMonitoredServerBase();
+
+		/* Begin IServer */
+		virtual void Stop() override;
+		virtual FString GetPathOnDisk() override;
+		/* End IServer */
+
+	protected:
+
+		virtual bool LaunchImpl(FLaunchArgs& InLaunchArgs, TMap<EEndpoint, FURL>& OutEndpoints) override;
+		virtual TSharedPtr<FMonitoredProcess> LaunchServerProcess(FLaunchArgs& InLaunchArgs, FString ServerAbsPath, TMap<EEndpoint, FURL>& OutEndpoints) = 0;
+
+		/**
+		 * Implemented by derived types. Implementation specific but somehow the server has been tested to see if it ready for connections.
+		 * @return	True if the server is able to be connected to.
+		 **/
+		virtual bool TestConnection() = 0;
+
+		/**
+		 * @return	Get the directory after PixelStreaming/Resources/WebServers/{ReturnThisDirectoryName}.
+		 **/
+		virtual FString GetServerResourceDirectory() = 0;
+
+		virtual bool FindServerAbsPath(FLaunchArgs& InLaunchArgs, FString& OutServerAbsPath);
+
+	protected:
+		TSharedPtr<FMonitoredProcess> ServerProcess;
+		FString ServerRootAbsPath;
+		FDelegateHandle EngineShutdownHandle;
+	};
+
+} // UE::PixelStreamingServers
diff --git a/Engine/Plugins/Media/PixelStreaming/Source/PixelStreamingServers/Private/PixelStreamingServers.cpp b/Engine/Plugins/Media/PixelStreaming/Source/PixelStreamingServers/Private/PixelStreamingServers.cpp
new file mode 100644
--- /dev/null
+++ b/Engine/Plugins/Media/PixelStreaming/Source/PixelStreamingServers/Private/PixelStreamingServers.cpp
@@ -0,0 +1,24 @@
+// Copyright Epic Games, Inc. All Rights Reserved.
+
+#include "PixelStreamingServers.h"
+#include "CirrusWrapper.h"
+#include "SignallingServer.h"
+
+namespace UE::PixelStreamingServers
+{
+	TSharedPtr<IServer> MakeCirrusServer()
+	{
+		return MakeShared<FCirrusWrapper>();
+	}
+
+	TSharedPtr<IServer> MakeSignallingServer()
+	{
+		return MakeShared<FSignallingServer>();
+	}
+
+	TSharedPtr<FMonitoredProcess> DownloadPixelStreamingServers(bool bSkipIfPresent)
+	{
+		return Utils::DownloadPixelStreamingServers(bSkipIfPresent);
+	}
+
+} // namespace UE::PixelStreamingServers
diff --git a/Engine/Plugins/Media/PixelStreaming/Source/PixelStreamingServers/Private/PixelStreamingServersLog.cpp b/Engine/Plugins/Media/PixelStreaming/Source/PixelStreamingServers/Private/PixelStreamingServersLog.cpp
new file mode 100644
--- /dev/null
+++ b/Engine/Plugins/Media/PixelStreaming/Source/PixelStreamingServers/Private/PixelStreamingServersLog.cpp
@@ -0,0 +1,5 @@
+// Copyright Epic Games, Inc. All Rights Reserved.
+
+#include "PixelStreamingServersLog.h"
+
+DEFINE_LOG_CATEGORY(LogPixelStreamingServers);
\ No newline at end of file
diff --git a/Engine/Plugins/Media/PixelStreaming/Source/PixelStreamingServers/Private/PixelStreamingServersLog.h b/Engine/Plugins/Media/PixelStreaming/Source/PixelStreamingServers/Private/PixelStreamingServersLog.h
new file mode 100644
--- /dev/null
+++ b/Engine/Plugins/Media/PixelStreaming/Source/PixelStreamingServers/Private/PixelStreamingServersLog.h
@@ -0,0 +1,8 @@
+// Copyright Epic Games, Inc. All Rights Reserved.
+
+#pragma once
+
+#include "HAL/IConsoleManager.h"
+#include "Logging/LogMacros.h"
+
+DECLARE_LOG_CATEGORY_EXTERN(LogPixelStreamingServers, Log, All);
\ No newline at end of file
diff --git a/Engine/Plugins/Media/PixelStreaming/Source/PixelStreamingServers/Private/PixelStreamingServersModule.cpp b/Engine/Plugins/Media/PixelStreaming/Source/PixelStreamingServers/Private/PixelStreamingServersModule.cpp
new file mode 100644
--- /dev/null
+++ b/Engine/Plugins/Media/PixelStreaming/Source/PixelStreamingServers/Private/PixelStreamingServersModule.cpp
@@ -0,0 +1,21 @@
+// Copyright Epic Games, Inc. All Rights Reserved.
+
+#include "PixelStreamingServersModule.h"
+
+FPixelStreamingServersModule& FPixelStreamingServersModule::Get()
+{
+	return FModuleManager::LoadModuleChecked<FPixelStreamingServersModule>("PixelStreamingServers");
+}
+
+void FPixelStreamingServersModule::StartupModule()
+{
+	// Any startup logic require goes here.
+}
+
+int FPixelStreamingServersModule::NextPort()
+{
+	// Todo: Checking for in-use ports.
+	return (4000 + NextGeneratedPort.Increment()) % 65535;
+}
+
+IMPLEMENT_MODULE(FPixelStreamingServersModule, PixelStreamingServers)
diff --git a/Engine/Plugins/Media/PixelStreaming/Source/PixelStreamingServers/Private/PixelStreamingServersModule.h b/Engine/Plugins/Media/PixelStreaming/Source/PixelStreamingServers/Private/PixelStreamingServersModule.h
new file mode 100644
--- /dev/null
+++ b/Engine/Plugins/Media/PixelStreaming/Source/PixelStreamingServers/Private/PixelStreamingServersModule.h
@@ -0,0 +1,22 @@
+// Copyright Epic Games, Inc. All Rights Reserved.
+
+#pragma once
+
+#include "CoreMinimal.h"
+#include "Containers/UnrealString.h"
+#include "Modules/ModuleManager.h"
+
+class FPixelStreamingServersModule : public IModuleInterface
+{
+public:
+	virtual ~FPixelStreamingServersModule() = default;
+	virtual void StartupModule() override;
+	int NextPort();
+
+public:
+	static FPixelStreamingServersModule& Get();
+
+private:
+	FThreadSafeCounter NextGeneratedPort = 0;
+
+};
\ No newline at end of file
diff --git a/Engine/Plugins/Media/PixelStreaming/Source/PixelStreamingServers/Private/ServerBase.cpp b/Engine/Plugins/Media/PixelStreaming/Source/PixelStreamingServers/Private/ServerBase.cpp
new file mode 100644
--- /dev/null
+++ b/Engine/Plugins/Media/PixelStreaming/Source/PixelStreamingServers/Private/ServerBase.cpp
@@ -0,0 +1,97 @@
+// Copyright Epic Games, Inc. All Rights Reserved.
+
+#include "ServerBase.h"
+#include "PixelStreamingServersLog.h"
+
+namespace UE::PixelStreamingServers
+{
+
+	bool FServerBase::Launch(FLaunchArgs& InLaunchArgs)
+	{
+		LaunchArgs = InLaunchArgs;
+		bPollUntilReady = InLaunchArgs.bPollUntilReady;
+
+		bHasLaunched = LaunchImpl(InLaunchArgs, Endpoints);
+		PollingStartedSeconds = FPlatformTime::Seconds();
+		return bHasLaunched;
+	}
+
+	bool FServerBase::IsReady()
+	{
+		return bIsReady;
+	}
+
+	bool FServerBase::IsTickableWhenPaused() const
+	{
+		return true;
+	}
+
+	bool FServerBase::IsTickableInEditor() const
+	{
+		return true;
+	}
+
+	bool FServerBase::IsAllowedToTick() const
+	{
+		return bAllowedToTick;
+	}
+
+	void FServerBase::Tick(float DeltaTime)
+	{
+		// No need to do polling if polling is turned off
+		if (!bPollUntilReady)
+		{
+			return;
+		}
+
+		// No need to start polling if we have not launched or we have already concluded the server is ready.
+		if (!bHasLaunched || bIsReady)
+		{
+			return;
+		}
+
+		float SecondsElapsedPolling = FPlatformTime::Seconds() - PollingStartedSeconds;
+		float SecondsSinceReconnect = FPlatformTime::Seconds() - LastReconnectionTimeSeconds;
+
+		if (SecondsElapsedPolling < LaunchArgs.ReconnectionTimeoutSeconds)
+		{
+			if (SecondsSinceReconnect < LaunchArgs.ReconnectionIntervalSeconds)
+			{
+				return;
+			}
+
+			if (TestConnection())
+			{
+				UE_LOG(LogPixelStreamingServers, Log, TEXT("Connected to the server. Server is now ready.  - Broadcasting OnReady event..."));
+				OnReady.Broadcast(Endpoints);
+				// No need to poll anymore we are connected
+				bAllowedToTick = false;
+				bIsReady = true;
+				return;
+			}
+			else
+			{
+				UE_LOG(LogPixelStreamingServers, Log, TEXT("Polling again in another %.f seconds for server to become ready..."), LaunchArgs.ReconnectionIntervalSeconds);
+				LastReconnectionTimeSeconds = FPlatformTime::Seconds();
+			}
+		}
+		else
+		{
+			// No need to poll anymore we timed out
+			UE_LOG(LogPixelStreamingServers, Error, TEXT("Server was not ready after %.f seconds, polling timed out."), LaunchArgs.ReconnectionTimeoutSeconds);
+			bAllowedToTick = false;
+			bTimedOut = true;
+			OnFailedToReady.Broadcast();
+		}
+	}
+
+	bool FServerBase::HasLaunched()
+	{
+		return bHasLaunched;
+	}
+
+	bool FServerBase::IsTimedOut()
+	{
+		return bTimedOut;
+	}
+}
\ No newline at end of file
diff --git a/Engine/Plugins/Media/PixelStreaming/Source/PixelStreamingServers/Private/ServerBase.h b/Engine/Plugins/Media/PixelStreaming/Source/PixelStreamingServers/Private/ServerBase.h
new file mode 100644
--- /dev/null
+++ b/Engine/Plugins/Media/PixelStreaming/Source/PixelStreamingServers/Private/ServerBase.h
@@ -0,0 +1,62 @@
+// Copyright Epic Games, Inc. All Rights Reserved.
+
+#pragma once
+
+#include "PixelStreamingServers.h"
+#include "Tickable.h"
+
+namespace UE::PixelStreamingServers
+{
+
+	// Base class for all servers, provides functionality for servers to be polled for readiness while they start up.
+	class FServerBase : public IServer, public FTickableGameObject
+	{
+	
+	public:
+		virtual ~FServerBase() = default;
+
+		/* Begin IServer */
+		virtual void Stop() = 0;
+		virtual bool HasLaunched() override;
+		virtual bool Launch(FLaunchArgs& InLaunchArgs) override;
+		virtual bool IsReady() override;
+		virtual bool IsTimedOut() override;
+		/* End IServer */
+
+		/* Begin FTickableGameObject */
+		virtual bool IsTickableWhenPaused() const override;
+		virtual bool IsTickableInEditor() const override;
+		virtual void Tick(float DeltaTime) override;
+		virtual bool IsAllowedToTick() const override;
+		TStatId GetStatId() const { RETURN_QUICK_DECLARE_CYCLE_STAT(PixelStreamingServers, STATGROUP_Tickables); }
+		/* End FTickableGameObject */
+	
+	protected:
+		/**
+		 * Implemented by derived types. Implementation specific but somehow the server has been tested to see if it ready for connections.
+		 * @return	True if the server is able to be connected to.
+		 **/
+		virtual bool TestConnection() = 0;
+
+		/**
+		 * Implemented by derived types. The actual implementation of how this specified server is launched.
+		 * @param LaunchArgs	The launch arguments to use.
+		 * @param OutEndpoints	The output endpoints the user can expect to use with this server.
+		 * @return True if the server was able to launch.
+		 **/
+		virtual bool LaunchImpl(FLaunchArgs& InLaunchArgs, TMap<EEndpoint, FURL>& OutEndpoints) = 0;
+
+	protected:
+		FLaunchArgs LaunchArgs;
+		bool bHasLaunched = false;
+		TMap<EEndpoint, FURL> Endpoints;
+		FThreadSafeBool bIsReady = false;
+		float PollingStartedSeconds = 0.0f;
+		float LastReconnectionTimeSeconds = 0.0f;
+		bool bAllowedToTick = true;
+		bool bTimedOut = false;
+		bool bPollUntilReady = false;
+
+	};
+
+} // UE::PixelStreamingServers
\ No newline at end of file
diff --git a/Engine/Plugins/Media/PixelStreaming/Source/PixelStreamingServers/Private/ServerUtils.cpp b/Engine/Plugins/Media/PixelStreaming/Source/PixelStreamingServers/Private/ServerUtils.cpp
new file mode 100644
--- /dev/null
+++ b/Engine/Plugins/Media/PixelStreaming/Source/PixelStreamingServers/Private/ServerUtils.cpp
@@ -0,0 +1,278 @@
+// Copyright Epic Games, Inc. All Rights Reserved.
+
+#include "ServerUtils.h"
+#include "HAL/PlatformFileManager.h"
+#include "PixelStreamingServers.h"
+#include "PixelStreamingServersModule.h"
+#include "Containers/UnrealString.h"
+#include "PixelStreamingServersLog.h"
+#include "Serialization/JsonSerializer.h"
+#include "Serialization/JsonWriter.h"
+#include "Serialization/JsonReader.h"
+#include "Policies/CondensedJsonPrintPolicy.h"
+#include "Misc/Paths.h"
+#include "GenericPlatform/GenericPlatformFile.h"
+
+namespace UE::PixelStreamingServers::Utils
+{
+
+	TSharedPtr<FMonitoredProcess> LaunchChildProcess(FString ExecutableAbsPath, FString Args, FString LogPrefix, bool bRunAsScript)
+	{
+		// Check if the binary actually exists
+		IPlatformFile& FileManager = FPlatformFileManager::Get().GetPlatformFile();
+		if (!FileManager.FileExists(*ExecutableAbsPath))
+		{
+			UE_LOG(LogPixelStreamingServers, Error, TEXT("Cannot start child process - the specified file did not exist. File=%s"), *ExecutableAbsPath);
+			return TSharedPtr<FMonitoredProcess>();
+		}
+
+		if(bRunAsScript)
+		{
+// Get the executable we will use to run the scripts (e.g. cmd.exe on Windows)
+#if PLATFORM_WINDOWS
+			Args = FString::Printf(TEXT("/c \"\"%s\" %s\""), *ExecutableAbsPath, *Args);
+			ExecutableAbsPath = TEXT("cmd.exe");
+			
+#elif PLATFORM_LINUX
+			Args = FString::Printf(TEXT(" -- \"%s\" %s --nosudo"), *ExecutableAbsPath, *Args);
+			ExecutableAbsPath = TEXT("/usr/bin/bash");
+#else
+			UE_LOG(LogPixelStreamingServers, Error, TEXT("Unsupported platform for Pixel Streaming."));
+			return TSharedPtr<FMonitoredProcess>();
+#endif
+		}
+
+		TSharedPtr<FMonitoredProcess> ChildProcess = MakeShared<FMonitoredProcess>(ExecutableAbsPath, Args, true, true);
+		// Bind to output so we can capture the output in the log
+		ChildProcess->OnOutput().BindLambda([LogPrefix](FString Output) {
+			UE_LOG(LogPixelStreamingServers, Log, TEXT("%s - %s"), *LogPrefix, *Output);
+		});
+		// Run the child process
+		UE_LOG(LogPixelStreamingServers, Log, TEXT("Launch child process - %s %s"), *ExecutableAbsPath, *Args);
+		ChildProcess->Launch();
+		return ChildProcess;
+	}
+
+	bool ExtractValueFromArgs(FString ArgsString, FString ArgKey, FString FallbackValue, FString& OutValue)
+	{
+		// Tokenize string in single whitespace " ".
+		TArray<FString> ArgTokens;
+		ArgsString.ParseIntoArray(ArgTokens, TEXT(" "), true);
+
+		for (FString& Token : ArgTokens)
+		{
+			Token.TrimStartAndEndInline();
+
+			if (!Token.StartsWith(ArgKey, ESearchCase::Type::CaseSensitive))
+			{
+				continue;
+			}
+
+			// We have a matching token for our search "key" - split on it.
+			FString RightStr;
+			if (!Token.Split(TEXT("="), nullptr, &RightStr))
+			{
+				continue;
+			}
+
+			OutValue = RightStr;
+			return true;
+		}
+		OutValue = FallbackValue;
+		return false;
+	}
+
+	FString QueryOrSetProcessArgs(FLaunchArgs& LaunchArgs, FString ArgKey, FString FallbackArgValue)
+	{
+		FString OutValue;
+		bool bExtractedValue = ExtractValueFromArgs(LaunchArgs.ProcessArgs, ArgKey, FallbackArgValue, OutValue);
+
+		// No key was present so we will inject our own.
+		if (!bExtractedValue)
+		{
+			LaunchArgs.ProcessArgs += FString::Printf(TEXT(" %s%s"), *ArgKey, *FallbackArgValue);
+		}
+
+		return OutValue;
+	}
+
+	bool GetResourcesDir(FString& OutResourcesDir)
+	{
+#if WITH_EDITOR
+		OutResourcesDir = FPaths::EnginePluginsDir() / TEXT("Media") / TEXT("PixelStreaming") / TEXT("Resources");
+#else
+		OutResourcesDir = FPaths::ProjectDir() / TEXT("Samples") / TEXT("PixelStreaming");
+#endif // WITH_EDITOR
+
+		OutResourcesDir = FPaths::ConvertRelativePathToFull(OutResourcesDir);
+
+		if(FPaths::DirectoryExists(OutResourcesDir))
+		{
+			return true;
+		}
+		return false;
+	}
+
+	bool GetWebServersDir(FString& OutWebServersAbsPath)
+	{
+		bool bResourceDirExists = GetResourcesDir(OutWebServersAbsPath);
+
+		if(!bResourceDirExists)
+		{
+			return false;
+		}
+
+		OutWebServersAbsPath = OutWebServersAbsPath / TEXT("WebServers");
+
+		if(FPaths::DirectoryExists(OutWebServersAbsPath))
+		{
+			return true;
+		}
+		return false;
+	}
+
+	bool GetDownloadedServer(FString& OutAbsPath, FString ServerDirectoryName)
+	{
+		bool bServersDirExists = GetWebServersDir(OutAbsPath);
+
+		if(!bServersDirExists)
+		{
+			return false;
+		}
+
+		// Now add the {ServerDirectoryName}/plaform_script/{os_specific_path}/{run_local.bat|sh}
+		OutAbsPath = OutAbsPath / ServerDirectoryName / TEXT("platform_scripts");
+#if PLATFORM_WINDOWS
+		OutAbsPath = OutAbsPath / TEXT("cmd") / TEXT("run_local.bat");
+#elif PLATFORM_LINUX
+		OutAbsPath = OutAbsPath / TEXT("bash") / TEXT("run_local.sh");
+#else
+		UE_LOG(LogPixelStreamingServers, Error, TEXT("Unsupported platform for Pixel Streaming scripts."));
+		return false;
+#endif
+
+		IPlatformFile& PlatformFile = FPlatformFileManager::Get().GetPlatformFile();
+		if(PlatformFile.FileExists(*OutAbsPath))
+		{
+			return true;
+		}
+		return false;
+	}
+
+	TSharedPtr<FMonitoredProcess> DownloadPixelStreamingServers(bool bSkipIfPresent)
+	{
+		FString OutScriptPath;
+		if(bSkipIfPresent && GetDownloadedServer(OutScriptPath, FString(TEXT("SignallingWebServer"))))
+		{
+			UE_LOG(LogPixelStreamingServers, Log, TEXT("Found pixel streaming servers, skipping download."));
+			// empty process
+			return TSharedPtr<FMonitoredProcess>();
+		}
+		
+		bool bHasWebServersDir = GetWebServersDir(OutScriptPath);
+
+		if(!bHasWebServersDir)
+		{
+			UE_LOG(LogPixelStreamingServers, Error, TEXT("Could not download ps servers, no PixelStreaming/Resources/WebServers directory found."));
+		}
+
+		FString Args = TEXT("");
+
+#if PLATFORM_WINDOWS
+		OutScriptPath = OutScriptPath / TEXT("get_ps_servers.bat");
+#elif PLATFORM_LINUX
+		OutScriptPath = OutScriptPath / TEXT("get_ps_servers.sh");
+#else
+		UE_LOG(LogPixelStreamingServers, Error, TEXT("Unsupported platform for Pixel Streaming scripts."));
+		// empty process
+		return TSharedPtr<FMonitoredProcess>();
+#endif
+
+		return LaunchChildProcess(OutScriptPath, Args, FString(TEXT("Download ps servers")), true /*bRunAsScript*/);
+	}
+
+	void PopulateCirrusEndPoints(FLaunchArgs& InLaunchArgs, TMap<EEndpoint, FURL>& OutEndPoints)
+	{
+		FPixelStreamingServersModule& Module = FPixelStreamingServersModule::Get();
+
+		// Query for ports, or set them if they don't exist
+		FString StreamerPort = Utils::QueryOrSetProcessArgs(InLaunchArgs, TEXT("--StreamerPort="), FString::FromInt(Module.NextPort()));
+		FString SFUPort = Utils::QueryOrSetProcessArgs(InLaunchArgs, TEXT("--SFUPort="), FString::FromInt(Module.NextPort()));
+		FString MatchmakerPort = Utils::QueryOrSetProcessArgs(InLaunchArgs, TEXT("--MatchmakerPort="), FString::FromInt(Module.NextPort()));
+		FString HttpPort = Utils::QueryOrSetProcessArgs(InLaunchArgs, TEXT("--HttpPort="), FString::FromInt(Module.NextPort()));
+
+		// Construct endpoint urls
+
+		// Streamer
+		FURL SignallingStreamerUrl;
+		SignallingStreamerUrl.Protocol = TEXT("ws");
+		SignallingStreamerUrl.Host = TEXT("127.0.0.1");
+		SignallingStreamerUrl.Port = FCString::Atoi(*StreamerPort); 
+		SignallingStreamerUrl.Map = FString();
+
+		// Players
+		FURL PlayersUrl;
+		PlayersUrl.Protocol = TEXT("ws");
+		PlayersUrl.Host = TEXT("127.0.0.1");
+		PlayersUrl.Port = FCString::Atoi(*HttpPort);
+		PlayersUrl.Map = FString();
+
+		// Webserver
+		FURL WebServerUrl;
+		WebServerUrl.Protocol = TEXT("http");
+		WebServerUrl.Host = TEXT("127.0.0.1");
+		WebServerUrl.Port = FCString::Atoi(*HttpPort);
+		WebServerUrl.Map = FString();
+
+		// SFU
+		FURL SFUUrl;
+		SFUUrl.Protocol = TEXT("ws");
+		SFUUrl.Host = TEXT("127.0.0.1");
+		SFUUrl.Port = FCString::Atoi(*SFUPort);
+		SFUUrl.Map = FString();
+
+		// Matchmaker
+		FURL MatchmakerUrl;
+		MatchmakerUrl.Protocol = TEXT("ws");
+		MatchmakerUrl.Host = TEXT("127.0.0.1");
+		MatchmakerUrl.Port = FCString::Atoi(*MatchmakerPort);
+		MatchmakerUrl.Map = FString();
+
+		OutEndPoints.Add(EEndpoint::Signalling_Streamer, SignallingStreamerUrl);
+		OutEndPoints.Add(EEndpoint::Signalling_Players, PlayersUrl);
+		OutEndPoints.Add(EEndpoint::Signalling_SFU, SFUUrl);
+		OutEndPoints.Add(EEndpoint::Signalling_Matchmaker, MatchmakerUrl);
+		OutEndPoints.Add(EEndpoint::Signalling_Webserver, WebServerUrl);
+	}
+
+	FString ToString(FURL Url)
+	{
+		return FString::Printf(TEXT("%s://%s:%d"), *(Url.Protocol), *(Url.Host), Url.Port);
+	}
+
+	FString ToString(TArrayView<uint8> UTF8Bytes)
+	{
+		FUTF8ToTCHAR Converted((const ANSICHAR*)UTF8Bytes.GetData(), UTF8Bytes.Num());
+		FString OutString(Converted.Length(), Converted.Get());
+		return OutString;
+	}
+
+	FString ToString(TSharedRef<FJsonObject> JSONObj)
+	{
+		FString Res;
+		auto JsonWriter = TJsonWriterFactory<TCHAR, TCondensedJsonPrintPolicy<TCHAR>>::Create(&Res);
+		bool bSerialized = FJsonSerializer::Serialize(JSONObj, JsonWriter);
+		if(!bSerialized)
+		{
+			UE_LOG(LogPixelStreamingServers, Error, TEXT("Failed to stringify JSON object."));
+		}
+		return Res;
+	}
+
+	bool Jsonify(FString InJSONString, TSharedPtr<FJsonObject>& OutJSON)
+	{
+		const auto JsonReader = TJsonReaderFactory<TCHAR>::Create(InJSONString);
+		return FJsonSerializer::Deserialize(JsonReader, OutJSON);
+	}
+
+} // UE::PixelStreamingServers::Utils
\ No newline at end of file
diff --git a/Engine/Plugins/Media/PixelStreaming/Source/PixelStreamingServers/Private/ServerUtils.h b/Engine/Plugins/Media/PixelStreaming/Source/PixelStreamingServers/Private/ServerUtils.h
new file mode 100644
--- /dev/null
+++ b/Engine/Plugins/Media/PixelStreaming/Source/PixelStreamingServers/Private/ServerUtils.h
@@ -0,0 +1,103 @@
+// Copyright Epic Games, Inc. All Rights Reserved.
+
+#pragma once
+
+#include "PixelStreamingServers.h"
+#include "Templates/SharedPointer.h"
+#include "Misc/MonitoredProcess.h"
+#include "Dom/JsonObject.h"
+#include "Engine/EngineBaseTypes.h"
+
+namespace UE::PixelStreamingServers
+{
+	namespace Utils
+	{
+		/**
+		 * Inspects the `LaunchArgs.ProcessArgs` for a string key matching `ArgKey`.
+		 * We assume `LaunchArgs.ProcessArgs` are passed as `--key=value` in a String.
+		 * @param LaunchArgs The launch args to get/set the ProcessArgs of.
+		 * @param ArgKey The key to look for in the `LaunchArgs.ProcessArgs`
+		 * @param FallbackArgValue If the key does not exist in the `LaunchArgs.ProcessArgs` this value to set for it in-place.
+		 * @return The extracted/set process arg.
+		 **/
+		inline FString QueryOrSetProcessArgs(FLaunchArgs& LaunchArgs, FString ArgKey, FString FallbackArgValue);
+
+		/**
+		 * Launches an executable or script as a child process of Unreal Engine.
+		 * Note: When running scripts they must be launchable through cmd.exe on Windows or bash on Linux.
+		 * @param ExecutableAbsPath The absolute path to the executable or script we want to run.
+		 * @param Args The argument to pass to the child process we are launching.
+		 * @param LogPrefix If the child process has output the prefix to put on the UE log message for that output.
+		 * @param bRunAsScript If true the `ExecutableAbsPath` will be passed to cmd/bash.
+		 **/
+		inline TSharedPtr<FMonitoredProcess> LaunchChildProcess(FString ExecutableAbsPath, FString Args, FString LogPrefix, bool bRunAsScript);
+
+		/**
+		 * Get the absolute path to the Pixel Streaming resources dir (changes based on whether called in editor or not.)
+		 * In editor this PixelStreaming/Resources in a game this is Samples/PixelStreaming
+		 * @param OutResourcesDir The absolute path to "resources" dir.
+		 * @return True if the returned path exists.
+		 **/
+		inline bool GetResourcesDir(FString& OutResourcesDir);
+
+		/**
+		 * Get the absolute path to the webservers (changes based on whether called in editor or not.)
+		 * @param OutWebServersAbsPath The absolute path to the webservers.
+		 * @return True if the returned path exists.
+		 **/
+		inline bool GetWebServersDir(FString& OutWebServersAbsPath);
+
+		/**
+		 * Attempts to get the absolute path to a downloaded server residing in the /Resources directory
+		 * of the Pixel Streaming plugin, if it exists.
+		 * @param OutAbsPath The absolute path of the downloaded server.
+		 * @param ServerDirectoryName The directory name of the server we should look for under /Resources
+		 * @return True if the scripts to launch the server exist at this location.
+		 **/
+		inline bool GetDownloadedServer(FString& OutAbsPath, FString ServerDirectoryName);
+
+		/**
+		 * Download the Pixel Streaming servers using the `get_ps_servers` scripts.
+		 * @param bSkipIfPresent Servers will not be downloaded if they are already present.
+		 * @return The child process that is used to download the servers.
+		 **/
+		inline TSharedPtr<FMonitoredProcess> DownloadPixelStreamingServers(bool bSkipIfPresent);
+
+		/**
+		 * Populates the cirrus endpoints by extracting endpoint information from the `InLaunchArgs.ProcessArgs`, and if not found, using defaults.
+		 * @param InLaunchArgs The launch args containing the process args that we will be inspecting.
+		 * @param OutEndPoints The end points for the server if we launch it with these launch args.
+		 **/
+		inline void PopulateCirrusEndPoints(FLaunchArgs& InLaunchArgs, TMap<EEndpoint, FURL>& OutEndPoints);
+
+		/**
+		 * Convert a `FURL` into a `FString`. `FURL::ToString()` already exists; however, it appends a problematic trailing slash to the output string.
+		 * @param Url The url to format into a string.
+		 * @return The url formatted without a trailing slash (the trailing slash messes up the Websocket clients).
+		 **/
+		inline FString ToString(FURL Url);
+
+		/**
+		 * Convert a UTF8 byte array to an FString
+		 * @param UTF8Bytes Byte array of UTF8 characters.
+		 * @return A string representation of the byte array.
+		 **/
+		inline FString ToString(TArrayView<uint8> UTF8Bytes);
+
+		/**
+		 * Convert JSON object to FString.
+		 * @param JSONObj A JSON object.
+		 * @return The stringified JSON object.
+		 **/
+		inline FString ToString(TSharedRef<FJsonObject> JSONObj);
+
+		/** 
+		 * Convert a JSON formatted string into a JSON object, if possible.
+		 * @param InJSONString The JSON object expressed as a string.
+		 * @param OutJSON The output json object.
+		 * @return True if we were able to deserialize the string into a JSON object.
+		 **/
+		inline bool Jsonify(FString InJSONString, TSharedPtr<FJsonObject>& OutJSON);
+
+	} // Utils
+} // UE::PixelStreamingServers
diff --git a/Engine/Plugins/Media/PixelStreaming/Source/PixelStreamingServers/Private/SignallingServer.cpp b/Engine/Plugins/Media/PixelStreaming/Source/PixelStreamingServers/Private/SignallingServer.cpp
new file mode 100644
--- /dev/null
+++ b/Engine/Plugins/Media/PixelStreaming/Source/PixelStreamingServers/Private/SignallingServer.cpp
@@ -0,0 +1,345 @@
+// Copyright Epic Games, Inc. All Rights Reserved.
+
+#include "SignallingServer.h"
+#include "ServerUtils.h"
+#include "PixelStreamingServersLog.h"
+#include "Misc/Paths.h"
+
+namespace UE::PixelStreamingServers
+{
+
+	FString FSignallingServer::GetPathOnDisk()
+	{
+		return FString();
+	}
+
+	void FSignallingServer::Stop()
+	{
+		if(PlayersWS)
+		{
+			PlayersWS->Stop();
+			PlayersWS.Reset();
+		}
+		if(StreamersWS)
+		{
+			StreamersWS->Stop();
+			StreamersWS.Reset();
+		}
+		if(Probe)
+		{
+			Probe.Reset();
+		}
+	}
+
+	bool FSignallingServer::TestConnection()
+	{
+		if (bIsReady)
+		{
+			return true;
+		}
+		else
+		{
+			bool bConnected = Probe->Probe();
+			if (bConnected)
+			{
+				// Close the websocket connection so others can use it
+				Probe.Reset();
+				return true;
+			}
+			else
+			{
+				return false;
+			}
+		}
+	}
+
+	TArray<FWebSocketHttpMount> FSignallingServer::GenerateDirectoriesToServe() const
+	{
+		FString ServersDir;
+		bool bServersDirExists = Utils::GetWebServersDir(ServersDir);
+		if(bServersDirExists)
+		{
+			ServersDir = ServersDir / TEXT("SignallingWebServer");
+			bServersDirExists = FPaths::DirectoryExists(ServersDir);
+		}
+
+		TArray<FWebSocketHttpMount> MountsArr;
+
+#if WITH_EDITOR
+		// If server directory doesn't exist we will serve a known directory that gives the user a message
+		// telling them to run the `get_ps_servers` script.
+		if(!bServersDirExists)
+		{
+			FString OutResourcesDir;
+			bool bResourcesDirExists = Utils::GetResourcesDir(OutResourcesDir);
+			FString NotFoundDir = OutResourcesDir / TEXT("NotFound");
+
+			if(bResourcesDirExists && FPaths::DirectoryExists(NotFoundDir))
+			{
+				FWebSocketHttpMount Mount;
+				Mount.SetPathOnDisk(NotFoundDir);
+				Mount.SetWebPath(FString(TEXT("/")));
+				Mount.SetDefaultFile(FString(TEXT("not_found.html")));
+				MountsArr.Add(Mount);
+				return MountsArr;
+			}
+		}
+#endif //WITH_EDITOR
+
+		// Add /Public
+		FWebSocketHttpMount PublicMount;
+		PublicMount.SetPathOnDisk(ServersDir / TEXT("Public"));
+		PublicMount.SetWebPath(FString(TEXT("/")));
+		PublicMount.SetDefaultFile(FString(TEXT("player.html")));
+		MountsArr.Add(PublicMount);
+
+		// Add /images
+		FWebSocketHttpMount ImagesMount;
+		ImagesMount.SetPathOnDisk(ServersDir / TEXT("images"));
+		ImagesMount.SetWebPath(FString(TEXT("/images")));
+		MountsArr.Add(ImagesMount);
+
+		// Add /scripts
+		FWebSocketHttpMount ScriptsMount;
+		ScriptsMount.SetPathOnDisk(ServersDir / TEXT("scripts"));
+		ScriptsMount.SetWebPath(FString(TEXT("/scripts")));
+		MountsArr.Add(ScriptsMount);
+
+		return MountsArr;
+	}
+
+	bool FSignallingServer::LaunchImpl(FLaunchArgs& InLaunchArgs, TMap<EEndpoint, FURL>& OutEndpoints)
+	{
+		Utils::PopulateCirrusEndPoints(InLaunchArgs, OutEndpoints);
+		FURL PlayersURL = OutEndpoints[EEndpoint::Signalling_Players];
+		FURL StreamerURL = OutEndpoints[EEndpoint::Signalling_Streamer];
+
+		/*
+		* --------------- Streamers websocket server ---------------
+		*/
+		StreamersWS = MakeUnique<FWebSocketServerWrapper>();
+		StreamersWS->OnMessage.AddRaw(this, &FSignallingServer::OnStreamerMessage);
+		StreamersWS->OnOpenConnection.AddRaw(this, &FSignallingServer::OnStreamerConnected);
+		bool bLaunchedStreamerServer = StreamersWS->Launch(StreamerURL.Port);
+
+		if(!bLaunchedStreamerServer)
+		{
+			UE_LOG(LogPixelStreamingServers, Error, TEXT("Failed to launch websocket server for streamers on port=%d"), StreamerURL.Port);
+			return false;
+		}
+
+		/*
+		* --------------- Players websocket server ---------------
+		*/
+		PlayersWS = MakeUnique<FWebSocketServerWrapper>();
+		PlayersWS->EnableWebServer(GenerateDirectoriesToServe());
+		PlayersWS->OnMessage.AddRaw(this, &FSignallingServer::OnPlayerMessage);
+		PlayersWS->OnOpenConnection.AddRaw(this, &FSignallingServer::OnPlayerConnected);
+		bool bLaunchedPlayerServer = PlayersWS->Launch(PlayersURL.Port);
+
+		if(!bLaunchedPlayerServer)
+		{
+			UE_LOG(LogPixelStreamingServers, Error, TEXT("Failed to launch websocket server for players on port=%d"), PlayersURL.Port);
+			return false;
+		}
+
+		
+		/*
+		* --------------- Websocket probe ---------------
+		*/
+
+		if (bPollUntilReady)
+		{
+			TArray<FString> Protocols;
+			Protocols.Add(FString(TEXT("binary")));
+			Probe = MakeUnique<FWebSocketProbe>(StreamerURL, Protocols);
+		}
+
+		return true;
+	}
+
+	void FSignallingServer::OnStreamerConnected(uint16 ConnectionId)
+	{
+		UE_LOG(LogPixelStreamingServers, Log, TEXT("Streamer websocket connected, id=%d"), ConnectionId);
+
+		// Send a config message to the streamer passing ICE servers to be used.
+		FString JSONString = CreateConfigJSON();
+		UE_LOG(LogPixelStreamingServers, Log, TEXT("Sending to streamer id=%d: %s"), ConnectionId, *JSONString);
+		StreamersWS->Send(ConnectionId, JSONString);
+	}
+
+	FString FSignallingServer::CreateConfigJSON() const
+	{
+		// Todo (Luke): Parse `iceServers` from the process args `--peerConnectionOptions`
+		TArray<TSharedPtr<FJsonValue>> IceServersArr;
+
+		TSharedPtr<FJsonObject> PeerConnectionOptionsJSON = MakeShared<FJsonObject>();
+		PeerConnectionOptionsJSON->SetArrayField(FString(TEXT("iceServers")), IceServersArr);
+
+		TSharedRef<FJsonObject> ConfigJSON = MakeShared<FJsonObject>();
+		ConfigJSON->SetStringField(FString(TEXT("type")), FString(TEXT("config")));
+		ConfigJSON->SetObjectField(FString(TEXT("peerConnectionOptions")), PeerConnectionOptionsJSON);
+		return Utils::ToString(ConfigJSON);
+	}
+
+	TSharedPtr<FJsonObject> FSignallingServer::ParseToJSON(FString Message) const
+	{
+		TSharedPtr<FJsonObject> JSONObj = MakeShared<FJsonObject>();
+		bool bSuccess = Utils::Jsonify(Message, JSONObj);
+		if(!bSuccess)
+		{
+			UE_LOG(LogPixelStreamingServers, Warning, TEXT("Could not deserialize message into a JSON object - message=%s"), *Message);
+		}
+		return JSONObj;
+	}
+
+	bool FSignallingServer::GetMessageType(TSharedPtr<FJsonObject> JSONObj, FString& OutMessageType) const
+	{
+		if (!JSONObj->TryGetStringField(FString(TEXT("type")), OutMessageType))
+		{
+			FString Message;
+			bool bSuccess = Utils::Jsonify(Message, JSONObj);
+			UE_LOG(LogPixelStreamingServers, Warning, TEXT("Message did not contain a field called 'type' - message=%s"), *Message);
+			return false;
+		}
+		return true;
+	}
+
+	void FSignallingServer::OnStreamerMessage(uint16 ConnectionId, TArrayView<uint8> Message)
+	{
+		FString Msg = Utils::ToString(Message);
+		UE_LOG(LogPixelStreamingServers, Log, TEXT("From Streamer id=%d: %s"), ConnectionId, *Msg);
+
+		TSharedPtr<FJsonObject> JSONObj = ParseToJSON(Msg);
+		FString MsgType;
+		bool bHasMsgType = GetMessageType(JSONObj, MsgType);
+		if(!bHasMsgType) { return; }
+
+		// Got ping, respond with pong ---  { "type" : "pong", "time" : "1657509114" }
+		if(MsgType == FString(TEXT("ping")))
+		{
+			const double UnixTime = FDateTime::UtcNow().ToUnixTimestamp();
+			TSharedRef<FJsonObject> PongJSON = MakeShared<FJsonObject>();
+			PongJSON->SetStringField(FString(TEXT("type")), FString(TEXT("pong")));
+			PongJSON->SetNumberField(FString(TEXT("time")), UnixTime);
+			StreamersWS->Send(ConnectionId, Utils::ToString(PongJSON));
+			return;
+		}
+		
+		// All other message types require a `playerId` field to be valid.
+		FString PlayerId;
+		if (!JSONObj->TryGetStringField(TEXT("playerId"), PlayerId))
+		{
+			UE_LOG(LogPixelStreamingServers, Warning, TEXT("Message did not contain a field called 'playerId' - message=%s"), *Msg);
+			return;
+		}
+
+		// Convert PlayerId into a uint16 so we can look it up
+		uint16 PlayerIdInt = FCString::Atoi(*PlayerId);
+
+		// As message are going to the player they don't actually need the playerId field, the field exists only so we know who to send it to.
+		JSONObj->RemoveField(TEXT("playerId"));
+
+		// All other message types just get routed to the player that streamer wishes to recieve it
+		bool bForwardStraightToPlayer = 
+			MsgType == FString(TEXT("offer")) ||
+			MsgType == FString(TEXT("answer")) ||
+			MsgType == FString(TEXT("iceCandidate"));
+
+		if(bForwardStraightToPlayer)
+		{
+			PlayersWS->Send(PlayerIdInt, Utils::ToString(JSONObj.ToSharedRef()));
+			return;
+		}
+		else if(MsgType == FString(TEXT("disconnectPlayer")))
+		{
+			PlayersWS->Close(PlayerIdInt);
+			return;
+		}
+		else
+		{
+			// Unsupported message type
+			UE_LOG(LogPixelStreamingServers, Warning, TEXT("Unsupported message type receieved from streamer, message=%s"), *Msg);
+			return;
+		}
+	}
+
+	void FSignallingServer::OnPlayerMessage(uint16 ConnectionId, TArrayView<uint8> Message)
+	{
+		uint16 OutStreamerId;
+		if(!StreamersWS->GetFirstConnection(OutStreamerId))
+		{
+			UE_LOG(LogPixelStreamingServers, Warning, TEXT("Player message ignored because no streamer is yet connected to the signalling server."));
+			return;
+		}
+
+		FString Msg = Utils::ToString(Message);
+		UE_LOG(LogPixelStreamingServers, Log, TEXT("From Player id=%d: %s"), ConnectionId, *Msg);
+
+		TSharedPtr<FJsonObject> JSONObj = ParseToJSON(Msg);
+		FString MsgType;
+		bool bHasMsgType = GetMessageType(JSONObj, MsgType);
+		if(!bHasMsgType) { return; }
+
+		// Add player id to any messages going to streamer so streamer knows who sent it
+		JSONObj->SetStringField(FString(TEXT("playerId")), FString::FromInt(ConnectionId));
+
+		bool bForwardStraightToStreamer =
+			MsgType == FString(TEXT("answer")) ||
+			MsgType == FString(TEXT("offer")) ||
+			MsgType == FString(TEXT("iceCandidate")) ||
+			MsgType == FString(TEXT("dataChannelRequest")) ||
+			MsgType == FString(TEXT("peerDataChannelsReady"));
+
+		if(bForwardStraightToStreamer)
+		{
+			StreamersWS->Send(OutStreamerId, Utils::ToString(JSONObj.ToSharedRef()));
+			return;
+		}
+		else if(MsgType == FString(TEXT("stats")))
+		{
+			UE_LOG(LogPixelStreamingServers, Log, TEXT("Player stats = \n %s"), *Msg);
+			return;
+		}
+		else
+		{
+			// Unsupported message type
+			UE_LOG(LogPixelStreamingServers, Warning, TEXT("Unsupported message type receieved from player, message=%s"), *Msg);
+			return;
+		}
+	}
+
+	void FSignallingServer::OnPlayerConnected(uint16 ConnectionId)
+	{
+		// Todo (Luke): If we did want to support multiple streamers we could grab streamer id from url params perhaps. 
+		// For now we just grab the first streamer websocket connection.
+
+		uint16 OutStreamerId;
+		if(StreamersWS->GetFirstConnection(OutStreamerId))
+		{
+			// Send "playerConnected" message to streamer which kicks off making a new RTC connection
+			UE_LOG(LogPixelStreamingServers, Log, TEXT("Player connected over websocket. PlayerId=%d"), ConnectionId);
+			TSharedRef<FJsonObject> OnPlayerConnectedJSON = MakeShared<FJsonObject>();
+			OnPlayerConnectedJSON->SetStringField(FString(TEXT("type")), FString(TEXT("playerConnected")));
+			OnPlayerConnectedJSON->SetStringField(FString(TEXT("playerId")), FString::FromInt(ConnectionId));
+			OnPlayerConnectedJSON->SetBoolField(FString(TEXT("dataChannel")), true);
+			OnPlayerConnectedJSON->SetBoolField(FString(TEXT("sfu")), false);
+			FString JSONString = Utils::ToString(OnPlayerConnectedJSON);
+
+			UE_LOG(LogPixelStreamingServers, Log, TEXT("Sending to streamer: %s"), *JSONString);
+			StreamersWS->Send(OutStreamerId, JSONString);
+
+			// Send config to newly connected player, which kicks off making a new RTC connection
+			FString ConfigJSON = CreateConfigJSON();
+			UE_LOG(LogPixelStreamingServers, Log, TEXT("Sending to player %d: %s"), ConnectionId, *ConfigJSON);
+			PlayersWS->Send(ConnectionId, ConfigJSON);
+		}
+		else
+		{
+			UE_LOG(LogPixelStreamingServers, Warning, TEXT("Player connection dropped because no streamer is yet connected to the signalling server."));
+			return;
+		}
+	}
+
+} // UE::PixelStreamingServers
+
diff --git a/Engine/Plugins/Media/PixelStreaming/Source/PixelStreamingServers/Private/SignallingServer.h b/Engine/Plugins/Media/PixelStreaming/Source/PixelStreamingServers/Private/SignallingServer.h
new file mode 100644
--- /dev/null
+++ b/Engine/Plugins/Media/PixelStreaming/Source/PixelStreamingServers/Private/SignallingServer.h
@@ -0,0 +1,53 @@
+// Copyright Epic Games, Inc. All Rights Reserved.
+
+#pragma once
+
+#include "ServerBase.h"
+#include "WebSocketProbe.h"
+#include "WebSocketServerWrapper.h"
+#include "Templates/UniquePtr.h"
+#include "Templates/SharedPointer.h"
+#include "Dom/JsonObject.h"
+
+namespace UE::PixelStreamingServers
+{
+
+	/**
+	* A native C++ implementation of a Pixel Streaming signalling server (similar to cirrus.js).
+	* This implementation is missing some features such as SFU, matchmaker, and Webserver support.
+	* The purpose of this implementation is to have signalling server callable from UE, without bundling Cirrus itself.
+	**/
+	class FSignallingServer : public FServerBase
+	{
+
+	public:
+
+		virtual ~FSignallingServer() = default;
+
+		// Begin FServerBase
+		virtual void Stop() override;
+		virtual bool TestConnection() override;
+		virtual bool LaunchImpl(FLaunchArgs& InLaunchArgs, TMap<EEndpoint, FURL>& OutEndpoints) override;
+		virtual FString GetPathOnDisk() override;
+		// End FServerBase
+
+	protected:
+		void OnStreamerConnected(uint16 ConnectionId);
+		void OnStreamerMessage(uint16 ConnectionId, TArrayView<uint8> Message);
+
+		void OnPlayerMessage(uint16 ConnectionId, TArrayView<uint8> Message);
+		void OnPlayerConnected(uint16 ConnectionId);
+
+		TArray<FWebSocketHttpMount> GenerateDirectoriesToServe() const;
+		FString CreateConfigJSON() const;
+		TSharedPtr<FJsonObject> ParseToJSON(FString Message) const;
+		bool GetMessageType(TSharedPtr<FJsonObject> JSONObj, FString& OutMessageType) const;
+
+	private:
+		TUniquePtr<FWebSocketProbe> Probe;
+		TUniquePtr<FWebSocketServerWrapper> StreamersWS;
+		TUniquePtr<FWebSocketServerWrapper> PlayersWS;
+		
+	};
+
+} // UE::PixelStreamingServers
\ No newline at end of file
diff --git a/Engine/Plugins/Media/PixelStreaming/Source/PixelStreamingServers/Private/Tests/ServerTests.cpp b/Engine/Plugins/Media/PixelStreaming/Source/PixelStreamingServers/Private/Tests/ServerTests.cpp
new file mode 100644
--- /dev/null
+++ b/Engine/Plugins/Media/PixelStreaming/Source/PixelStreamingServers/Private/Tests/ServerTests.cpp
@@ -0,0 +1,195 @@
+// Copyright Epic Games, Inc. All Rights Reserved.
+
+#include "Misc/AutomationTest.h"
+#include "Tests/AutomationCommon.h"
+#include "PixelStreamingServers.h"
+#include "Misc/Paths.h"
+#include "PixelStreamingServersLog.h"
+#include "ServerUtils.h"
+
+#if WITH_DEV_AUTOMATION_TESTS
+
+namespace UE::PixelStreamingServers
+{
+
+	DEFINE_LATENT_AUTOMATION_COMMAND_ONE_PARAMETER(FWaitForServerOrTimeout, TSharedPtr<UE::PixelStreamingServers::IServer>, Server);
+	bool FWaitForServerOrTimeout::Update()
+	{
+		if (Server)
+		{
+			return Server->IsTimedOut() || Server->IsReady();
+		}
+		return true;
+	}
+
+	DEFINE_LATENT_AUTOMATION_COMMAND_ONE_PARAMETER(FCleanupServer, TSharedPtr<IServer>, Server);
+	bool FCleanupServer::Update()
+	{
+		if (Server)
+		{
+			Server->Stop();
+		}
+		return true;
+	}
+
+	static const int HttpPort = 85;
+	static const FString ExpectedWebserverAddress = FString::Printf(TEXT("http://127.0.0.1:%d"), HttpPort);
+	static const FString ExpectedPlayerWSAddress = FString::Printf(TEXT("ws://127.0.0.1:%d"), HttpPort);
+
+	static const int SFUPort = 8889;
+	static const FString ExpectedSFUAddress = FString::Printf(TEXT("ws://127.0.0.1:%d"), SFUPort);
+
+	static const int StreamerPort = 8989;
+	static const FString ExpectedStreamerAddress = FString::Printf(TEXT("ws://127.0.0.1:%d"), StreamerPort);
+
+	static const int MatchmakerPort = 9999;
+	static const FString ExpectedMatchmakerAddress = FString::Printf(TEXT("ws://127.0.0.1:%d"), MatchmakerPort);
+
+	static const bool bTestServerBinary = false;
+
+	FString GetCirrusBinaryAbsPath()
+	{
+		FString ServerPath = 
+			FPaths::EnginePluginsDir() / 
+			TEXT("Media") / 
+			TEXT("PixelStreaming") / 
+			TEXT("Resources") /
+			TEXT("WebServers") /
+			TEXT("SignallingWebServer");
+
+#if PLATFORM_WINDOWS
+		ServerPath = ServerPath / TEXT("cirrus.exe");
+#elif PLATFORM_LINUX
+		ServerPath = ServerPath / TEXT("cirrus");
+#else
+		UE_LOG(LogPixelStreaming, Error, TEXT("Unsupported platform for Pixel Streaming."));
+		return false
+#endif
+
+		ServerPath = FPaths::ConvertRelativePathToFull(ServerPath);
+
+		return ServerPath;
+	}
+
+	IMPLEMENT_SIMPLE_AUTOMATION_TEST(FLaunchDownloadedCirrusTest, "System.Plugins.PixelStreaming.LaunchDownloadedCirrus", EAutomationTestFlags::EditorContext | EAutomationTestFlags::ClientContext | EAutomationTestFlags::ProductFilter)
+	bool FLaunchDownloadedCirrusTest::RunTest(const FString& Parameters)
+	{
+		UE_LOG(LogPixelStreamingServers, Log, TEXT("----------- LaunchDownloadedCirrusTest -----------"));
+
+		TSharedPtr<IServer> SignallingServer = MakeCirrusServer();
+		FLaunchArgs LaunchArgs;
+		LaunchArgs.bPollUntilReady = true;
+		LaunchArgs.ReconnectionTimeoutSeconds = 30.0f;
+		LaunchArgs.ReconnectionIntervalSeconds = 2.0f;
+		LaunchArgs.ProcessArgs = FString::Printf(TEXT("--HttpPort=%d --SFUPort=%d --StreamerPort=%d --MatchmakerPort=%d"), HttpPort, SFUPort, StreamerPort, MatchmakerPort);
+
+		if(bTestServerBinary)
+		{
+			LaunchArgs.ServerBinaryOverridePath = GetCirrusBinaryAbsPath();
+		}
+
+		bool bLaunched = SignallingServer->Launch(LaunchArgs);
+
+		if(!bLaunched)
+		{
+			// If we were unable to launch this means some files were missing, we early exit here because this will always happen on Horde
+			// and we don't want a permanently failing test on Horde. 
+			// Todo: Determine way to only disable test on Horde but not locally, or to make it actually download the required scripts.
+			return true;
+		}
+
+		SignallingServer->OnReady.AddLambda([this](TMap<EEndpoint, FURL> Endpoints) {
+			TestTrue("Got server OnReady.", true);
+
+			FString ActualWebserverUrl = Utils::ToString(Endpoints[EEndpoint::Signalling_Webserver]);
+			UE_LOG(LogPixelStreamingServers, Log, TEXT("Http address for webserver. Actual=%s | Expected=%s"), *ActualWebserverUrl, *ExpectedWebserverAddress);
+			TestTrue(FString::Printf(TEXT("Http address for webserver. Actual=%s | Expected=%s"), *ActualWebserverUrl, *ExpectedWebserverAddress), 
+			ActualWebserverUrl == ExpectedWebserverAddress);
+
+			FString ActualStreamerUrl = Utils::ToString(Endpoints[EEndpoint::Signalling_Streamer]);
+			UE_LOG(LogPixelStreamingServers, Log, TEXT("Websocket address for streamer messages. Actual=%s | Expected=%s"), *ActualStreamerUrl, *ExpectedStreamerAddress);
+			TestTrue(FString::Printf(TEXT("Websocket address for streamer messages. Actual=%s | Expected=%s"), *ActualStreamerUrl, *ExpectedStreamerAddress), 
+			ActualStreamerUrl == ExpectedStreamerAddress);
+
+			FString ActualPlayersUrl = Utils::ToString(Endpoints[EEndpoint::Signalling_Players]);
+			UE_LOG(LogPixelStreamingServers, Log, TEXT("Websocket address for player messages. Actual=%s | Expected=%s"), *ActualPlayersUrl, *ExpectedPlayerWSAddress);
+			TestTrue(FString::Printf(TEXT("Websocket address for player messages. Actual=%s | Expected=%s"), *ActualPlayersUrl, *ExpectedPlayerWSAddress), 
+			ActualPlayersUrl == ExpectedPlayerWSAddress);
+
+			FString ActualSFUUrl = Utils::ToString(Endpoints[EEndpoint::Signalling_SFU]);
+			UE_LOG(LogPixelStreamingServers, Log, TEXT("Websocket address for SFU messages. Actual=%s | Expected=%s"), *ActualSFUUrl, *ExpectedSFUAddress);
+			TestTrue(FString::Printf(TEXT("Websocket address for SFU messages. Actual=%s | Expected=%s"), *ActualSFUUrl, *ExpectedSFUAddress), 
+			ActualSFUUrl == ExpectedSFUAddress);
+
+			FString ActualMatchmakerUrl = Utils::ToString(Endpoints[EEndpoint::Signalling_Matchmaker]);
+			UE_LOG(LogPixelStreamingServers, Log, TEXT("Websocket address for matchmaker messages. Actual=%s | Expected=%s"), *ActualMatchmakerUrl, *ExpectedMatchmakerAddress);
+			TestTrue(FString::Printf(TEXT("Websocket address for matchmaker messages. Actual=%s | Expected=%s"), *ActualMatchmakerUrl, *ExpectedMatchmakerAddress), 
+			ActualMatchmakerUrl == ExpectedMatchmakerAddress);
+		});
+
+		SignallingServer->OnFailedToReady.AddLambda([this]() {
+			TestTrue("Server was not ready.", false);
+		});
+
+		ADD_LATENT_AUTOMATION_COMMAND(FWaitForServerOrTimeout(SignallingServer));
+		ADD_LATENT_AUTOMATION_COMMAND(FCleanupServer(SignallingServer));
+
+		return true;
+	}
+
+	IMPLEMENT_SIMPLE_AUTOMATION_TEST(FLaunchEmbeddedCirrusTest, "System.Plugins.PixelStreaming.LaunchEmbeddedCirrus", EAutomationTestFlags::EditorContext | EAutomationTestFlags::ClientContext | EAutomationTestFlags::ProductFilter)
+	bool FLaunchEmbeddedCirrusTest::RunTest(const FString& Parameters)
+	{
+		UE_LOG(LogPixelStreamingServers, Log, TEXT("----------- LaunchEmbeddedCirrusTest -----------"));
+
+		TSharedPtr<IServer> SignallingServer = MakeSignallingServer();
+		FLaunchArgs LaunchArgs;
+		LaunchArgs.bPollUntilReady = true;
+		LaunchArgs.ReconnectionTimeoutSeconds = 30.0f;
+		LaunchArgs.ReconnectionIntervalSeconds = 2.0f;
+		LaunchArgs.ProcessArgs = FString::Printf(TEXT("--HttpPort=%d --StreamerPort=%d"), HttpPort, StreamerPort);
+
+		bool bLaunched = SignallingServer->Launch(LaunchArgs);
+		UE_LOG(LogPixelStreamingServers, Log, TEXT("Embedded cirrus launched: %s"), bLaunched ? TEXT("true") : TEXT("false"));
+		TestTrue("Embedded cirrus launched.", bLaunched);
+
+		if(!bLaunched)
+		{
+			return false;
+		}
+
+		SignallingServer->OnReady.AddLambda([this](TMap<EEndpoint, FURL> Endpoints) {
+			TestTrue("Got server OnReady.", true);
+
+			FString ActualWebserverUrl = Utils::ToString(Endpoints[EEndpoint::Signalling_Webserver]);
+			UE_LOG(LogPixelStreamingServers, Log, TEXT("Http address for webserver. Actual=%s | Expected=%s"), *ActualWebserverUrl, *ExpectedWebserverAddress);
+			TestTrue(FString::Printf(TEXT("Http address for webserver. Actual=%s | Expected=%s"), *ActualWebserverUrl, *ExpectedWebserverAddress), 
+			ActualWebserverUrl == ExpectedWebserverAddress);
+
+			FString ActualStreamerUrl = Utils::ToString(Endpoints[EEndpoint::Signalling_Streamer]);
+			UE_LOG(LogPixelStreamingServers, Log, TEXT("Websocket address for streamer messages. Actual=%s | Expected=%s"), *ActualStreamerUrl, *ExpectedStreamerAddress);
+			TestTrue(FString::Printf(TEXT("Websocket address for streamer messages. Actual=%s | Expected=%s"), *ActualStreamerUrl, *ExpectedStreamerAddress), 
+			ActualStreamerUrl == ExpectedStreamerAddress);
+
+			FString ActualPlayersUrl = Utils::ToString(Endpoints[EEndpoint::Signalling_Players]);
+			UE_LOG(LogPixelStreamingServers, Log, TEXT("Websocket address for player messages. Actual=%s | Expected=%s"), *ActualPlayersUrl, *ExpectedPlayerWSAddress);
+			TestTrue(FString::Printf(TEXT("Websocket address for player messages. Actual=%s | Expected=%s"), *ActualPlayersUrl, *ExpectedPlayerWSAddress), 
+			ActualPlayersUrl == ExpectedPlayerWSAddress);
+
+		});
+
+		SignallingServer->OnFailedToReady.AddLambda([this]() {
+			TestTrue("Server was not ready.", false);
+		});
+
+		ADD_LATENT_AUTOMATION_COMMAND(FWaitForServerOrTimeout(SignallingServer));
+		ADD_LATENT_AUTOMATION_COMMAND(FCleanupServer(SignallingServer));
+
+		return true;
+	}
+
+	// Todo test where create and teardown signalling server 10 times in quick succession with probe to ensure ports are freed
+
+} // namespace UE::PixelStreamingServers
+
+#endif // WITH_DEV_AUTOMATION_TESTS
\ No newline at end of file
diff --git a/Engine/Plugins/Media/PixelStreaming/Source/PixelStreamingServers/Private/WebSocketProbe.h b/Engine/Plugins/Media/PixelStreaming/Source/PixelStreamingServers/Private/WebSocketProbe.h
new file mode 100644
--- /dev/null
+++ b/Engine/Plugins/Media/PixelStreaming/Source/PixelStreamingServers/Private/WebSocketProbe.h
@@ -0,0 +1,51 @@
+// Copyright Epic Games, Inc. All Rights Reserved.
+
+#pragma once
+
+#include "IWebSocket.h"
+#include "HAL/ThreadSafeCounter.h"
+#include "WebSocketsModule.h"
+#include "Templates/SharedPointer.h"
+#include "PixelStreamingServersLog.h"
+#include "Engine/EngineBaseTypes.h"
+#include "ServerUtils.h"
+
+namespace UE::PixelStreamingServers
+{
+
+	/* 
+	* A utility class that tries to establish a websocket connection. 
+	* Useful for testing whether servers have come online yet.
+	*/
+	class FWebSocketProbe
+	{
+	private:
+		TSharedRef<IWebSocket> WebSocket;
+		FThreadSafeBool bShouldAttemptReconnect;
+
+	public:
+		FWebSocketProbe(FURL Url, TArray<FString> Protocols = TArray<FString>())
+			: WebSocket(FWebSocketsModule::Get().CreateWebSocket(Utils::ToString(Url), Protocols))
+			, bShouldAttemptReconnect(true)
+		{
+			WebSocket->OnConnectionError().AddLambda([Url, &bShouldAttemptReconnect = bShouldAttemptReconnect](const FString& Error) {
+				UE_LOG(LogPixelStreamingServers, Log, TEXT("Probing websocket %s | Msg= \"%s\" | Retrying..."), *Utils::ToString(Url),  *Error);
+				bShouldAttemptReconnect = true;
+			});
+		}
+
+		bool Probe()
+		{
+			bool bIsConnected = WebSocket->IsConnected();
+
+			if (!bIsConnected && bShouldAttemptReconnect)
+			{
+				WebSocket->Connect();
+				bShouldAttemptReconnect = false;
+			}
+
+			return bIsConnected;
+		}
+	};
+
+} // UE::PixelStreamingServers
\ No newline at end of file
diff --git a/Engine/Plugins/Media/PixelStreaming/Source/PixelStreamingServers/Private/WebSocketServerWrapper.cpp b/Engine/Plugins/Media/PixelStreaming/Source/PixelStreamingServers/Private/WebSocketServerWrapper.cpp
new file mode 100644
--- /dev/null
+++ b/Engine/Plugins/Media/PixelStreaming/Source/PixelStreamingServers/Private/WebSocketServerWrapper.cpp
@@ -0,0 +1,214 @@
+// Copyright Epic Games, Inc. All Rights Reserved.
+
+#include "WebSocketServerWrapper.h"
+#include "IWebSocketNetworkingModule.h"
+#include "PixelStreamingServersLog.h"
+#include "Modules/ModuleManager.h"
+
+namespace UE::PixelStreamingServers
+{
+
+	/*
+	* ------------- FWebSocketConnection -------------
+	*/
+
+	FThreadSafeCounter FWebSocketConnection::IdGenerator = FThreadSafeCounter(100);
+
+	FWebSocketConnection::FWebSocketConnection(INetworkingWebSocket* InSocketConnection)
+		: Id(IdGenerator.Increment())
+		, SocketConnection(InSocketConnection)
+	{
+	}
+
+	FWebSocketConnection::~FWebSocketConnection()
+	{
+		if (SocketConnection)
+		{
+			delete SocketConnection;
+			SocketConnection = nullptr;
+		}
+	}
+
+	uint16 FWebSocketConnection::GetId() const
+	{
+		return Id;
+	}
+
+	bool FWebSocketConnection::Send(FString Message) const
+	{
+		// Convert FString into uint8 array.
+		FTCHARToUTF8 UTF8String(*Message);
+		int32 MessageSize = UTF8String.Length();
+		TArray<uint8> MessageArr;
+		MessageArr.SetNum(MessageSize);
+		FMemory::Memcpy(MessageArr.GetData(), UTF8String.Get(), MessageSize);
+		
+		// Send the uint8 buffer
+		return SocketConnection->Send(MessageArr.GetData(), MessageSize, false);
+	}
+
+	void FWebSocketConnection::SetCallbacks()
+	{
+		if(SocketConnection)
+		{
+			SocketConnection->SetReceiveCallBack(OnPacketReceivedCallback);
+			SocketConnection->SetSocketClosedCallBack(OnClosedCallback);
+		}
+	}
+
+	/*
+	* ------------- FWebSocketServerWrapper -------------
+	*/
+
+	FWebSocketServerWrapper::FWebSocketServerWrapper()
+		: bLaunched(false)
+	{
+
+	}
+
+	FWebSocketServerWrapper::~FWebSocketServerWrapper()
+	{
+		Stop();
+	}
+
+	void FWebSocketServerWrapper::EnableWebServer(TArray<FWebSocketHttpMount> InDirectoriesToServe)
+	{
+		bEnableWebServer = true;
+		DirectoriesToServe = InDirectoriesToServe;
+	}
+
+	bool FWebSocketServerWrapper::Launch(uint16 Port)
+	{
+		WSServer = FModuleManager::Get().LoadModuleChecked<IWebSocketNetworkingModule>(TEXT("WebSocketNetworking")).CreateServer();
+		
+		if(bEnableWebServer)
+		{
+			WSServer->EnableHTTPServer(DirectoriesToServe);
+		}
+
+		OnClientConnectedCallback.BindRaw(this, &FWebSocketServerWrapper::OnConnectionOpened);
+		bLaunched = WSServer->Init(Port, OnClientConnectedCallback);
+		if(!bLaunched)
+		{
+			UE_LOG(LogPixelStreamingServers, Error, TEXT("Failed to launch Websocket server at ws://127.0.0.1:%d"), Port);
+			WSServer.Reset();
+			OnClientConnectedCallback.Unbind();
+		}
+		else
+		{
+			UE_LOG(LogPixelStreamingServers, Log, TEXT("Started Websocket server at ws://127.0.0.1:%d"), Port);
+		}
+		return bLaunched;
+	}
+
+	bool FWebSocketServerWrapper::IsLaunched() const
+	{
+		return bLaunched;
+	}
+
+	bool FWebSocketServerWrapper::HasConnections() const
+	{
+		return Connections.Num() > 0;
+	}
+
+	bool FWebSocketServerWrapper::GetFirstConnection(uint16& OutConnectionId) const
+	{
+		if(Connections.Num() > 0)
+		{
+			OutConnectionId = Connections.CreateConstIterator().Key();
+			return true;
+		}
+		return false;
+	}
+
+	void FWebSocketServerWrapper::Stop()
+	{
+		bLaunched = false;
+		Connections.Empty();
+		OnClientConnectedCallback.Unbind();
+		WSServer.Reset();
+	}
+
+	bool FWebSocketServerWrapper::Close(uint16 ConnectionId)
+	{
+		if(Connections.Contains(ConnectionId))
+		{
+			return Connections.Remove(ConnectionId) > 0;
+		}
+		else
+		{
+			UE_LOG(LogPixelStreamingServers, Warning, TEXT("Could not close websocket connection because there was no connection=%d."), ConnectionId);
+			return false;
+		}
+	}
+
+	bool FWebSocketServerWrapper::Send(uint16 ConnectionId, FString Message) const
+	{
+		if(Connections.Contains(ConnectionId))
+		{
+			return Connections[ConnectionId]->Send(Message);
+		}
+		else
+		{
+			UE_LOG(LogPixelStreamingServers, Warning, TEXT("Did not send websocket message because there was no connection=%d."), ConnectionId);
+			return false;
+		}
+	}
+
+	void FWebSocketServerWrapper::OnConnectionOpened(INetworkingWebSocket* Socket)
+	{
+		if(!Socket)
+		{
+			UE_LOG(LogPixelStreamingServers, Error, TEXT("Websocket client connected with a null socket."));
+			return;
+		}
+		else
+		{
+			UE_LOG(LogPixelStreamingServers, Log, TEXT("Websocket client connected. Remote=%s | Local=%s"), *Socket->RemoteEndPoint(true), *Socket->LocalEndPoint(true));
+		}
+
+		// Had a new client connect over websocket, store the connection with a unique ID.
+		TUniquePtr<FWebSocketConnection> Connection = MakeUnique<FWebSocketConnection>(Socket);
+		const uint16 Id = Connection->GetId();
+
+		// Bind to socket callbacks for messages/closed.
+		Connection->OnPacketReceivedCallback.BindRaw(this, &FWebSocketServerWrapper::OnPacketReceived, Id);
+		Connection->OnClosedCallback.BindRaw(this, &FWebSocketServerWrapper::OnConnectionClosed, Id);
+		Connection->SetCallbacks();
+
+		Connections.Add(Id, MoveTemp(Connection));
+
+		// Broadcast that we got a new websocket connection
+		OnOpenConnection.Broadcast(Id);
+	}
+
+	void FWebSocketServerWrapper::OnPacketReceived(void* Data, int32 Size, uint16 ConnectionId)
+	{
+		if(Size > 0)
+		{
+			TArrayView<uint8> DataView = MakeArrayView(static_cast<uint8*>(Data), Size);
+			OnMessage.Broadcast(ConnectionId, DataView);
+		}
+	}
+
+	void FWebSocketServerWrapper::OnConnectionClosed(uint16 ConnectionId)
+	{
+		Connections.Remove(ConnectionId);
+		OnClosedConnection.Broadcast(ConnectionId);
+	}
+
+	void FWebSocketServerWrapper::Tick(float DeltaTime)
+	{
+		if(!IsLaunched())
+		{
+			return;
+		}
+
+		if(WSServer)
+		{
+			// Tick the websocket server
+			WSServer->Tick();
+		}
+	}
+
+} // UE::PixelStreamingServers
\ No newline at end of file
diff --git a/Engine/Plugins/Media/PixelStreaming/Source/PixelStreamingServers/Private/WebSocketServerWrapper.h b/Engine/Plugins/Media/PixelStreaming/Source/PixelStreamingServers/Private/WebSocketServerWrapper.h
new file mode 100644
--- /dev/null
+++ b/Engine/Plugins/Media/PixelStreaming/Source/PixelStreamingServers/Private/WebSocketServerWrapper.h
@@ -0,0 +1,84 @@
+// Copyright Epic Games, Inc. All Rights Reserved.
+
+#pragma once
+
+#include "IWebSocketServer.h"
+#include "Templates/UniquePtr.h"
+#include "Templates/SharedPointer.h"
+#include "WebSocketNetworkingDelegates.h"
+#include "INetworkingWebSocket.h"
+#include "Tickable.h"
+#include "HAL/ThreadSafeCounter.h"
+#include "Containers/ArrayView.h"
+#include "HAL/ThreadSafeBool.h"
+
+namespace UE::PixelStreamingServers
+{
+
+	struct FWebSocketConnection
+	{
+	public:
+		FWebSocketConnection(INetworkingWebSocket* InSocketConnection);
+		~FWebSocketConnection();
+		uint16 GetId() const;
+		bool Send(FString Message) const;
+		void SetCallbacks();
+
+	public:
+		FWebSocketPacketReceivedCallBack OnPacketReceivedCallback;
+		FWebSocketInfoCallBack OnClosedCallback;
+
+	private:
+		static FThreadSafeCounter IdGenerator;
+		uint16 Id;
+		INetworkingWebSocket* SocketConnection;
+	};
+
+
+	class FWebSocketServerWrapper : public FTickableGameObject
+	{
+	public:
+		FWebSocketServerWrapper();
+		virtual ~FWebSocketServerWrapper();
+		void EnableWebServer(TArray<FWebSocketHttpMount> InDirectoriesToServe);
+		bool Launch(uint16 Port);
+		void Stop();
+		bool IsLaunched() const;
+		bool HasConnections() const;
+		bool Close(uint16 ConnectionId);
+		bool Send(uint16 ConnectionId, FString Message) const;
+		bool GetFirstConnection(uint16& OutConnectionId) const;
+
+		/* Begin FTickableGameObject */
+		virtual bool IsTickableWhenPaused() const { return true; }
+		virtual bool IsTickableInEditor() const { return true; }
+		virtual void Tick(float DeltaTime) override;
+		virtual bool IsAllowedToTick() const { return true; }
+		TStatId GetStatId() const { RETURN_QUICK_DECLARE_CYCLE_STAT(FWebSocketServerWrapper, STATGROUP_Tickables); }
+		/* End FTickableGameObject */
+
+	public:
+		DECLARE_MULTICAST_DELEGATE_OneParam(FOnNewWebSocketConnection, uint16 /*Connection Id*/);
+		FOnNewWebSocketConnection OnOpenConnection;
+
+		DECLARE_MULTICAST_DELEGATE_OneParam(FOnClosedWebSocketConnection, uint16 /*Connection Id*/);
+		FOnClosedWebSocketConnection OnClosedConnection;
+
+		DECLARE_MULTICAST_DELEGATE_TwoParams(FOnMessageWebSocketConnection, uint16 /*Connection Id*/, TArrayView<uint8> /*Message*/);
+		FOnMessageWebSocketConnection OnMessage;
+
+	protected:
+		virtual void OnConnectionOpened(INetworkingWebSocket* Socket);
+		virtual void OnPacketReceived(void* Data, int32 Size, uint16 ConnectionId);
+		virtual void OnConnectionClosed(uint16 ConnectionId);
+
+	private:
+		FThreadSafeBool bLaunched;
+		TUniquePtr<IWebSocketServer> WSServer;
+		FWebSocketClientConnectedCallBack OnClientConnectedCallback;
+		TMap<uint16, TUniquePtr<FWebSocketConnection>> Connections;
+		bool bEnableWebServer = false;
+		TArray<FWebSocketHttpMount> DirectoriesToServe;
+	};
+
+} // UE::PixelStreamingServers
\ No newline at end of file
diff --git a/Engine/Plugins/Media/PixelStreaming/Source/PixelStreamingServers/Public/PixelStreamingServers.h b/Engine/Plugins/Media/PixelStreaming/Source/PixelStreamingServers/Public/PixelStreamingServers.h
new file mode 100644
--- /dev/null
+++ b/Engine/Plugins/Media/PixelStreaming/Source/PixelStreamingServers/Public/PixelStreamingServers.h
@@ -0,0 +1,141 @@
+// Copyright Epic Games, Inc. All Rights Reserved.
+
+#pragma once
+
+#include "Templates/SharedPointer.h"
+#include "GenericPlatform/GenericPlatformMisc.h"
+#include "Containers/Map.h"
+#include "Engine/EngineBaseTypes.h"
+#include "HAL/ThreadSafeBool.h"
+#include "Misc/MonitoredProcess.h"
+
+/*
+* Utility namespace for launching Pixel Streaming servers signaling, selective forwarding, or matchmaking>
+* This utility can launch embedded servers from:
+* 1. C++ classes in this module. 
+* 2. From known server binaries.
+* 3. From the Github releases of the servers.
+*/
+namespace UE::PixelStreamingServers
+{
+	/**
+	 * Configuration to control behaviour when launching any of the Pixel Streaming servers.
+	 **/
+	struct PIXELSTREAMINGSERVERS_API FLaunchArgs
+	{
+		// Arguments passed to the actual server when its process is started.
+		FString ProcessArgs = TEXT("");
+
+		// If true poll until ready
+		bool bPollUntilReady = false;
+
+		// Reconnection timeout in seconds
+		float ReconnectionTimeoutSeconds = 30.0f;
+
+		// Reconnect interval in seconds.
+		float ReconnectionIntervalSeconds = 2.0f;
+
+		// Path the server binary to run instead of launching server by running scripts
+		TOptional<FString> ServerBinaryOverridePath;
+	};
+
+	/**
+	 * Endpoints for the various Pixel Streaming servers.
+	 **/
+	enum class PIXELSTREAMINGSERVERS_API EEndpoint
+	{
+		// The websocket signalling url between the server and the UE streamer - e.g. ws://localhost:8888
+		Signalling_Streamer,
+
+		// The websocket signalling url between the server and the players (aka. web browsers) - e.g. ws://localhost:80
+		Signalling_Players,
+
+		// The websocket signalling url between the server and the matchmaker server - e.g. ws://localhost:9999
+		Signalling_Matchmaker,
+
+		// The websocket signalling url between the server and the SFU server - e.g. ws://localhost:8889
+		Signalling_SFU,
+
+		// The http url for the webserver hosted within the signalling server - e.g. http://localhost
+		Signalling_Webserver
+	};
+
+	// ---------------------------------------------------------------------------------------------
+	typedef TMap<EEndpoint, FURL> FEndpoints;
+	DECLARE_MULTICAST_DELEGATE_OneParam(FOnReady, const FEndpoints& /* Endpoint urls */);
+	// ---------------------------------------------------------------------------------------------
+
+	/**
+	 * Interface for all Pixel Streaming servers.
+	**/
+	class PIXELSTREAMINGSERVERS_API IServer
+	{
+	public:
+		virtual ~IServer() = default;
+
+		/* Immediately stops the server. */
+		virtual void Stop() = 0;
+
+		/**
+		 * @return	The absolute path to the root directory that the server was launched from.
+		 **/
+		virtual FString GetPathOnDisk() = 0;
+
+		/**
+		 * @return	True if the server has been launched. Note: Launched does not necessarily mean it is connectible yet. Bind to OnReady for that.
+		 **/
+		virtual bool HasLaunched() = 0;
+
+		/**
+		 * Launch the server in a child process using the supplied launch arguments.
+		 * @param LaunchArgs	The launch arguments to control how the server is launched, including what args to pass to the child process.
+		 * @return True if the server was able to start launching, this can fail when launching child process servers where files must exist on disk.
+		 **/
+		virtual bool Launch(FLaunchArgs& InLaunchArgs) = 0;
+
+		/**
+		 * @return	True if the server has been connected to and is ready for new connections.
+		 **/
+		virtual bool IsReady() = 0;
+
+		/**
+		 * @return	True if the server has timed out while trying to establish a connection.
+		 **/
+		virtual bool IsTimedOut() = 0;
+	
+	public:
+		// Delegate fired when the server is ready for connections, first parameter is a map of all supported endpoints and their urls.
+		FOnReady OnReady;
+
+		DECLARE_MULTICAST_DELEGATE(FOnFailedToReady);
+		/* Can fire when the server is unable to be contacted or connecting to it timed out. */
+		FOnFailedToReady OnFailedToReady;
+
+	};
+
+	/* -------------- Static utility methods for working with Pixel Streaming servers. ----------------- */
+
+	/**
+	 * Creates a NodeJS Cirrus signalling server by launching it as a child process.
+	 * Note: Calling this method does not launch the server. You should call Launch() yourself
+	 * once you have bound to appropriate delegates such as OnReady.
+	 * @return	The Cirrus signalling server wrapped in a server object for simple management.
+	 **/
+	PIXELSTREAMINGSERVERS_API TSharedPtr<IServer> MakeCirrusServer();
+
+	/**
+	 * Creates a native C++ signalling server (similar to cirrus.js) with no dependencies launched inside the Unreal Engine process.
+	 * Note: Calling this method does not launch the server. You should call Launch() yourself
+	 * once you have bound to appropriate delegates such as OnReady.
+	 * @return	The embedded signalling server.
+	 **/
+	PIXELSTREAMINGSERVERS_API TSharedPtr<IServer> MakeSignallingServer();
+
+	/**
+	 * Download the Pixel Streaming servers using the `get_ps_servers` scripts.
+	 * @param bSkipIfPresent Servers will not be downloaded if they are already present.
+	 * @return The child process that is used to download the servers.
+	 **/
+	PIXELSTREAMINGSERVERS_API TSharedPtr<FMonitoredProcess> DownloadPixelStreamingServers(bool bSkipIfPresent);
+
+} // namespace UE::PixelStreaming
diff --git a/Engine/Plugins/Media/PixelStreamingVCam/Source/PixelStreamingVCam/PixelStreamingVCam.Build.cs b/Engine/Plugins/Media/PixelStreamingVCam/Source/PixelStreamingVCam/PixelStreamingVCam.Build.cs
--- a/Engine/Plugins/Media/PixelStreamingVCam/Source/PixelStreamingVCam/PixelStreamingVCam.Build.cs
+++ b/Engine/Plugins/Media/PixelStreamingVCam/Source/PixelStreamingVCam/PixelStreamingVCam.Build.cs
@@ -16,7 +16,8 @@ namespace UnrealBuildTool.Rules
 				"InputDevice",
 				"MediaIOCore",
 				"VCamCore",
-				"PixelStreaming"
+				"PixelStreaming",
+				"PixelStreamingServers"
 			});
 
 			PrivateDependencyModuleNames.AddRange(new string[]
@@ -25,6 +26,7 @@ namespace UnrealBuildTool.Rules
 				"Core",
 				"CoreUObject",
 				"Engine",
+				"Projects",
 				"UMG",
 				"RHI",
 				"CinematicCamera",
diff --git a/Engine/Plugins/Media/PixelStreamingVCam/Source/PixelStreamingVCam/Private/PixelStreamingMediaOutput.cpp b/Engine/Plugins/Media/PixelStreamingVCam/Source/PixelStreamingVCam/Private/PixelStreamingMediaOutput.cpp
--- a/Engine/Plugins/Media/PixelStreamingVCam/Source/PixelStreamingVCam/Private/PixelStreamingMediaOutput.cpp
+++ b/Engine/Plugins/Media/PixelStreamingVCam/Source/PixelStreamingVCam/Private/PixelStreamingMediaOutput.cpp
@@ -53,11 +53,15 @@ void UPixelStreamingMediaOutput::StartStreaming()
 {
 	if (Streamer)
 	{
+		Streamer->SetSignallingServerURL(SignallingServerURL);
 		Streamer->SetVideoInput(Capture->GetVideoInput());
 		Streamer->SetTargetViewport(Capture->GetViewport().Get());
 		Streamer->SetTargetWindow(Capture->GetViewport()->FindWindow());
-		Streamer->SetSignallingServerURL(SignallingServerURL);
-		Streamer->StartStreaming();
+
+		if(!Streamer->IsStreaming())
+		{
+			Streamer->StartStreaming();
+		}
 	}
 }
 
diff --git a/Engine/Plugins/Media/PixelStreamingVCam/Source/PixelStreamingVCam/Private/PixelStreamingVCamLog.cpp b/Engine/Plugins/Media/PixelStreamingVCam/Source/PixelStreamingVCam/Private/PixelStreamingVCamLog.cpp
new file mode 100644
--- /dev/null
+++ b/Engine/Plugins/Media/PixelStreamingVCam/Source/PixelStreamingVCam/Private/PixelStreamingVCamLog.cpp
@@ -0,0 +1,5 @@
+// Copyright Epic Games, Inc. All Rights Reserved.
+
+#include "PixelStreamingVCamLog.h"
+
+DEFINE_LOG_CATEGORY(LogPixelStreamingVCam);
\ No newline at end of file
diff --git a/Engine/Plugins/Media/PixelStreamingVCam/Source/PixelStreamingVCam/Private/PixelStreamingVCamLog.h b/Engine/Plugins/Media/PixelStreamingVCam/Source/PixelStreamingVCam/Private/PixelStreamingVCamLog.h
new file mode 100644
--- /dev/null
+++ b/Engine/Plugins/Media/PixelStreamingVCam/Source/PixelStreamingVCam/Private/PixelStreamingVCamLog.h
@@ -0,0 +1,8 @@
+// Copyright Epic Games, Inc. All Rights Reserved.
+
+#pragma once
+
+#include "HAL/IConsoleManager.h"
+#include "Logging/LogMacros.h"
+
+DECLARE_LOG_CATEGORY_EXTERN(LogPixelStreamingVCam, Log, All);
\ No newline at end of file
diff --git a/Engine/Plugins/Media/PixelStreamingVCam/Source/PixelStreamingVCam/Private/PixelStreamingVCamUtil.h b/Engine/Plugins/Media/PixelStreamingVCam/Source/PixelStreamingVCam/Private/PixelStreamingVCamUtil.h
new file mode 100644
--- /dev/null
+++ b/Engine/Plugins/Media/PixelStreamingVCam/Source/PixelStreamingVCam/Private/PixelStreamingVCamUtil.h
@@ -0,0 +1,35 @@
+// Copyright Epic Games, Inc. All Rights Reserved.
+
+#pragma once
+
+#include "Misc/Paths.h"
+#include "Interfaces/IPluginManager.h"
+#include "PixelStreamingVCamLogging.h"
+
+class FPixelStreamingVCamUtil
+{
+public:
+	static bool GetSignallingServerBinaryAbsPath(FString& OutPath)
+	{
+		TSharedPtr<IPlugin> VCamPlugin = IPluginManager::Get().FindPlugin("PixelStreamingVCam");
+
+		if(!VCamPlugin.IsValid())
+		{
+			return false;
+		}
+
+		FString PluginContentDir = FPaths::ConvertRelativePathToFull(VCamPlugin->GetContentDir());
+
+#if PLATFORM_WINDOWS
+		OutPath = PluginContentDir / TEXT("SignallingWebServer") / TEXT("cirrus.exe");
+#elif PLATFORM_LINUX
+		OutPath = PluginContentDir / TEXT("SignallingWebServer") / TEXT("cirrus");
+#else
+		UE_LOG(LogPixelStreamingVCam, Log, TEXT("Unsupported platform for Pixel Streaming VCam plugin - only Windows and Linux are supported."))
+		return false;
+#endif
+
+		// Check if the out path we set is actually valid
+		return FPaths::FileExists(OutPath);
+	}
+};
diff --git a/Engine/Plugins/Media/PixelStreamingVCam/Source/PixelStreamingVCam/Private/VCamPixelStreamingSession.cpp b/Engine/Plugins/Media/PixelStreamingVCam/Source/PixelStreamingVCam/Private/VCamPixelStreamingSession.cpp
--- a/Engine/Plugins/Media/PixelStreamingVCam/Source/PixelStreamingVCam/Private/VCamPixelStreamingSession.cpp
+++ b/Engine/Plugins/Media/PixelStreamingVCam/Source/PixelStreamingVCam/Private/VCamPixelStreamingSession.cpp
@@ -13,8 +13,8 @@
 #include "Serialization/MemoryReader.h"
 #include "VCamPixelStreamingSubsystem.h"
 #include "VCamPixelStreamingLiveLink.h"
+#include "PixelStreamingVCamLog.h"
 
-DEFINE_LOG_CATEGORY(LogVCamOutputProvider);
 namespace VCamPixelStreamingSession
 {
 	static const FName LevelEditorName(TEXT("LevelEditor"));
@@ -23,10 +23,6 @@ namespace VCamPixelStreamingSession
 
 void UVCamPixelStreamingSession::Initialize()
 {
-	if (!bInitialized && (MediaOutput == nullptr))
-	{
-		MediaOutput = NewObject<UPixelStreamingMediaOutput>(GetTransientPackage(), UPixelStreamingMediaOutput::StaticClass());
-	}
 	Super::Initialize();
 }
 
@@ -56,11 +52,16 @@ void UVCamPixelStreamingSession::Activate()
 
 	if (!bInitialized)
 	{
-		UE_LOG(LogVCamOutputProvider, Warning, TEXT("Trying to start Pixel Streaming, but has not been initialized yet"));
+		UE_LOG(LogPixelStreamingVCam, Warning, TEXT("Trying to start Pixel Streaming, but has not been initialized yet"));
 		SetActive(false);
 		return;
 	}
 
+	if (MediaOutput == nullptr)
+	{
+		MediaOutput = NewObject<UPixelStreamingMediaOutput>(GetTransientPackage(), UPixelStreamingMediaOutput::StaticClass());
+	}
+
 	UEditorPerformanceSettings* Settings = GetMutableDefault<UEditorPerformanceSettings>();
 	bOldThrottleCPUWhenNotForeground = Settings->bThrottleCPUWhenNotForeground;
 	if (PreventEditorIdle)
@@ -71,15 +72,14 @@ void UVCamPixelStreamingSession::Activate()
 
 	if (StartSignallingServer)
 	{
-		CirrusProcess = UE::PixelStreaming::Servers::MakeSignallingServer();
-		UE::PixelStreaming::Servers::FLaunchArgs LaunchArgs;
-		LaunchArgs.bEphemeral = false;
-		LaunchArgs.ProcessArgs = FString::Printf(TEXT("--StreamerPort=%s --HttpPort=%s"), *FString::FromInt(PortNumber), *FString::FromInt(HttpPort));
-		CirrusProcess->Launch(LaunchArgs);
+		if (UVCamPixelStreamingSubsystem* PixelStreamingSubsystem = UVCamPixelStreamingSubsystem::Get())
+		{
+			PixelStreamingSubsystem->LaunchSignallingServer(PortNumber, HttpPort);
+		}
 	}
 
 	MediaOutput->SetSignallingServerURL(FString::Printf(TEXT("%s:%s"), *IP, *FString::FromInt(PortNumber)));
-	UE_LOG(LogVCamOutputProvider, Log, TEXT("Activating PixelStreaming VCam Session. Endpoint: %s:%s"), *IP, *FString::FromInt(PortNumber));
+	UE_LOG(LogPixelStreamingVCam, Log, TEXT("Activating PixelStreaming VCam Session. Endpoint: %s:%s"), *IP, *FString::FromInt(PortNumber));
 
 	MediaCapture = Cast<UPixelStreamingMediaCapture>(MediaOutput->CreateMediaCapture());
 
@@ -92,11 +92,11 @@ void UVCamPixelStreamingSession::Activate()
 		if (ComposureProvider->FinalOutputRenderTarget)
 		{
 			MediaCapture->CaptureTextureRenderTarget2D(ComposureProvider->FinalOutputRenderTarget, Options);
-			UE_LOG(LogVCamOutputProvider, Log, TEXT("PixelStreaming set with ComposureRenderTarget"));
+			UE_LOG(LogPixelStreamingVCam, Log, TEXT("PixelStreaming set with ComposureRenderTarget"));
 		}
 		else
 		{
-			UE_LOG(LogVCamOutputProvider, Warning, TEXT("PixelStreaming Composure usage was requested, but the specified ComposureOutputProvider has no FinalOutputRenderTarget set"));
+			UE_LOG(LogPixelStreamingVCam, Warning, TEXT("PixelStreaming Composure usage was requested, but the specified ComposureOutputProvider has no FinalOutputRenderTarget set"));
 		}
 	}
 	else
@@ -105,7 +105,7 @@ void UVCamPixelStreamingSession::Activate()
 		if (TSharedPtr<FSceneViewport> PinnedSceneViewport = SceneViewport.Pin())
 		{
 			MediaCapture->CaptureSceneViewport(PinnedSceneViewport, Options);
-			UE_LOG(LogVCamOutputProvider, Log, TEXT("PixelStreaming set with viewport"));
+			UE_LOG(LogPixelStreamingVCam, Log, TEXT("PixelStreaming set with viewport"));
 		}
 	}
 
@@ -122,7 +122,7 @@ void UVCamPixelStreamingSession::Activate()
 			static constexpr int ExpectedDataSize = 1 + (16 * 4) + 8;
 			if (Data.Num() != ExpectedDataSize)
 			{
-				UE_LOG(LogVCamOutputProvider, Warning, TEXT("ARKitTransform data received but with invalid size. Expects %d bytes but received %d"),
+				UE_LOG(LogPixelStreamingVCam, Warning, TEXT("ARKitTransform data received but with invalid size. Expects %d bytes but received %d"),
 					ExpectedDataSize,
 					Data.Num());
 
@@ -174,9 +174,9 @@ void UVCamPixelStreamingSession::Activate()
 
 void UVCamPixelStreamingSession::StopSignallingServer()
 {
-	if (CirrusProcess.IsValid())
+	if (UVCamPixelStreamingSubsystem* PixelStreamingSubsystem = UVCamPixelStreamingSubsystem::Get())
 	{
-		CirrusProcess->Stop();
+		PixelStreamingSubsystem->StopSignallingServer();
 	}
 }
 
@@ -232,9 +232,17 @@ void UVCamPixelStreamingSession::PostEditChangeProperty(FPropertyChangedEvent& P
 		static FName NAME_IP = GET_MEMBER_NAME_CHECKED(UVCamPixelStreamingSession, IP);
 		static FName NAME_PortNumber = GET_MEMBER_NAME_CHECKED(UVCamPixelStreamingSession, PortNumber);
 		static FName NAME_FromComposureOutputProviderIndex = GET_MEMBER_NAME_CHECKED(UVCamPixelStreamingSession, FromComposureOutputProviderIndex);
+		
 		if ((Property->GetFName() == NAME_IP) || (Property->GetFName() == NAME_PortNumber) || (Property->GetFName() == NAME_FromComposureOutputProviderIndex))
 		{
-			SetActive(!bIsActive);
+			SetActive(false);
+		}
+		else if(Property->GetFName() == FName(TEXT("bIsActive")))
+		{
+			if(bIsActive) {
+				// this is a hack to suppress double activation, livelink tick will activate it
+				return;
+			}
 		}
 	}
 	Super::PostEditChangeProperty(PropertyChangedEvent);
diff --git a/Engine/Plugins/Media/PixelStreamingVCam/Source/PixelStreamingVCam/Private/VCamPixelStreamingSubsystem.cpp b/Engine/Plugins/Media/PixelStreamingVCam/Source/PixelStreamingVCam/Private/VCamPixelStreamingSubsystem.cpp
--- a/Engine/Plugins/Media/PixelStreamingVCam/Source/PixelStreamingVCam/Private/VCamPixelStreamingSubsystem.cpp
+++ b/Engine/Plugins/Media/PixelStreamingVCam/Source/PixelStreamingVCam/Private/VCamPixelStreamingSubsystem.cpp
@@ -4,6 +4,7 @@
 #include "ILiveLinkClient.h"
 #include "VCamPixelStreamingLiveLink.h"
 #include "VCamPixelStreamingSession.h"
+#include "PixelStreamingVCamLog.h"
 
 void UVCamPixelStreamingSubsystem::Initialize(FSubsystemCollectionBase& Collection)
 {
@@ -53,4 +54,45 @@ void UVCamPixelStreamingSubsystem::UnregisterActiveOutputProvider(UVCamPixelStre
 	{
 		LiveLinkSource->RemoveSubject(OutputProvider->GetFName());
 	}
+}
+
+TSharedPtr<UE::PixelStreamingServers::IServer> UVCamPixelStreamingSubsystem::LaunchSignallingServer(int StreamerPort, int PlayerPort)
+{
+	bool bAlreadyLaunched = Server.IsValid() && Server->HasLaunched();
+	if(bAlreadyLaunched)
+	{
+		return Server;
+	}
+
+	// Download Pixel Streaming servers/frontend if we want to use a browser to view Pixel Streaming output
+	// but only attempt this is we haven't already started a download before.
+	if(!DownloadProcess.IsValid())
+	{
+		DownloadProcess = UE::PixelStreamingServers::DownloadPixelStreamingServers(/*bSkipIfPresent*/ true);
+		if(DownloadProcess.IsValid())
+		{
+			DownloadProcess->OnCompleted().BindLambda([this, StreamerPort, PlayerPort](int ExitCode){
+				UE_LOG(LogPixelStreamingVCam, Log, TEXT("Pixel Streaming frontend downloaded, restarting signalling server to use it."));
+				StopSignallingServer();
+				LaunchSignallingServer(StreamerPort, PlayerPort);
+			});
+		}
+	}
+
+	// Launch signalling server
+	Server = UE::PixelStreamingServers::MakeSignallingServer();
+	UE::PixelStreamingServers::FLaunchArgs LaunchArgs;
+	LaunchArgs.bPollUntilReady = false;
+	LaunchArgs.ProcessArgs = FString::Printf(TEXT("--StreamerPort=%d --HttpPort=%d"), StreamerPort, PlayerPort);
+	Server->Launch(LaunchArgs);
+	return Server;
+}
+
+void UVCamPixelStreamingSubsystem::StopSignallingServer()
+{
+	if(Server.IsValid())
+	{
+		Server->Stop();
+		Server.Reset();
+	}
 }
\ No newline at end of file
diff --git a/Engine/Plugins/Media/PixelStreamingVCam/Source/PixelStreamingVCam/Public/VCamPixelStreamingSession.h b/Engine/Plugins/Media/PixelStreamingVCam/Source/PixelStreamingVCam/Public/VCamPixelStreamingSession.h
--- a/Engine/Plugins/Media/PixelStreamingVCam/Source/PixelStreamingVCam/Public/VCamPixelStreamingSession.h
+++ b/Engine/Plugins/Media/PixelStreamingVCam/Source/PixelStreamingVCam/Public/VCamPixelStreamingSession.h
@@ -67,5 +67,5 @@ private:
 	FHitResult LastViewportTouchResult;
 	bool bUsingDummyUMG = false;
 	bool bOldThrottleCPUWhenNotForeground;
-	TSharedPtr<UE::PixelStreaming::Servers::FServerBase> CirrusProcess;
+	TSharedPtr<UE::PixelStreamingServers::IServer> CirrusProcess;
 };
diff --git a/Engine/Plugins/Media/PixelStreamingVCam/Source/PixelStreamingVCam/Public/VCamPixelStreamingSubsystem.h b/Engine/Plugins/Media/PixelStreamingVCam/Source/PixelStreamingVCam/Public/VCamPixelStreamingSubsystem.h
--- a/Engine/Plugins/Media/PixelStreamingVCam/Source/PixelStreamingVCam/Public/VCamPixelStreamingSubsystem.h
+++ b/Engine/Plugins/Media/PixelStreamingVCam/Source/PixelStreamingVCam/Public/VCamPixelStreamingSubsystem.h
@@ -3,6 +3,7 @@
 
 #include "CoreMinimal.h"
 #include "Subsystems/EngineSubsystem.h"
+#include "PixelStreamingServers.h"
 #include "VCamPixelStreamingSubsystem.generated.h"
 
 class FPixelStreamingLiveLinkSource;
@@ -19,6 +20,9 @@ public:
 	static UVCamPixelStreamingSubsystem* Get();
 	void RegisterActiveOutputProvider(UVCamPixelStreamingSession* OutputProvider);
 	void UnregisterActiveOutputProvider(UVCamPixelStreamingSession* OutputProvider);
+
+	TSharedPtr<UE::PixelStreamingServers::IServer> LaunchSignallingServer(int StreamerPort, int PlayerPort);
+	void StopSignallingServer();
 	
 	// An associated Live Link Source shared by all output providers
 	TSharedPtr<FPixelStreamingLiveLinkSource> LiveLinkSource;
@@ -26,4 +30,10 @@ private:
 	// Keep track of which output providers are currently active
 	UPROPERTY(Transient)
 	TArray<UVCamPixelStreamingSession*> ActiveOutputProviders;
+
+	// Signalling/webserver
+	TSharedPtr<UE::PixelStreamingServers::IServer> Server;
+
+	// Download process for PS web frontend files (if we want to view output in the browser)
+	TSharedPtr<FMonitoredProcess> DownloadProcess;
 };
