Fix failing HandshakePlayerOffer and HandshakeStreamerOffer test on Linux.

diff --git a/Engine/Plugins/Media/PixelStreaming/Source/PixelStreaming/Private/Tests/HandshakeTest.cpp b/Engine/Plugins/Media/PixelStreaming/Source/PixelStreaming/Private/Tests/HandshakeTest.cpp
--- a/Engine/Plugins/Media/PixelStreaming/Source/PixelStreaming/Private/Tests/HandshakeTest.cpp
+++ b/Engine/Plugins/Media/PixelStreaming/Source/PixelStreaming/Private/Tests/HandshakeTest.cpp
@@ -285,8 +285,9 @@ namespace UE::PixelStreaming
 	IMPLEMENT_SIMPLE_AUTOMATION_TEST(FHandshakeTestStreamerOffer, "System.Plugins.PixelStreaming.HandshakeStreamerOffer", EAutomationTestFlags::EditorContext | EAutomationTestFlags::ClientContext | EAutomationTestFlags::ProductFilter)
 	bool FHandshakeTestStreamerOffer::RunTest(const FString& Parameters)
 	{
+		// Note: Ports must be > 1024 to bind on Linux due to UE running non-root
 		int32 StreamerPort = 8866;
-		int32 PlayerPort = 86;
+		int32 PlayerPort = 1886;
 		TSharedPtr<UE::PixelStreamingServers::IServer> OutSignallingServer = CreateSignallingServer(StreamerPort, PlayerPort);
 		TSharedPtr<FStreamer> OutStreamer = CreateStreamer(StreamerPort);
 		TSharedPtr<FMockPlayer> OutPlayer = CreatePlayer(FMockPlayer::EMode::AcceptOffers, PlayerPort);
@@ -306,8 +307,9 @@ namespace UE::PixelStreaming
 	IMPLEMENT_SIMPLE_AUTOMATION_TEST(FHandshakeTestPlayerOffer, "System.Plugins.PixelStreaming.HandshakePlayerOffer", EAutomationTestFlags::EditorContext | EAutomationTestFlags::ClientContext | EAutomationTestFlags::ProductFilter)
 	bool FHandshakeTestPlayerOffer::RunTest(const FString& Parameters)
 	{
+		// Note: Ports must be > 1024 to bind on Linux due to UE running non-root
 		int32 StreamerPort = 6446;
-		int32 PlayerPort = 68;
+		int32 PlayerPort = 1668;
 		TSharedPtr<UE::PixelStreamingServers::IServer> OutSignallingServer = CreateSignallingServer(StreamerPort, PlayerPort);
 		TSharedPtr<FStreamer> OutStreamer = CreateStreamer(StreamerPort);
 		TSharedPtr<FMockPlayer> OutPlayer = CreatePlayer(FMockPlayer::EMode::CreateOffers, PlayerPort);
