Fix linker collisions on `CommandLineParseValue` in Pixel Streaming modules by giving these functions more verbose names

diff --git a/Engine/Plugins/Media/PixelStreaming/Source/PixelStreaming/Private/Settings.cpp b/Engine/Plugins/Media/PixelStreaming/Source/PixelStreaming/Private/Settings.cpp
--- a/Engine/Plugins/Media/PixelStreaming/Source/PixelStreaming/Private/Settings.cpp
+++ b/Engine/Plugins/Media/PixelStreaming/Source/PixelStreaming/Private/Settings.cpp
@@ -11,21 +11,21 @@
 #include "PixelStreamingDelegates.h"
 
 template <typename T>
-void CommandLineParseValue(const TCHAR* Match, TAutoConsoleVariable<T>& CVar)
+void PSCommandLineParseValue(const TCHAR* Match, TAutoConsoleVariable<T>& CVar)
 {
 	T Value;
 	if (FParse::Value(FCommandLine::Get(), Match, Value))
 		CVar->Set(Value, ECVF_SetByCommandline);
 };
 
-void CommandLineParseValue(const TCHAR* Match, TAutoConsoleVariable<FString>& CVar, bool bStopOnSeparator = false)
+void PSCommandLineParseValue(const TCHAR* Match, TAutoConsoleVariable<FString>& CVar, bool bStopOnSeparator = false)
 {
 	FString Value;
 	if (FParse::Value(FCommandLine::Get(), Match, Value, bStopOnSeparator))
 		CVar->Set(*Value, ECVF_SetByCommandline);
 };
 
-void CommandLineParseOption(const TCHAR* Match, TAutoConsoleVariable<bool>& CVar)
+void PSCommandLineParseOption(const TCHAR* Match, TAutoConsoleVariable<bool>& CVar)
 {
 	FString ValueMatch(Match);
 	ValueMatch.Append(TEXT("="));
@@ -588,7 +588,7 @@ namespace UE::PixelStreaming::Settings
 	void OnStreamerReady(IPixelStreamingModule& Module)
 	{
 		CVarPixelStreamingLogStats.AsVariable()->SetOnChangedCallback(FConsoleVariableDelegate::CreateStatic(&OnLogStatsChanged));
-		CommandLineParseOption(TEXT("PixelStreamingLogStats"), CVarPixelStreamingLogStats);
+		PSCommandLineParseOption(TEXT("PixelStreamingLogStats"), CVarPixelStreamingLogStats);
 	}
 
 	void InitialiseSettings()
@@ -609,44 +609,44 @@ namespace UE::PixelStreaming::Settings
 		CVarPixelStreamingEncoderKeyframeInterval.AsVariable()->SetOnChangedCallback(FConsoleVariableDelegate::CreateStatic(&OnKeyframeIntervalChanged));
 
 		// Values parse from commands line
-		CommandLineParseValue(TEXT("PixelStreamingEncoderKeyframeInterval="), CVarPixelStreamingEncoderKeyframeInterval);
-		CommandLineParseValue(TEXT("PixelStreamingEncoderTargetBitrate="), CVarPixelStreamingEncoderTargetBitrate);
-		CommandLineParseValue(TEXT("PixelStreamingEncoderMaxBitrate="), CVarPixelStreamingEncoderMaxBitrate);
-		CommandLineParseValue(TEXT("PixelStreamingEncoderMinQP="), CVarPixelStreamingEncoderMinQP);
-		CommandLineParseValue(TEXT("PixelStreamingEncoderMaxQP="), CVarPixelStreamingEncoderMaxQP);
-		CommandLineParseValue(TEXT("PixelStreamingEncoderRateControl="), CVarPixelStreamingEncoderRateControl);
-		CommandLineParseValue(TEXT("PixelStreamingEncoderMultipass="), CVarPixelStreamingEncoderMultipass);
-		CommandLineParseValue(TEXT("PixelStreamingEncoderCodec="), CVarPixelStreamingEncoderCodec);
-		CommandLineParseValue(TEXT("PixelStreamingH264Profile="), CVarPixelStreamingH264Profile);
-		CommandLineParseValue(TEXT("PixelStreamingDegradationPreference="), CVarPixelStreamingDegradationPreference);
-		CommandLineParseValue(TEXT("PixelStreamingWebRTCDegradationPreference="), CVarPixelStreamingDegradationPreference);
-		CommandLineParseValue(TEXT("PixelStreamingWebRTCFps="), CVarPixelStreamingWebRTCFps);
-		CommandLineParseValue(TEXT("PixelStreamingWebRTCStartBitrate="), CVarPixelStreamingWebRTCStartBitrate);
-		CommandLineParseValue(TEXT("PixelStreamingWebRTCMinBitrate="), CVarPixelStreamingWebRTCMinBitrate);
-		CommandLineParseValue(TEXT("PixelStreamingWebRTCMaxBitrate="), CVarPixelStreamingWebRTCMaxBitrate);
-		CommandLineParseValue(TEXT("PixelStreamingWebRTCLowQpThreshold="), CVarPixelStreamingWebRTCLowQpThreshold);
-		CommandLineParseValue(TEXT("PixelStreamingWebRTCHighQpThreshold="), CVarPixelStreamingWebRTCHighQpThreshold);
-		CommandLineParseValue(TEXT("PixelStreamingFreezeFrameQuality"), CVarPixelStreamingFreezeFrameQuality);
-		CommandLineParseValue(TEXT("PixelStreamingKeyFilter="), CVarPixelStreamingKeyFilter);
-		CommandLineParseValue(TEXT("PixelStreamingVideoTracks="), CVarPixelStreamingVideoTracks);
-		CommandLineParseValue(TEXT("PixelStreamingFrameScale="), CVarPixelStreamingFrameScale);
-		CommandLineParseValue(TEXT("PixelStreamingInputController="), CVarPixelStreamingInputController);
-		CommandLineParseValue(TEXT("PixelStreamingSignalingReconnectInterval="), CVarPixelStreamingSignalingReconnectInterval);
+		PSCommandLineParseValue(TEXT("PixelStreamingEncoderKeyframeInterval="), CVarPixelStreamingEncoderKeyframeInterval);
+		PSCommandLineParseValue(TEXT("PixelStreamingEncoderTargetBitrate="), CVarPixelStreamingEncoderTargetBitrate);
+		PSCommandLineParseValue(TEXT("PixelStreamingEncoderMaxBitrate="), CVarPixelStreamingEncoderMaxBitrate);
+		PSCommandLineParseValue(TEXT("PixelStreamingEncoderMinQP="), CVarPixelStreamingEncoderMinQP);
+		PSCommandLineParseValue(TEXT("PixelStreamingEncoderMaxQP="), CVarPixelStreamingEncoderMaxQP);
+		PSCommandLineParseValue(TEXT("PixelStreamingEncoderRateControl="), CVarPixelStreamingEncoderRateControl);
+		PSCommandLineParseValue(TEXT("PixelStreamingEncoderMultipass="), CVarPixelStreamingEncoderMultipass);
+		PSCommandLineParseValue(TEXT("PixelStreamingEncoderCodec="), CVarPixelStreamingEncoderCodec);
+		PSCommandLineParseValue(TEXT("PixelStreamingH264Profile="), CVarPixelStreamingH264Profile);
+		PSCommandLineParseValue(TEXT("PixelStreamingDegradationPreference="), CVarPixelStreamingDegradationPreference);
+		PSCommandLineParseValue(TEXT("PixelStreamingWebRTCDegradationPreference="), CVarPixelStreamingDegradationPreference);
+		PSCommandLineParseValue(TEXT("PixelStreamingWebRTCFps="), CVarPixelStreamingWebRTCFps);
+		PSCommandLineParseValue(TEXT("PixelStreamingWebRTCStartBitrate="), CVarPixelStreamingWebRTCStartBitrate);
+		PSCommandLineParseValue(TEXT("PixelStreamingWebRTCMinBitrate="), CVarPixelStreamingWebRTCMinBitrate);
+		PSCommandLineParseValue(TEXT("PixelStreamingWebRTCMaxBitrate="), CVarPixelStreamingWebRTCMaxBitrate);
+		PSCommandLineParseValue(TEXT("PixelStreamingWebRTCLowQpThreshold="), CVarPixelStreamingWebRTCLowQpThreshold);
+		PSCommandLineParseValue(TEXT("PixelStreamingWebRTCHighQpThreshold="), CVarPixelStreamingWebRTCHighQpThreshold);
+		PSCommandLineParseValue(TEXT("PixelStreamingFreezeFrameQuality"), CVarPixelStreamingFreezeFrameQuality);
+		PSCommandLineParseValue(TEXT("PixelStreamingKeyFilter="), CVarPixelStreamingKeyFilter);
+		PSCommandLineParseValue(TEXT("PixelStreamingVideoTracks="), CVarPixelStreamingVideoTracks);
+		PSCommandLineParseValue(TEXT("PixelStreamingFrameScale="), CVarPixelStreamingFrameScale);
+		PSCommandLineParseValue(TEXT("PixelStreamingInputController="), CVarPixelStreamingInputController);
+		PSCommandLineParseValue(TEXT("PixelStreamingSignalingReconnectInterval="), CVarPixelStreamingSignalingReconnectInterval);
 
 		// Options parse (if these exist they are set to true)
-		CommandLineParseOption(TEXT("AllowPixelStreamingCommands"), CVarPixelStreamingAllowConsoleCommands);
-		CommandLineParseOption(TEXT("PixelStreamingOnScreenStats"), CVarPixelStreamingOnScreenStats);
-		CommandLineParseOption(TEXT("PixelStreamingHudStats"), CVarPixelStreamingOnScreenStats);
-		CommandLineParseOption(TEXT("PixelStreamingDebugDumpFrame"), CVarPixelStreamingDebugDumpFrame);
-		CommandLineParseOption(TEXT("PixelStreamingEnableFillerData"), CVarPixelStreamingEnableFillerData);
-		CommandLineParseOption(TEXT("PixelStreamingWebRTCDisableStats"), CVarPixelStreamingWebRTCDisableStats);
-		CommandLineParseOption(TEXT("PixelStreamingWebRTCDisableReceiveAudio"), CVarPixelStreamingWebRTCDisableReceiveAudio);
-		CommandLineParseOption(TEXT("PixelStreamingWebRTCDisableTransmitAudio"), CVarPixelStreamingWebRTCDisableTransmitAudio);
-		CommandLineParseOption(TEXT("PixelStreamingWebRTCDisableAudioSync"), CVarPixelStreamingWebRTCDisableAudioSync);
-		CommandLineParseOption(TEXT("PixelStreamingSendPlayerIdAsInteger"), CVarSendPlayerIdAsInteger);
-		CommandLineParseOption(TEXT("PixelStreamingWebRTCUseLegacyAudioDevice"), CVarPixelStreamingWebRTCUseLegacyAudioDevice);
-		CommandLineParseOption(TEXT("PixelStreamingDisableLatencyTester"), CVarPixelStreamingDisableLatencyTester);
-		CommandLineParseOption(TEXT("PixelStreamingVPXUseCompute"), CVarPixelStreamingVPXUseCompute);
+		PSCommandLineParseOption(TEXT("AllowPixelStreamingCommands"), CVarPixelStreamingAllowConsoleCommands);
+		PSCommandLineParseOption(TEXT("PixelStreamingOnScreenStats"), CVarPixelStreamingOnScreenStats);
+		PSCommandLineParseOption(TEXT("PixelStreamingHudStats"), CVarPixelStreamingOnScreenStats);
+		PSCommandLineParseOption(TEXT("PixelStreamingDebugDumpFrame"), CVarPixelStreamingDebugDumpFrame);
+		PSCommandLineParseOption(TEXT("PixelStreamingEnableFillerData"), CVarPixelStreamingEnableFillerData);
+		PSCommandLineParseOption(TEXT("PixelStreamingWebRTCDisableStats"), CVarPixelStreamingWebRTCDisableStats);
+		PSCommandLineParseOption(TEXT("PixelStreamingWebRTCDisableReceiveAudio"), CVarPixelStreamingWebRTCDisableReceiveAudio);
+		PSCommandLineParseOption(TEXT("PixelStreamingWebRTCDisableTransmitAudio"), CVarPixelStreamingWebRTCDisableTransmitAudio);
+		PSCommandLineParseOption(TEXT("PixelStreamingWebRTCDisableAudioSync"), CVarPixelStreamingWebRTCDisableAudioSync);
+		PSCommandLineParseOption(TEXT("PixelStreamingSendPlayerIdAsInteger"), CVarSendPlayerIdAsInteger);
+		PSCommandLineParseOption(TEXT("PixelStreamingWebRTCUseLegacyAudioDevice"), CVarPixelStreamingWebRTCUseLegacyAudioDevice);
+		PSCommandLineParseOption(TEXT("PixelStreamingDisableLatencyTester"), CVarPixelStreamingDisableLatencyTester);
+		PSCommandLineParseOption(TEXT("PixelStreamingVPXUseCompute"), CVarPixelStreamingVPXUseCompute);
 
 		ReadSimulcastParameters();
 
diff --git a/Engine/Plugins/Media/PixelStreaming/Source/PixelStreamingEditor/Private/Settings.cpp b/Engine/Plugins/Media/PixelStreaming/Source/PixelStreamingEditor/Private/Settings.cpp
--- a/Engine/Plugins/Media/PixelStreaming/Source/PixelStreamingEditor/Private/Settings.cpp
+++ b/Engine/Plugins/Media/PixelStreaming/Source/PixelStreamingEditor/Private/Settings.cpp
@@ -4,21 +4,21 @@
 #include "Utils.h"
 
 template <typename T>
-void CommandLineParseValue(const TCHAR* Match, TAutoConsoleVariable<T>& CVar)
+void PSEditorCommandLineParseValue(const TCHAR* Match, TAutoConsoleVariable<T>& CVar)
 {
 	T Value;
 	if (FParse::Value(FCommandLine::Get(), Match, Value))
 		CVar->Set(Value, ECVF_SetByCommandline);
 };
 
-void CommandLineParseValue(const TCHAR* Match, TAutoConsoleVariable<FString>& CVar, bool bStopOnSeparator = false)
+void PSEditorCommandLineParseValue(const TCHAR* Match, TAutoConsoleVariable<FString>& CVar, bool bStopOnSeparator = false)
 {
 	FString Value;
 	if (FParse::Value(FCommandLine::Get(), Match, Value, bStopOnSeparator))
 		CVar->Set(*Value, ECVF_SetByCommandline);
 };
 
-void CommandLineParseOption(const TCHAR* Match, TAutoConsoleVariable<bool>& CVar)
+void PSEditorCommandLineParseOption(const TCHAR* Match, TAutoConsoleVariable<bool>& CVar)
 {
 	FString ValueMatch(Match);
 	ValueMatch.Append(TEXT("="));
@@ -51,6 +51,6 @@ namespace UE::PixelStreaming::Settings::Editor
     void InitialiseSettings()
     {
         // Options parse (if these exist they are set to true)
-		CommandLineParseOption(TEXT("EditorPixelStreamingStartOnLaunch"), CVarEditorPixelStreamingStartOnLaunch);
+		PSEditorCommandLineParseOption(TEXT("EditorPixelStreamingStartOnLaunch"), CVarEditorPixelStreamingStartOnLaunch);
     }
 }
\ No newline at end of file
