Cherrypicked 21502862, 21502930, 21556336 from Dev-TensorWorks:

Fix: Allow Editor dimensions to be set at launch time with ("-EditorPixelStreamingResX= and -EditorPixelStreamingResY=") OR "-EditorPixelStreamingRes=".
The Editor can now resize appropriately if requested by the front end.

Fixing some stats issues
- The capture time of frames was using a different clock to everything else so times were getting skewed
- Removed captureToSend since it was calculating off a bad stat (total sum of delay per packet which is not helpful)
- Adding a new stat in its place that is only available with a (pending) webrtc patch to expose avg packet delay.

diff --git a/Engine/Plugins/Media/PixelStreaming/Source/PixelStreaming/Private/InputHandlers.cpp b/Engine/Plugins/Media/PixelStreaming/Source/PixelStreaming/Private/InputHandlers.cpp
new file mode 100644
--- /dev/null
+++ b/Engine/Plugins/Media/PixelStreaming/Source/PixelStreaming/Private/InputHandlers.cpp
@@ -0,0 +1,73 @@
+// Copyright Epic Games, Inc. All Rights Reserved.
+
+#include "InputHandlers.h"
+#include "IPixelStreamingModule.h"
+#include "Framework/Application/SlateApplication.h"
+
+namespace UE::PixelStreaming
+{
+        FInputHandlers::FInputHandlers(const TSharedRef<FGenericApplicationMessageHandler>& InMessageHandler)
+        {
+            // This is imperative for editor streaming as when a modal is open or we've hit a BP breakpoint, the engine tick loop will not run, so instead we rely on this delegate to tick for us
+		    FSlateApplication::Get().OnPreTick().AddRaw(this, &FInputHandlers::Tick);
+        }
+    
+		void FInputHandlers::Tick(float DeltaTime)
+        {
+            ForEachInputHandler([DeltaTime](IInputDevice* InputHandler) {
+                InputHandler->Tick(DeltaTime);
+            });
+        }
+
+		void FInputHandlers::SendControllerEvents()
+        {
+            ForEachInputHandler([](IInputDevice* InputHandler) {
+                InputHandler->SendControllerEvents();
+            });
+        }
+
+		void FInputHandlers::SetMessageHandler(const TSharedRef<FGenericApplicationMessageHandler>& InMessageHandler)
+        {
+            ForEachInputHandler([&MessageHandler = InMessageHandler](IInputDevice* InputHandler) {
+                InputHandler->SetMessageHandler(MessageHandler);
+            });
+        }
+
+		bool FInputHandlers::Exec(UWorld* InWorld, const TCHAR* Cmd, FOutputDevice& Ar) 
+        {
+            ForEachInputHandler([InWorld, Cmd, &Ar](IInputDevice* InputHandler) {
+                InputHandler->Exec(InWorld, Cmd, Ar);
+            });
+            return true;
+        }
+
+		void FInputHandlers::SetChannelValue(int32 ControllerId, FForceFeedbackChannelType ChannelType, float Value) 
+        {
+            ForEachInputHandler([&ControllerId, &ChannelType, &Value](IInputDevice* InputHandler) {
+                InputHandler->SetChannelValue(ControllerId, ChannelType, Value);
+            });
+        }
+
+		void FInputHandlers::SetChannelValues(int32 ControllerId, const FForceFeedbackValues& Values)
+        {
+            ForEachInputHandler([&ControllerId, &Values](IInputDevice* InputHandler) {
+                InputHandler->SetChannelValues(ControllerId, Values);
+            });
+        }
+
+		void FInputHandlers::ForEachInputHandler(TFunction<void(IInputDevice*)> const& Visitor)
+        {
+            IPixelStreamingModule& Module = IPixelStreamingModule::Get();
+            Module.ForEachStreamer([Visitor](TSharedPtr<IPixelStreamingStreamer> Streamer) {
+                if(Streamer)
+                {
+                    TWeakPtr<IPixelStreamingInputHandler> WeakInputHandler = Streamer->GetInputHandler();
+                    TSharedPtr<IPixelStreamingInputHandler> InputHandler = WeakInputHandler.Pin();
+                    if (InputHandler)
+                    {
+                        Visitor(InputHandler.Get());
+                    }
+                }
+            });
+        }
+} // namespace UE::PixelStreaming
diff --git a/Engine/Plugins/Media/PixelStreaming/Source/PixelStreaming/Private/InputHandlers.h b/Engine/Plugins/Media/PixelStreaming/Source/PixelStreaming/Private/InputHandlers.h
new file mode 100644
--- /dev/null
+++ b/Engine/Plugins/Media/PixelStreaming/Source/PixelStreaming/Private/InputHandlers.h
@@ -0,0 +1,29 @@
+// Copyright Epic Games, Inc. All Rights Reserved.
+
+#pragma once
+
+#include "IInputDevice.h"
+#include "GenericPlatform/GenericApplicationMessageHandler.h"
+
+namespace UE::PixelStreaming
+{
+	// The IInputDevice module only creates a single IInputDevice but we have one per streamer.
+	// As such, this wrapper exists to act as the bridge between the single IInputDevice created 
+	// by the IInputDeviceModule and the multiple streamers
+	class FInputHandlers : public IInputDevice
+	{
+	public:
+		// This wrapper does effectively nothing, and it is up to the individual streamers to manage their input handlers
+		FInputHandlers(const TSharedRef<FGenericApplicationMessageHandler>& InMessageHandler);
+		virtual ~FInputHandlers() = default;
+
+		virtual void Tick(float DeltaTime) override;
+		virtual void SendControllerEvents() override;
+		virtual void SetMessageHandler(const TSharedRef<FGenericApplicationMessageHandler>& InMessageHandler) override;
+		virtual bool Exec(UWorld* InWorld, const TCHAR* Cmd, FOutputDevice& Ar) override;
+		virtual void SetChannelValue(int32 ControllerId, FForceFeedbackChannelType ChannelType, float Value) override;
+		virtual void SetChannelValues(int32 ControllerId, const FForceFeedbackValues& Values) override;
+
+		void ForEachInputHandler(TFunction<void(IInputDevice*)> const& Visitor);
+	};
+} // namespace UE::PixelStreaming
diff --git a/Engine/Plugins/Media/PixelStreaming/Source/PixelStreaming/Private/PixelStreamingInputChannel.cpp b/Engine/Plugins/Media/PixelStreaming/Source/PixelStreaming/Private/PixelStreamingInputChannel.cpp
deleted file mode 100644
diff --git a/Engine/Plugins/Media/PixelStreaming/Source/PixelStreaming/Private/PixelStreamingInputChannel.h b/Engine/Plugins/Media/PixelStreaming/Source/PixelStreaming/Private/PixelStreamingInputChannel.h
deleted file mode 100644
diff --git a/Engine/Plugins/Media/PixelStreaming/Source/PixelStreaming/Private/PixelStreamingInputHandler.cpp b/Engine/Plugins/Media/PixelStreaming/Source/PixelStreaming/Private/PixelStreamingInputHandler.cpp
new file mode 100644
--- /dev/null
+++ b/Engine/Plugins/Media/PixelStreaming/Source/PixelStreaming/Private/PixelStreamingInputHandler.cpp
@@ -0,0 +1,826 @@
+// Copyright Epic Games, Inc. All Rights Reserved.
+
+#include "PixelStreamingInputHandler.h"
+#include "InputStructures.h"
+#include "PixelStreamingProtocol.h"
+#include "PixelStreamingModule.h"
+#include "JavaScriptKeyCodes.inl"
+#include "Settings.h"
+#include "Layout/ArrangedChildren.h"
+#include "Layout/WidgetPath.h"
+#include "Framework/Application/SlateApplication.h"
+#include "Widgets/SViewport.h"
+#include "Policies/CondensedJsonPrintPolicy.h"
+#include "Serialization/JsonSerializer.h"
+#include "Utils.h"
+#include "PixelStreamingInputComponent.h"
+#include "Engine/Engine.h"
+#include "Framework/Application/SlateUser.h"
+#include "Widgets/Input/SEditableTextBox.h"
+#include "Misc/CoreMiscDefines.h"
+
+DECLARE_LOG_CATEGORY_EXTERN(LogPixelStreamingInputHandler, Log, VeryVerbose);
+DEFINE_LOG_CATEGORY(LogPixelStreamingInputHandler);
+
+// TODO: Gesture recognition is moving to the browser, so add handlers for the gesture events.
+// The gestures supported will be swipe, pinch,
+
+namespace UE::PixelStreaming
+{
+    FPixelStreamingInputHandler::FPixelStreamingInputHandler(TSharedPtr<FPixelStreamingApplicationWrapper> InApplicationWrapper, const TSharedPtr<FGenericApplicationMessageHandler>& InTargetHandler)
+    : TargetViewport(nullptr)
+    , NumActiveTouches(0)
+    , bIsMouseActive(false)
+    , MessageHandler(InTargetHandler)
+    , PixelStreamingModule(FPixelStreamingModule::GetModule())
+    , PixelStreamerApplicationWrapper(InApplicationWrapper)
+    , FocusedPos(FVector2D(-1.0f, -1.0f))
+    , UnfocusedPos(FVector2D(-1.0f, -1.0f))
+    {
+        RegisterMessageHandler("KeyPress", [this](FMemoryReader Ar) { HandleOnKeyChar(Ar); });
+        RegisterMessageHandler("KeyUp", [this](FMemoryReader Ar) { HandleOnKeyUp(Ar); });
+        RegisterMessageHandler("KeyDown", [this](FMemoryReader Ar) { HandleOnKeyDown(Ar); });
+
+        RegisterMessageHandler("TouchStart", [this](FMemoryReader Ar) { HandleOnTouchStarted(Ar); });
+        RegisterMessageHandler("TouchMove", [this](FMemoryReader Ar) { HandleOnTouchMoved(Ar); });
+        RegisterMessageHandler("TouchEnd", [this](FMemoryReader Ar) { HandleOnTouchEnded(Ar); });
+
+        RegisterMessageHandler("GamepadAnalog", [this](FMemoryReader Ar) { HandleOnControllerAnalog(Ar); });
+        RegisterMessageHandler("GamepadButtonPressed", [this](FMemoryReader Ar) { HandleOnControllerButtonPressed(Ar); });
+        RegisterMessageHandler("GamepadButtonReleased", [this](FMemoryReader Ar) { HandleOnControllerButtonReleased(Ar); });
+
+        RegisterMessageHandler("MouseEnter", [this](FMemoryReader Ar) { HandleOnMouseEnter(Ar); });
+        RegisterMessageHandler("MouseLeave", [this](FMemoryReader Ar) { HandleOnMouseLeave(Ar); });
+        RegisterMessageHandler("MouseUp", [this](FMemoryReader Ar) { HandleOnMouseUp(Ar); });
+        RegisterMessageHandler("MouseDown", [this](FMemoryReader Ar) { HandleOnMouseDown(Ar); });
+        RegisterMessageHandler("MouseMove", [this](FMemoryReader Ar) { HandleOnMouseMove(Ar); });
+        RegisterMessageHandler("MouseWheel", [this](FMemoryReader Ar) { HandleOnMouseWheel(Ar); });
+		RegisterMessageHandler("MouseDouble", [this](FMemoryReader Ar) { HandleOnMouseDoubleClick(Ar); });
+
+        RegisterMessageHandler("Command", [this](FMemoryReader Ar) { HandleCommand(Ar); });
+        RegisterMessageHandler("UIInteraction", [this](FMemoryReader Ar) { HandleUIInteraction(Ar); });
+    }
+
+    FPixelStreamingInputHandler::~FPixelStreamingInputHandler()
+    {
+    }
+
+    void FPixelStreamingInputHandler::RegisterMessageHandler(const FString& MessageType, const TFunction<void(FMemoryReader)>& Handler)
+    {
+        TMap<FString, Protocol::FPixelStreamingInputMessage> Protocol = PixelStreamingModule->GetProtocol().ToStreamerProtocol;
+        DispatchTable.Add(Protocol.Find(MessageType)->Id, Handler);
+    }
+
+	TFunction<void(FMemoryReader)> FPixelStreamingInputHandler::FindMessageHandler(const FString& MessageType)
+	{
+		TMap<FString, Protocol::FPixelStreamingInputMessage> Protocol = PixelStreamingModule->GetProtocol().ToStreamerProtocol;
+		return DispatchTable.FindRef(Protocol.Find(MessageType)->Id);
+	}
+
+    void FPixelStreamingInputHandler::Tick(const float InDeltaTime)
+    {
+        FMessage Message;
+		while (Messages.Dequeue(Message))
+		{
+            FMemoryReader Ar(Message.Data);
+            (*Message.Handler)(Ar);
+        }
+    }
+
+	void FPixelStreamingInputHandler::OnMessage(const webrtc::DataBuffer& Buffer)
+	{
+		using namespace Protocol;
+
+		const uint8* Data = Buffer.data.data();
+		uint32 Size = (uint32)Buffer.data.size();
+        if(sizeof(uint8) > Size)
+        {
+            UE_LOG(LogPixelStreamingInputHandler, Warning, TEXT("Buffer size is too small to extract message type. Buffer size (bytes): %d"), Size);
+            return;
+        }
+		const uint8 MsgType = *Data;
+		Data += sizeof(uint8);
+		Size -= sizeof(uint8);
+
+        TFunction<void(FMemoryReader)>* Handler = DispatchTable.Find(MsgType);
+        if (Handler != nullptr)
+        {
+            TArray<uint8> MessageData(Data, Size);
+            FMessage Message = {
+                Handler,   // The function to call
+                MessageData // The message data
+            };
+            Messages.Enqueue(Message);
+        }
+        else
+        {
+            UE_LOG(LogPixelStreamingInputHandler, Warning, TEXT("No handler registered for message with id %d"), MsgType);
+        }
+    }
+
+	void FPixelStreamingInputHandler::SetTargetWindow(TWeakPtr<SWindow> InWindow)
+	{
+		TargetWindow = InWindow;
+	}
+
+	TWeakPtr<SWindow> FPixelStreamingInputHandler::GetTargetWindow()
+	{
+		return TargetWindow;
+	}
+
+	void FPixelStreamingInputHandler::SetTargetScreenSize(TWeakPtr<FIntPoint> InScreenSize)
+	{
+		TargetScreenSize = InScreenSize;
+	}
+
+	TWeakPtr<FIntPoint> FPixelStreamingInputHandler::GetTargetScreenSize()
+	{
+		return TargetScreenSize;
+	}
+
+	void FPixelStreamingInputHandler::SetTargetViewport(TWeakPtr<SViewport> InViewport)
+	{
+		TargetViewport = InViewport;
+	}
+
+	TWeakPtr<SViewport> FPixelStreamingInputHandler::GetTargetViewport()
+	{
+		return TargetViewport;
+	}
+
+	void FPixelStreamingInputHandler::SetMessageHandler(const TSharedRef<FGenericApplicationMessageHandler>& InTargetHandler)
+	{
+		MessageHandler = InTargetHandler;
+	}
+
+	bool FPixelStreamingInputHandler::Exec(UWorld* InWorld, const TCHAR* Cmd, FOutputDevice& Ar)
+	{
+		return GEngine->Exec(InWorld, Cmd, Ar);
+	}
+
+	void FPixelStreamingInputHandler::SetChannelValue(int32 ControllerId, FForceFeedbackChannelType ChannelType, float Value)
+	{
+		// TODO: Implement FFB
+	}
+
+	void FPixelStreamingInputHandler::SetChannelValues(int32 ControllerId, const FForceFeedbackValues& values)
+	{
+		// TODO: Implement FFB
+	}
+
+    void FPixelStreamingInputHandler::HandleOnKeyChar(FMemoryReader Ar)
+    {
+        TPayloadOneParam<TCHAR> Payload(Ar);
+        UE_LOG(LogPixelStreamingInputHandler, Verbose, TEXT("KEY_PRESSED: Character = '%c'"), Payload.Param1);
+        // A key char event is never repeated, so set it to false. It's value  
+        // ultimately doesn't matter as this paramater isn't used later
+        MessageHandler->OnKeyChar(Payload.Param1, false);
+    }
+
+    void FPixelStreamingInputHandler::HandleOnKeyDown(FMemoryReader Ar)
+    {
+        TPayloadTwoParam<uint8, uint8> Payload(Ar);
+
+		bool bIsRepeat = Payload.Param2 != 0;
+		const FKey* AgnosticKey = JavaScriptKeyCodeToFKey[Payload.Param1];
+		if (FilterKey(*AgnosticKey))
+		{
+			const uint32* KeyPtr;
+			const uint32* CharacterPtr;
+			FInputKeyManager::Get().GetCodesFromKey(*AgnosticKey, KeyPtr, CharacterPtr);
+			uint32 Key = KeyPtr ? *KeyPtr : 0;
+			uint32 Character = CharacterPtr ? *CharacterPtr : 0;
+
+			UE_LOG(LogPixelStreamingInputHandler, Verbose, TEXT("KEY_DOWN: Key = %d; Character = %d; IsRepeat = %s"), Key, Character, bIsRepeat ? TEXT("True") : TEXT("False"));
+			MessageHandler->OnKeyDown((int32)Key, (int32)Character, bIsRepeat);
+		}
+	}
+
+    void FPixelStreamingInputHandler::HandleOnKeyUp(FMemoryReader Ar)
+    {
+        TPayloadOneParam<uint8> Payload(Ar);
+        const FKey* AgnosticKey = JavaScriptKeyCodeToFKey[Payload.Param1];
+        if (FilterKey(*AgnosticKey))
+		{
+			const uint32* KeyPtr;
+			const uint32* CharacterPtr;
+			FInputKeyManager::Get().GetCodesFromKey(*AgnosticKey, KeyPtr, CharacterPtr);
+			uint32 Key = KeyPtr ? *KeyPtr : 0;
+			uint32 Character = CharacterPtr ? *CharacterPtr : 0;
+
+			UE_LOG(LogPixelStreamingInputHandler, Verbose, TEXT("KEY_UP: Key = %d; Character = %d"), Key, Character);
+			MessageHandler->OnKeyUp((int32)Key, (int32)Character, false);
+		}
+	}
+
+    void FPixelStreamingInputHandler::HandleOnTouchStarted(FMemoryReader Ar)
+    {
+        TPayloadOneParam<uint8> Payload(Ar);
+
+		uint8 NumTouches = Payload.Param1;
+		for (uint8 TouchIdx = 0; TouchIdx < NumTouches; TouchIdx++)
+		{
+			//                 PosX    PoxY    IDX   Force  Valid
+			TPayloadFiveParam<uint16, uint16, uint8, uint8, uint8> Touch(Ar);
+			// If Touch is valid
+			if (Touch.Param5 != 0)
+			{
+				if (NumActiveTouches == 0 && !bIsMouseActive)
+				{
+					FSlateApplication::Get().OnCursorSet();
+					// Make sure the application is active.
+					FSlateApplication::Get().ProcessApplicationActivationEvent(true);
+
+					FVector2D OldCursorLocation = PixelStreamerApplicationWrapper->WrappedApplication->Cursor->GetPosition();
+					PixelStreamerApplicationWrapper->Cursor->SetPosition(OldCursorLocation.X, OldCursorLocation.Y);
+					FSlateApplication::Get().OverridePlatformApplication(PixelStreamerApplicationWrapper);
+				}
+
+				//                                                                           convert range from 0,65536 -> 0,1
+				FVector2D TouchLocation = ConvertFromNormalizedScreenLocation(FVector2D(Touch.Param1 / uint16_MAX, Touch.Param2 / uint16_MAX));
+
+				// We must update the user cursor position explicitly before updating the application cursor position
+				// as if there's a delta between them, when the touch event is started it will trigger a move
+				// resulting in a large 'drag' across the screen
+				TSharedPtr<FSlateUser> User = FSlateApplication::Get().GetCursorUser();
+				User->SetCursorPosition(TouchLocation);
+				PixelStreamerApplicationWrapper->Cursor->SetPosition(TouchLocation.X, TouchLocation.Y);
+				PixelStreamerApplicationWrapper->WrappedApplication->Cursor->SetPosition(TouchLocation.X, TouchLocation.Y);
+
+				UE_LOG(LogPixelStreamingInputHandler, Verbose, TEXT("TOUCH_START: TouchIndex = %d; Pos = (%d, %d); CursorPos = (%d, %d); Force = %.3f"), Touch.Param3, Touch.Param1, Touch.Param2, static_cast<int>(TouchLocation.X), static_cast<int>(TouchLocation.Y), Touch.Param4 / 255.0f);
+				MessageHandler->OnTouchStarted(PixelStreamerApplicationWrapper->GetWindowUnderCursor(), TouchLocation, Touch.Param4 / 255.0f, Touch.Param3, 0); // TODO: ControllerId?
+
+				NumActiveTouches++;
+			}
+		}
+		
+		FindFocusedWidget();
+    }
+
+    void FPixelStreamingInputHandler::HandleOnTouchMoved(FMemoryReader Ar)
+    {
+        TPayloadOneParam<uint8> Payload(Ar);
+
+		uint8 NumTouches = Payload.Param1;
+		for (uint8 TouchIdx = 0; TouchIdx < NumTouches; TouchIdx++)
+		{
+			//                 PosX    PoxY    IDX   Force  Valid
+			TPayloadFiveParam<uint16, uint16, uint8, uint8, uint8> Touch(Ar);
+			// If Touch is valid
+			if (Touch.Param5 != 0)
+			{
+				//                                                                           convert range from 0,65536 -> 0,1
+				FVector2D TouchLocation = ConvertFromNormalizedScreenLocation(FVector2D(Touch.Param1 / uint16_MAX, Touch.Param2 / uint16_MAX));
+                UE_LOG(LogPixelStreamingInputHandler, Verbose, TEXT("TOUCH_MOVE: TouchIndex = %d; Pos = (%d, %d); CursorPos = (%d, %d); Force = %.3f"), Touch.Param3, Touch.Param1, Touch.Param2, static_cast<int>(TouchLocation.X), static_cast<int>(TouchLocation.Y),  Touch.Param4 / 255.0f);
+                MessageHandler->OnTouchMoved(TouchLocation, Touch.Param4 / 255.0f, Touch.Param3, 0); // TODO: ControllerId?
+            }
+        }
+    }
+
+    void FPixelStreamingInputHandler::HandleOnTouchEnded(FMemoryReader Ar)
+    {
+        TPayloadOneParam<uint8> Payload(Ar);
+        uint8 NumTouches = Payload.Param1;
+        for(uint8 TouchIdx = 0; TouchIdx < NumTouches; TouchIdx++)
+        {
+            //                 PosX    PoxY    IDX   Force  Valid
+            TPayloadFiveParam<uint16, uint16, uint8, uint8, uint8> Touch(Ar);
+			// Always allowing the "up" events regardless of in or outside the valid region so
+			// states aren't stuck "down". Might want to uncomment this if it causes other issues.
+			// if(Touch.Param5 != 0)
+			{
+				//                                                                           convert range from 0,65536 -> 0,1
+				FVector2D TouchLocation = ConvertFromNormalizedScreenLocation(FVector2D(Touch.Param1 / uint16_MAX, Touch.Param2 / uint16_MAX));
+
+				UE_LOG(LogPixelStreamingInputHandler, Verbose, TEXT("TOUCH_END: TouchIndex = %d; Pos = (%d, %d); CursorPos = (%d, %d)"), Touch.Param3, Touch.Param1, Touch.Param2, static_cast<int>(TouchLocation.X), static_cast<int>(TouchLocation.Y));
+				MessageHandler->OnTouchEnded(TouchLocation, Touch.Param3, 0); // TODO: ControllerId?
+				NumActiveTouches--;
+			}
+		}
+
+		// If there's no remaing touches, and there is also no mouse over the player window
+		// then set the platform application back to its default. We need to set it back to default
+		// so that people using the editor (if editor streaming) can click on buttons outside the target window
+		// and also have the correct cursor (pixel streaming forces default cursor)
+		if (NumActiveTouches == 0 && !bIsMouseActive)
+		{
+			FVector2D OldCursorLocation = PixelStreamerApplicationWrapper->Cursor->GetPosition();
+			PixelStreamerApplicationWrapper->WrappedApplication->Cursor->SetPosition(OldCursorLocation.X, OldCursorLocation.Y);
+			FSlateApplication::Get().OverridePlatformApplication(PixelStreamerApplicationWrapper->WrappedApplication);
+		}
+	}
+
+
+	void FPixelStreamingInputHandler::HandleOnControllerAnalog(FMemoryReader Ar)
+	{
+		TPayloadThreeParam<uint8, uint8, double> Payload(Ar);
+
+		FInputDeviceId ControllerId = FInputDeviceId::CreateFromInternalId((int32)Payload.Param1);
+		FGamepadKeyNames::Type Button = ConvertAxisIndexToGamepadAxis(Payload.Param2);
+		float AnalogValue = (float)Payload.Param3;
+		FPlatformUserId UserId = IPlatformInputDeviceMapper::Get().GetPrimaryPlatformUser();
+
+		UE_LOG(LogPixelStreamingInputHandler, Verbose, TEXT("GAMEPAD_ANALOG: ControllerId = %d; KeyName = %s; AnalogValue = %.4f;"), ControllerId.GetId(), *Button.ToString(), AnalogValue);
+        MessageHandler->OnControllerAnalog(Button, UserId, ControllerId, AnalogValue);
+	}
+
+    void FPixelStreamingInputHandler::HandleOnControllerButtonPressed(FMemoryReader Ar)
+    {
+        TPayloadThreeParam<uint8, uint8, uint8> Payload(Ar);
+
+        FInputDeviceId ControllerId = FInputDeviceId::CreateFromInternalId((int32) Payload.Param1);
+        FGamepadKeyNames::Type Button = ConvertButtonIndexToGamepadButton(Payload.Param2);
+        bool bIsRepeat = Payload.Param3 != 0;
+        FPlatformUserId UserId = IPlatformInputDeviceMapper::Get().GetPrimaryPlatformUser();
+        
+        UE_LOG(LogPixelStreamingInputHandler, Verbose, TEXT("GAMEPAD_PRESSED: ControllerId = %d; KeyName = %s; IsRepeat = %s;"), ControllerId.GetId(), *Button.ToString(), bIsRepeat ? TEXT("True") : TEXT("False"));
+        MessageHandler->OnControllerButtonPressed(Button, UserId, ControllerId, bIsRepeat);
+    }
+
+    void FPixelStreamingInputHandler::HandleOnControllerButtonReleased(FMemoryReader Ar)
+    {
+        TPayloadTwoParam<uint8, uint8> Payload(Ar);
+
+        FInputDeviceId ControllerId = FInputDeviceId::CreateFromInternalId((int32) Payload.Param1);
+        FGamepadKeyNames::Type Button = ConvertButtonIndexToGamepadButton(Payload.Param2);
+        FPlatformUserId UserId = IPlatformInputDeviceMapper::Get().GetPrimaryPlatformUser();
+
+        UE_LOG(LogPixelStreamingInputHandler, Verbose, TEXT("GAMEPAD_RELEASED: ControllerId = %d; KeyName = %s;"), ControllerId.GetId(), *Button.ToString());
+        MessageHandler->OnControllerButtonReleased(Button, UserId, ControllerId, false);
+    }
+
+    /**
+     * Mouse events
+     */
+    void FPixelStreamingInputHandler::HandleOnMouseEnter(FMemoryReader Ar)
+    {
+        if(NumActiveTouches == 0 && !bIsMouseActive)
+        {
+            FSlateApplication::Get().OnCursorSet();
+            FSlateApplication::Get().OverridePlatformApplication(PixelStreamerApplicationWrapper);
+            // Make sure the application is active.
+            FSlateApplication::Get().ProcessApplicationActivationEvent(true);
+        }
+        
+        bIsMouseActive = true;
+        UE_LOG(LogPixelStreamingInputHandler, Verbose, TEXT("MOUSE_ENTER"));
+    }
+
+    void FPixelStreamingInputHandler::HandleOnMouseLeave(FMemoryReader Ar)
+    {
+        if(NumActiveTouches == 0)
+        {
+            // Restore normal application layer if there is no active touches and MouseEnter hasn't been triggered
+		    FSlateApplication::Get().OverridePlatformApplication(PixelStreamerApplicationWrapper->WrappedApplication);
+        }
+        bIsMouseActive = false;
+	    UE_LOG(LogPixelStreamingInputHandler, Verbose, TEXT("MOUSE_LEAVE"));
+    }
+
+    void FPixelStreamingInputHandler::HandleOnMouseUp(FMemoryReader Ar)
+    {
+        TPayloadThreeParam<uint8, uint16, uint16> Payload(Ar);
+
+        EMouseButtons::Type Button = static_cast<EMouseButtons::Type>(Payload.Param1);
+        UE_LOG(LogPixelStreamingInputHandler, Verbose, TEXT("MOUSE_UP: Button = %d"), Button);
+        if(Button != EMouseButtons::Type::Invalid)
+        {
+            MessageHandler->OnMouseUp(Button);
+        }
+    }
+
+    void FPixelStreamingInputHandler::HandleOnMouseDown(FMemoryReader Ar)
+    {
+        TPayloadThreeParam<uint8, uint16, uint16> Payload(Ar);
+        //                                                                           convert range from 0,65536 -> 0,1
+        FVector2D ScreenLocation = ConvertFromNormalizedScreenLocation(FVector2D(Payload.Param2 / uint16_MAX, Payload.Param3 / uint16_MAX));
+        EMouseButtons::Type Button = static_cast<EMouseButtons::Type>(Payload.Param1);
+
+		UE_LOG(LogPixelStreamingInputHandler, Verbose, TEXT("MOUSE_DOWN: Button = %d; Pos = (%.4f, %.4f)"), Button, ScreenLocation.X, ScreenLocation.Y);
+		// Force window focus
+		FSlateApplication::Get().ProcessApplicationActivationEvent(true);
+		MessageHandler->OnMouseDown(PixelStreamerApplicationWrapper->GetWindowUnderCursor(), Button, ScreenLocation);
+	}
+
+    void FPixelStreamingInputHandler::HandleOnMouseMove(FMemoryReader Ar)
+    {
+        TPayloadFourParam<uint16, uint16, int16, int16> Payload(Ar);
+        //                                                                           convert range from 0,65536 -> 0,1
+        FIntPoint ScreenLocation = ConvertFromNormalizedScreenLocation(FVector2D(Payload.Param1 / uint16_MAX, Payload.Param2 / uint16_MAX));
+        //                                                                 convert range from -32,768 to 32,767 -> -1,1
+        FIntPoint Delta = ConvertFromNormalizedScreenLocation(FVector2D(Payload.Param3 / int16_MAX, Payload.Param4 / int16_MAX), false);
+
+		FSlateApplication::Get().OnCursorSet();
+		UE_LOG(LogPixelStreamingInputHandler, Verbose, TEXT("MOUSE_MOVE: Pos = (%d, %d); Delta = (%d, %d)"), ScreenLocation.X, ScreenLocation.Y, Delta.X, Delta.Y);
+		PixelStreamerApplicationWrapper->Cursor->SetPosition(ScreenLocation.X, ScreenLocation.Y);
+		MessageHandler->OnRawMouseMove(Delta.X, Delta.Y);
+	}
+
+    void FPixelStreamingInputHandler::HandleOnMouseWheel(FMemoryReader Ar)
+    {
+        TPayloadThreeParam<int16, uint16, uint16> Payload(Ar);
+         //                                                                           convert range from 0,65536 -> 0,1
+        FIntPoint ScreenLocation = ConvertFromNormalizedScreenLocation(FVector2D(Payload.Param2 / uint16_MAX, Payload.Param3 / uint16_MAX));
+        const float SpinFactor = 1 / 120.0f;
+        MessageHandler->OnMouseWheel(Payload.Param1 * SpinFactor, ScreenLocation);
+        UE_LOG(LogPixelStreamingInputHandler, Verbose, TEXT("MOUSE_WHEEL: Delta = %d; Pos = (%d, %d)"), Payload.Param1, ScreenLocation.X, ScreenLocation.Y);
+    }
+
+	void FPixelStreamingInputHandler::HandleOnMouseDoubleClick(FMemoryReader Ar)
+	{
+		TPayloadThreeParam<uint8, uint16, uint16> Payload(Ar);
+        //                                                                           convert range from 0,65536 -> 0,1
+        FVector2D ScreenLocation = ConvertFromNormalizedScreenLocation(FVector2D(Payload.Param2 / uint16_MAX, Payload.Param3 / uint16_MAX));
+        EMouseButtons::Type Button = static_cast<EMouseButtons::Type>(Payload.Param1);
+
+		UE_LOG(LogPixelStreamingInputHandler, Verbose, TEXT("MOUSE_DOWN: Button = %d; Pos = (%.4f, %.4f)"), Button, ScreenLocation.X, ScreenLocation.Y);
+		// Force window focus
+		FSlateApplication::Get().ProcessApplicationActivationEvent(true);
+		MessageHandler->OnMouseDoubleClick(PixelStreamerApplicationWrapper->GetWindowUnderCursor(), Button, ScreenLocation);
+	}
+
+    /**
+     * Command handling
+     */
+    void FPixelStreamingInputHandler::HandleCommand(FMemoryReader Ar)
+    {
+        FString Res;
+        Res.GetCharArray().SetNumUninitialized(Ar.TotalSize() / 2 + 1);
+        Ar.Serialize(Res.GetCharArray().GetData(), Ar.TotalSize());
+
+		FString Descriptor = Res.Mid(MessageHeaderOffset);
+		UE_LOG(LogPixelStreamingInputHandler, Verbose, TEXT("Command: %s"), *Descriptor);
+
+		bool bSuccess = false;
+
+		FString ConsoleCommand;
+		UE::PixelStreaming::ExtractJsonFromDescriptor(Descriptor, TEXT("ConsoleCommand"), ConsoleCommand, bSuccess);
+		if (bSuccess && UE::PixelStreaming::Settings::CVarPixelStreamingAllowConsoleCommands.GetValueOnAnyThread())
+		{
+			GEngine->Exec(GEngine->GetWorld(), *ConsoleCommand);
+			return;
+		}
+
+		/**
+		 * Allowed Console commands
+		 */
+		FString WidthString;
+		FString HeightString;
+		UE::PixelStreaming::ExtractJsonFromDescriptor(Descriptor, TEXT("Resolution.Width"), WidthString, bSuccess);
+		if (bSuccess)
+		{
+			UE::PixelStreaming::ExtractJsonFromDescriptor(Descriptor, TEXT("Resolution.Height"), HeightString, bSuccess);
+
+			if (bSuccess)
+			{
+				int Width = FCString::Atoi(*WidthString);
+				int Height = FCString::Atoi(*HeightString);
+				if (Width < 1 || Height < 1)
+				{
+					return;
+				}
+
+				FString ChangeResCommand = FString::Printf(TEXT("r.SetRes %dx%d"), Width, Height);
+				GEngine->Exec(GEngine->GetWorld(), *ChangeResCommand);
+				return;
+			}
+		}
+
+		FString StatFPSString;
+		UE::PixelStreaming::ExtractJsonFromDescriptor(Descriptor, TEXT("Stat.FPS"), StatFPSString, bSuccess);
+		if (bSuccess)
+		{
+			FString StatFPSCommand = FString::Printf(TEXT("stat fps"));
+			GEngine->Exec(GEngine->GetWorld(), *StatFPSCommand);
+			return;
+		}
+
+		/**
+		 * Encoder Settings
+		 */
+		FString MinQPString;
+		UE::PixelStreaming::ExtractJsonFromDescriptor(Descriptor, TEXT("Encoder.MinQP"), MinQPString, bSuccess);
+		if (bSuccess)
+		{
+			int MinQP = FCString::Atoi(*MinQPString);
+			UE::PixelStreaming::Settings::CVarPixelStreamingEncoderMinQP->Set(MinQP, ECVF_SetByCommandline);
+			return;
+		}
+
+		FString MaxQPString;
+		UE::PixelStreaming::ExtractJsonFromDescriptor(Descriptor, TEXT("Encoder.MaxQP"), MaxQPString, bSuccess);
+		if (bSuccess)
+		{
+			int MaxQP = FCString::Atoi(*MaxQPString);
+			UE::PixelStreaming::Settings::CVarPixelStreamingEncoderMaxQP->Set(MaxQP, ECVF_SetByCommandline);
+			return;
+		}
+
+		/**
+		 * WebRTC Settings
+		 */
+		FString FPSString;
+		UE::PixelStreaming::ExtractJsonFromDescriptor(Descriptor, TEXT("WebRTC.Fps"), FPSString, bSuccess);
+		if (bSuccess)
+		{
+			int FPS = FCString::Atoi(*FPSString);
+			UE::PixelStreaming::Settings::CVarPixelStreamingWebRTCFps->Set(FPS, ECVF_SetByCommandline);
+			return;
+		}
+
+		FString MinBitrateString;
+		UE::PixelStreaming::ExtractJsonFromDescriptor(Descriptor, TEXT("WebRTC.MinBitrate"), MinBitrateString, bSuccess);
+		if (bSuccess)
+		{
+			int MinBitrate = FCString::Atoi(*MinBitrateString);
+			UE::PixelStreaming::Settings::CVarPixelStreamingWebRTCMinBitrate->Set(MinBitrate, ECVF_SetByCommandline);
+			return;
+		}
+
+		FString MaxBitrateString;
+		UE::PixelStreaming::ExtractJsonFromDescriptor(Descriptor, TEXT("WebRTC.MaxBitrate"), MaxBitrateString, bSuccess);
+		if (bSuccess)
+		{
+			int MaxBitrate = FCString::Atoi(*MaxBitrateString);
+			UE::PixelStreaming::Settings::CVarPixelStreamingWebRTCMaxBitrate->Set(MaxBitrate, ECVF_SetByCommandline);
+			return;
+		}
+	}
+
+    /**
+     * UI Interaction handling
+     */
+    void FPixelStreamingInputHandler::HandleUIInteraction(FMemoryReader Ar)
+    {
+        FString Res;
+        Res.GetCharArray().SetNumUninitialized(Ar.TotalSize() / 2 + 1);
+        Ar.Serialize(Res.GetCharArray().GetData(), Ar.TotalSize());
+
+		FString Descriptor = Res.Mid(MessageHeaderOffset);
+
+		UE_LOG(LogPixelStreamingInputHandler, Verbose, TEXT("UIInteraction: %s"), *Descriptor);
+		for (UPixelStreamingInput* InputComponent : PixelStreamingModule->GetInputComponents())
+		{
+			InputComponent->OnInputEvent.Broadcast(Descriptor);
+		}
+	}
+
+	FIntPoint FPixelStreamingInputHandler::ConvertFromNormalizedScreenLocation(const FVector2D& ScreenLocation, bool bIncludeOffset)
+	{
+		FIntPoint OutVector((int32)ScreenLocation.X, (int32)ScreenLocation.Y);
+
+		TSharedPtr<SWindow> ApplicationWindow = TargetWindow.Pin();
+		if (ApplicationWindow.IsValid())
+		{
+			FVector2D WindowOrigin = ApplicationWindow->GetPositionInScreen();
+			if (TargetViewport.IsValid())
+			{
+				TSharedPtr<SViewport> ViewportWidget = TargetViewport.Pin();
+
+				if (ViewportWidget.IsValid())
+				{
+					FGeometry InnerWindowGeometry = ApplicationWindow->GetWindowGeometryInWindow();
+
+					// Find the widget path relative to the window
+					FArrangedChildren JustWindow(EVisibility::Visible);
+					JustWindow.AddWidget(FArrangedWidget(ApplicationWindow.ToSharedRef(), InnerWindowGeometry));
+
+					FWidgetPath PathToWidget(ApplicationWindow.ToSharedRef(), JustWindow);
+					if (PathToWidget.ExtendPathTo(FWidgetMatcher(ViewportWidget.ToSharedRef()), EVisibility::Visible))
+					{
+						FArrangedWidget ArrangedWidget = PathToWidget.FindArrangedWidget(ViewportWidget.ToSharedRef()).Get(FArrangedWidget::GetNullWidget());
+
+						FVector2D WindowClientOffset = ArrangedWidget.Geometry.GetAbsolutePosition();
+						FVector2D WindowClientSize = ArrangedWidget.Geometry.GetAbsoluteSize();
+
+						FVector2D OutTemp = bIncludeOffset ? WindowOrigin + WindowClientOffset + (ScreenLocation * WindowClientSize) : (ScreenLocation * WindowClientSize);
+						UE_LOG(LogPixelStreamingInputHandler, Verbose, TEXT("%.4f, %.4f"), ScreenLocation.X, ScreenLocation.Y);
+						OutVector = FIntPoint((int32)OutTemp.X, (int32)OutTemp.Y);
+					}
+				}
+			}
+			else
+			{
+				FVector2D SizeInScreen = ApplicationWindow->GetSizeInScreen();
+				FVector2D OutTemp = SizeInScreen * ScreenLocation;
+				OutVector = FIntPoint((int32)OutTemp.X, (int32)OutTemp.Y);
+			}
+		}
+		else if(TargetScreenSize.IsValid())
+		{
+			FIntPoint SizeInScreen = *(TargetScreenSize.Pin());
+			FVector2D OutTemp = FVector2D(SizeInScreen) * ScreenLocation;
+			OutVector = FIntPoint((int32)OutTemp.X, (int32)OutTemp.Y);
+		}
+
+		return OutVector;
+	}
+
+	bool FPixelStreamingInputHandler::FilterKey(const FKey& Key)
+	{
+		for (auto&& FilteredKey : UE::PixelStreaming::Settings::FilteredKeys)
+		{
+			if (FilteredKey == Key)
+				return false;
+		}
+		return true;
+	}
+
+	FGamepadKeyNames::Type FPixelStreamingInputHandler::ConvertAxisIndexToGamepadAxis(uint8 AnalogAxis)
+	{
+		switch (AnalogAxis)
+		{
+			case 1:
+			{
+				return FGamepadKeyNames::LeftAnalogX;
+			}
+			break;
+			case 2:
+			{
+				return FGamepadKeyNames::LeftAnalogY;
+			}
+			break;
+			case 3:
+			{
+				return FGamepadKeyNames::RightAnalogX;
+			}
+			break;
+			case 4:
+			{
+				return FGamepadKeyNames::RightAnalogY;
+			}
+			break;
+			case 5:
+			{
+				return FGamepadKeyNames::LeftTriggerAnalog;
+			}
+			break;
+			case 6:
+			{
+				return FGamepadKeyNames::RightTriggerAnalog;
+			}
+			break;
+			default:
+			{
+				return FGamepadKeyNames::Invalid;
+			}
+			break;
+		}
+	}
+
+	FGamepadKeyNames::Type FPixelStreamingInputHandler::ConvertButtonIndexToGamepadButton(uint8 ButtonIndex)
+	{
+		switch (ButtonIndex)
+		{
+			case 0:
+			{
+				return FGamepadKeyNames::FaceButtonBottom;
+			}
+			case 1:
+			{
+				return FGamepadKeyNames::FaceButtonRight;
+			}
+			break;
+			case 2:
+			{
+				return FGamepadKeyNames::FaceButtonLeft;
+			}
+			break;
+			case 3:
+			{
+				return FGamepadKeyNames::FaceButtonTop;
+			}
+			break;
+			case 4:
+			{
+				return FGamepadKeyNames::LeftShoulder;
+			}
+			break;
+			case 5:
+			{
+				return FGamepadKeyNames::RightShoulder;
+			}
+			break;
+			// Buttons 6 and 7 are mapped as analog axis as they are the triggers
+			case 8:
+			{
+				return FGamepadKeyNames::SpecialLeft;
+			}
+			break;
+			case 9:
+			{
+				return FGamepadKeyNames::SpecialRight;
+			}
+			case 10:
+			{
+				return FGamepadKeyNames::LeftThumb;
+			}
+			break;
+			case 11:
+			{
+				return FGamepadKeyNames::RightThumb;
+			}
+			break;
+			case 12:
+			{
+				return FGamepadKeyNames::DPadUp;
+			}
+			break;
+			case 13:
+			{
+				return FGamepadKeyNames::DPadDown;
+			}
+			case 14:
+			{
+				return FGamepadKeyNames::DPadLeft;
+			}
+			break;
+			case 15:
+			{
+				return FGamepadKeyNames::DPadRight;
+			}
+			break;
+			default:
+			{
+				return FGamepadKeyNames::Invalid;
+			}
+			break;
+		}
+	}
+
+	void FPixelStreamingInputHandler::FindFocusedWidget()
+	{
+		FSlateApplication::Get().ForEachUser([this](FSlateUser& User) {
+			TSharedPtr<SWidget> FocusedWidget = User.GetFocusedWidget();
+
+			static FName SEditableTextType(TEXT("SEditableText"));
+			static FName SMultiLineEditableTextType(TEXT("SMultiLineEditableText"));
+			bool bEditable = FocusedWidget && (FocusedWidget->GetType() == SEditableTextType || FocusedWidget->GetType() == SMultiLineEditableTextType);
+
+			// Check to see if the focus has changed.
+			FVector2D Pos = bEditable ? FocusedWidget->GetCachedGeometry().GetAbsolutePosition() : UnfocusedPos;
+			if (Pos != FocusedPos)
+			{
+				FocusedPos = Pos;
+
+				// Tell the browser that the focus has changed.
+				TSharedPtr<FJsonObject> JsonObject = MakeShareable(new FJsonObject);
+				JsonObject->SetStringField(TEXT("command"), TEXT("onScreenKeyboard"));
+				JsonObject->SetBoolField(TEXT("showOnScreenKeyboard"), bEditable);
+
+				if (bEditable)
+				{
+					FVector2D NormalizedLocation;
+					TSharedPtr<SWindow> ApplicationWindow = TargetWindow.Pin();
+					if (ApplicationWindow.IsValid())
+					{
+						FVector2D WindowOrigin = ApplicationWindow->GetPositionInScreen();
+						if (TargetViewport.IsValid())
+						{
+							TSharedPtr<SViewport> ViewportWidget = TargetViewport.Pin();
+
+							if (ViewportWidget.IsValid())
+							{
+								FGeometry InnerWindowGeometry = ApplicationWindow->GetWindowGeometryInWindow();
+
+								// Find the widget path relative to the window
+								FArrangedChildren JustWindow(EVisibility::Visible);
+								JustWindow.AddWidget(FArrangedWidget(ApplicationWindow.ToSharedRef(), InnerWindowGeometry));
+
+								FWidgetPath PathToWidget(ApplicationWindow.ToSharedRef(), JustWindow);
+								if (PathToWidget.ExtendPathTo(FWidgetMatcher(ViewportWidget.ToSharedRef()), EVisibility::Visible))
+								{
+									FArrangedWidget ArrangedWidget = PathToWidget.FindArrangedWidget(ViewportWidget.ToSharedRef()).Get(FArrangedWidget::GetNullWidget());
+
+									FVector2D WindowClientOffset = ArrangedWidget.Geometry.GetAbsolutePosition();
+									FVector2D WindowClientSize = ArrangedWidget.Geometry.GetAbsoluteSize();
+
+									Pos = Pos - WindowClientOffset;
+									NormalizedLocation = FVector2D(Pos / WindowClientSize);
+								}
+							}
+						}
+						else
+						{
+							FVector2D SizeInScreen = ApplicationWindow->GetSizeInScreen();
+							NormalizedLocation = FVector2D(Pos / SizeInScreen);
+						}
+					}
+
+					NormalizedLocation *= uint16_MAX;
+					// ConvertToNormalizedScreenLocation(Pos, NormalizedLocation);
+					JsonObject->SetNumberField(TEXT("x"), static_cast<uint16>(NormalizedLocation.X));
+					JsonObject->SetNumberField(TEXT("y"), static_cast<uint16>(NormalizedLocation.Y));
+				}
+
+				FString Descriptor;
+				TSharedRef<TJsonWriter<TCHAR, TCondensedJsonPrintPolicy<TCHAR>>> JsonWriter = TJsonWriterFactory<TCHAR, TCondensedJsonPrintPolicy<TCHAR>>::Create(&Descriptor);
+				FJsonSerializer::Serialize(JsonObject.ToSharedRef(), JsonWriter);
+
+				PixelStreamingModule->ForEachStreamer([&Descriptor, this](TSharedPtr<IPixelStreamingStreamer> Streamer){
+					Streamer->SendPlayerMessage(PixelStreamingModule->GetProtocol().FromStreamerProtocol.Find("Command")->Id, Descriptor);
+				});
+			}
+		});
+    }
+} // namespace UE::PixelStreaming
diff --git a/Engine/Plugins/Media/PixelStreaming/Source/PixelStreaming/Private/PixelStreamingInputHandler.h b/Engine/Plugins/Media/PixelStreaming/Source/PixelStreaming/Private/PixelStreamingInputHandler.h
new file mode 100644
--- /dev/null
+++ b/Engine/Plugins/Media/PixelStreaming/Source/PixelStreaming/Private/PixelStreamingInputHandler.h
@@ -0,0 +1,155 @@
+// Copyright Epic Games, Inc. All Rights Reserved.
+
+#pragma once
+
+#include "GenericPlatform/GenericApplicationMessageHandler.h"
+#include "Slate/SceneViewport.h"
+#include "Widgets/SWindow.h"
+#include "WebRTCIncludes.h"
+#include "IPixelStreamingModule.h"
+#include "InputCoreTypes.h"
+#include "PixelStreamingApplicationWrapper.h"
+#include "IPixelStreamingInputHandler.h"
+
+namespace UE::PixelStreaming
+{
+	class FPixelStreamingInputHandler : public IPixelStreamingInputHandler
+	{
+	public:
+		FPixelStreamingInputHandler(TSharedPtr<FPixelStreamingApplicationWrapper> InApplicationWrapper, const TSharedPtr<FGenericApplicationMessageHandler>& InTargetHandler);
+
+		virtual ~FPixelStreamingInputHandler();
+
+		virtual void Tick(float DeltaTime) override;
+
+		/** Poll for controller state and send events if needed */
+		virtual void SendControllerEvents() override {};
+
+		/** Set which MessageHandler will route input  */
+		virtual void SetMessageHandler(const TSharedRef<FGenericApplicationMessageHandler>& InTargetHandler) override;
+
+		/** Exec handler to allow console commands to be passed through for debugging */
+		virtual bool Exec(UWorld* InWorld, const TCHAR* Cmd, FOutputDevice& Ar) override;
+		/**
+		 * IForceFeedbackSystem pass through functions
+		 */
+		virtual void SetChannelValue(int32 ControllerId, FForceFeedbackChannelType ChannelType, float Value) override;
+
+		virtual void SetChannelValues(int32 ControllerId, const FForceFeedbackValues& values) override;
+
+		void OnMessage(const webrtc::DataBuffer& Buffer);
+
+		void SetTargetWindow(TWeakPtr<SWindow> InWindow);
+		TWeakPtr<SWindow> GetTargetWindow();
+
+		void SetTargetViewport(TWeakPtr<SViewport> InViewport);
+		TWeakPtr<SViewport> GetTargetViewport();
+
+		void SetTargetScreenSize(TWeakPtr<FIntPoint> InScreenSize);
+		TWeakPtr<FIntPoint> GetTargetScreenSize();
+
+		bool IsFakingTouchEvents() const { return bFakingTouchEvents; }
+
+        void RegisterMessageHandler(const FString& MessageType, const TFunction<void(FMemoryReader)>& Handler);
+		TFunction<void(FMemoryReader)> FindMessageHandler(const FString& MessageType);
+
+    protected:
+        /**
+         * Key press handling
+         */
+        virtual void HandleOnKeyChar(FMemoryReader Ar);
+        virtual void HandleOnKeyDown(FMemoryReader Ar);
+        virtual void HandleOnKeyUp(FMemoryReader Ar);
+        /**
+         * Touch handling
+         */
+        virtual void HandleOnTouchStarted(FMemoryReader Ar);
+        virtual void HandleOnTouchMoved(FMemoryReader Ar);
+        virtual void HandleOnTouchEnded(FMemoryReader Ar);
+        /**
+         * Controller handling
+         */
+        virtual void HandleOnControllerAnalog(FMemoryReader Ar);
+        virtual void HandleOnControllerButtonPressed(FMemoryReader Ar);
+        virtual void HandleOnControllerButtonReleased(FMemoryReader Ar);
+        /**
+         * Mouse handling
+         */
+        virtual void HandleOnMouseEnter(FMemoryReader Ar);
+        virtual void HandleOnMouseLeave(FMemoryReader Ar);
+        virtual void HandleOnMouseUp(FMemoryReader Ar);
+        virtual void HandleOnMouseDown(FMemoryReader Ar);
+        virtual void HandleOnMouseMove(FMemoryReader Ar);
+        virtual void HandleOnMouseWheel(FMemoryReader Ar);
+		virtual void HandleOnMouseDoubleClick(FMemoryReader Ar);
+        /**
+         * Command handling
+         */
+        virtual void HandleCommand(FMemoryReader Ar);
+        /**
+         * UI Interaction handling
+         */
+        virtual void HandleUIInteraction(FMemoryReader Ar);
+
+		FIntPoint ConvertFromNormalizedScreenLocation(const FVector2D& ScreenLocation, bool bIncludeOffset = true);
+
+		FGamepadKeyNames::Type ConvertAxisIndexToGamepadAxis(uint8 AnalogAxis);
+		FGamepadKeyNames::Type ConvertButtonIndexToGamepadButton(uint8 ButtonIndex);
+
+		void FindFocusedWidget();
+		bool FilterKey(const FKey& Key);
+
+        struct FMessage
+        {
+            TFunction<void(FMemoryReader)>* Handler;
+            TArray<uint8> Data;
+        };
+
+		TWeakPtr<SWindow> TargetWindow;
+		TWeakPtr<SViewport> TargetViewport;
+		TWeakPtr<FIntPoint> TargetScreenSize; // Manual size override used when we don't have a single window/viewport target
+		uint8 NumActiveTouches;
+		bool bIsMouseActive;
+		TMap<uint8,  TFunction<void(FMemoryReader)>> DispatchTable;
+		TQueue<FMessage> Messages;
+
+		/** Reference to the message handler which events should be passed to. */
+		TSharedPtr<FGenericApplicationMessageHandler> MessageHandler;
+
+		/** For convenience we keep a reference to the Pixel Streaming plugin. */
+		IPixelStreamingModule* PixelStreamingModule;
+
+		/** For convenience, we keep a reference to the application wrapper owned by the input channel */
+		TSharedPtr<FPixelStreamingApplicationWrapper> PixelStreamerApplicationWrapper;
+
+		/**
+		 * Is the application faking touch events by dragging the mouse along
+		 * the canvas? If so then we must put the browser canvas in a special
+		 * state to replicate the behavior of the application.
+		 */
+		bool bFakingTouchEvents;
+
+		/**
+		 * Touch only. Location of the focused UI widget. If no UI widget is focused
+		 * then this has the UnfocusedPos value.
+		 */
+		FVector2D FocusedPos;
+
+		/**
+		 * Touch only. A special position which indicates that no UI widget is
+		 * focused.
+		 */
+		const FVector2D UnfocusedPos;
+
+		/*
+		 * Padding for string parsing when handling messages.
+		 * 1 character for the actual message and then
+		 * 2 characters for the length which are skipped
+		 */
+		const size_t MessageHeaderOffset = 1;
+	
+	private:
+		float uint16_MAX = (float) UINT16_MAX;
+		float int16_MAX = (float) SHRT_MAX;
+	};
+} // namespace UE::PixelStreaming
\ No newline at end of file
diff --git a/Engine/Plugins/Media/PixelStreaming/Source/PixelStreaming/Private/PixelStreamingMessageHandler.cpp b/Engine/Plugins/Media/PixelStreaming/Source/PixelStreaming/Private/PixelStreamingMessageHandler.cpp
deleted file mode 100644
diff --git a/Engine/Plugins/Media/PixelStreaming/Source/PixelStreaming/Private/PixelStreamingMessageHandler.h b/Engine/Plugins/Media/PixelStreaming/Source/PixelStreaming/Private/PixelStreamingMessageHandler.h
deleted file mode 100644
diff --git a/Engine/Plugins/Media/PixelStreaming/Source/PixelStreaming/Private/PixelStreamingModule.cpp b/Engine/Plugins/Media/PixelStreaming/Source/PixelStreaming/Private/PixelStreamingModule.cpp
--- a/Engine/Plugins/Media/PixelStreaming/Source/PixelStreaming/Private/PixelStreamingModule.cpp
+++ b/Engine/Plugins/Media/PixelStreaming/Source/PixelStreaming/Private/PixelStreamingModule.cpp
@@ -2,7 +2,6 @@
 
 #include "PixelStreamingModule.h"
 #include "Streamer.h"
-#include "PixelStreamingInputChannel.h"
 #include "PixelStreamingInputComponent.h"
 #include "PixelStreamingDelegates.h"
 #include "PixelStreamingSignallingConnection.h"
@@ -57,6 +56,9 @@ THIRD_PARTY_INCLUDES_END
 #include "VideoSourceGroup.h"
 #include "PixelStreamingPeerConnection.h"
 #include "Engine/GameEngine.h"
+#include "PixelStreamingApplicationWrapper.h"
+#include "PixelStreamingInputHandler.h"
+#include "InputHandlers.h"
 
 DEFINE_LOG_CATEGORY(LogPixelStreaming);
 
@@ -86,7 +88,6 @@ namespace UE::PixelStreaming
 
 		const ERHIInterfaceType RHIType = GDynamicRHI ? RHIGetInterfaceType() : ERHIInterfaceType::Hidden;
 
-		StreamerInputChannels = MakeShared<FStreamerInputChannels>(FSlateApplication::Get().GetPlatformApplication()->GetMessageHandler());
 		IModularFeatures::Get().RegisterModularFeature(GetModularFeatureName(), this);
 		PopulateProtocol();
 		
@@ -136,8 +137,8 @@ namespace UE::PixelStreaming
 
 		// ExternalVideoSourceGroup is used so that we can have a video source without a streamer
 		ExternalVideoSourceGroup = FVideoSourceGroup::Create();
-		ExternalVideoSourceGroup->SetVideoInput(FPixelStreamingVideoInputBackBuffer::Create());
-		ExternalVideoSourceGroup->Start();
+		// ExternalVideoSourceGroup->SetVideoInput(FPixelStreamingVideoInputBackBuffer::Create());
+		// ExternalVideoSourceGroup->Start();
 
 		bStartupCompleted = true;
 	}
@@ -149,8 +150,6 @@ namespace UE::PixelStreaming
 			return;
 		}
 
-		IModularFeatures::Get().UnregisterModularFeature(GetModularFeatureName(), this);
-
 		// We explicitly call release on streamer so WebRTC gets shutdown before our module is deleted
 		Streamers.Empty();
 		ExternalVideoSourceGroup->Stop();
@@ -158,7 +157,6 @@ namespace UE::PixelStreaming
 		FPixelStreamingPeerConnection::Shutdown();
 
 		rtc::CleanupSSL();
-
 		IModularFeatures::Get().UnregisterModularFeature(GetModularFeatureName(), this);
 
 		bStartupCompleted = false;
@@ -208,34 +206,30 @@ namespace UE::PixelStreaming
 	bool FPixelStreamingModule::StartStreaming()
 	{
 		bool bSuccess = true;
-		TMap<FString, TSharedPtr<IPixelStreamingStreamer>>::TIterator Iter = Streamers.CreateIterator();
-		for (; Iter; ++Iter)
+		ForEachStreamer([&bSuccess](TSharedPtr<IPixelStreamingStreamer> Streamer)
 		{
-			TSharedPtr<IPixelStreamingStreamer> Streamer = Iter.Value();
-
 			if (Streamer.IsValid())
 			{
 				Streamer->StartStreaming();
 				bSuccess &= true;
-				continue;
 			}
-			bSuccess &= false;
-		}
+			else
+			{
+				bSuccess = false;
+			}
+		});
 		return bSuccess;
 	}
 
 	void FPixelStreamingModule::StopStreaming()
 	{
-		TMap<FString, TSharedPtr<IPixelStreamingStreamer>>::TIterator Iter = Streamers.CreateIterator();
-		for (; Iter; ++Iter)
+		ForEachStreamer([this](TSharedPtr<IPixelStreamingStreamer> Streamer)
 		{
-			TSharedPtr<IPixelStreamingStreamer> Streamer = Iter.Value();
-
 			if (Streamer.IsValid())
 			{
 				Streamer->StopStreaming();
 			}
-		}
+		});
 	}
 
 	TSharedPtr<IPixelStreamingStreamer> FPixelStreamingModule::CreateStreamer(const FString& StreamerId)
@@ -251,7 +245,6 @@ namespace UE::PixelStreaming
 			FScopeLock Lock(&StreamersCS);
 			Streamers.Add(StreamerId, NewStreamer);
 		}
-		NewStreamer->SetInputChannel(StreamerInputChannels->CreateInputChannel());
 
 		return NewStreamer;
 	}
@@ -359,11 +352,11 @@ namespace UE::PixelStreaming
 			MessageProtocol.ToStreamerProtocol.Add(MessageType, Message);
 			ForEachStreamer([&Handler = Handler, &MessageType = MessageType](TSharedPtr<IPixelStreamingStreamer> Streamer)
 			{
-				TWeakPtr<IPixelStreamingInputChannel> WeakChannel = Streamer->GetInputChannel();
-				TSharedPtr<IPixelStreamingInputChannel> Channel = WeakChannel.Pin();
-				if (Channel)
+				TWeakPtr<IPixelStreamingInputHandler> WeakInputHandler = Streamer->GetInputHandler();
+				TSharedPtr<IPixelStreamingInputHandler> InputHandler = WeakInputHandler.Pin();
+				if (InputHandler)
 				{
-					Channel->RegisterHandler(MessageType, Handler);
+					InputHandler->RegisterMessageHandler(MessageType, Handler);
 				}
 			});
 		}
@@ -374,6 +367,19 @@ namespace UE::PixelStreaming
 		OnProtocolUpdated.Broadcast();
 	}
 
+	TFunction<void(FMemoryReader)> FPixelStreamingModule::FindMessageHandler(const FString& MessageType)
+	{
+		// All streamers have the same protocol so we just use the first streamers input channel to get the message handler
+		TSharedPtr<IPixelStreamingStreamer> Streamer = Streamers.begin()->Value;
+		TWeakPtr<IPixelStreamingInputHandler> WeakInputHandler = Streamer->GetInputHandler();
+		if (TSharedPtr<IPixelStreamingInputHandler> InputHandler = WeakInputHandler.Pin())
+		{
+			return InputHandler->FindMessageHandler(MessageType);
+		}
+		// If the channel doesn't exist, just return an empty function
+		return ([](FMemoryReader Ar) {});
+	}
+
 	/**
 	 * End IPixelStreamingModule implementation
 	 */
@@ -450,20 +456,13 @@ namespace UE::PixelStreaming
 
 		return bCompatible;
 	}
-
 	/**
 	 * End own methods
 	 */
 
 	TSharedPtr<IInputDevice> FPixelStreamingModule::CreateInputDevice(const TSharedRef<FGenericApplicationMessageHandler>& InMessageHandler)
 	{
-		return StreamerInputChannels;
-	}
-
-	void FPixelStreamingModule::RegisterCreateInputChannel(IPixelStreamingInputChannel::FCreateInputChannelFunc& InCreateInputChannel)
-	{
-		checkf(StreamerInputChannels, TEXT("StreamerInputChannels does not exist yet"));
-		StreamerInputChannels->OverrideInputChannel(InCreateInputChannel);
+		return MakeShared<FInputHandlers>(InMessageHandler);
 	}
 
 	void FPixelStreamingModule::PopulateProtocol()
diff --git a/Engine/Plugins/Media/PixelStreaming/Source/PixelStreaming/Private/PixelStreamingModule.h b/Engine/Plugins/Media/PixelStreaming/Source/PixelStreaming/Private/PixelStreamingModule.h
--- a/Engine/Plugins/Media/PixelStreaming/Source/PixelStreaming/Private/PixelStreamingModule.h
+++ b/Engine/Plugins/Media/PixelStreaming/Source/PixelStreaming/Private/PixelStreamingModule.h
@@ -1,3 +1,4 @@
+
 // Copyright Epic Games, Inc. All Rights Reserved.
 
 #pragma once
@@ -5,8 +6,6 @@
 #include "IPixelStreamingModule.h"
 #include "RHI.h"
 #include "Tickable.h"
-#include "PixelStreamingInputChannel.h"
-#include "StreamerInputChannels.h"
 #include "PixelStreamingProtocol.h"
 
 class UPixelStreamingInput;
@@ -42,6 +41,7 @@ namespace UE::PixelStreaming
 		virtual const Protocol::FPixelStreamingProtocol& GetProtocol() override;
 		
 		virtual void RegisterMessage(Protocol::EPixelStreamingMessageDirection MessageDirection, const FString& MessageType, Protocol::FPixelStreamingInputMessage Message, const TFunction<void(FMemoryReader)>& Handler) override;
+		virtual TFunction<void(FMemoryReader)> FindMessageHandler(const FString& MessageType) override;
 		// These are staying on the module at the moment as theres no way of the BPs knowing which streamer they are relevant to
 		virtual void AddInputComponent(UPixelStreamingInput* InInputComponent) override;
 		virtual void RemoveInputComponent(UPixelStreamingInput* InInputComponent) override;
@@ -54,20 +54,19 @@ namespace UE::PixelStreaming
 		virtual void ForEachStreamer(const TFunction<void(TSharedPtr<IPixelStreamingStreamer>)>& Func) override;
 		/** End IPixelStreamingModule implementation */
 
-		virtual void RegisterCreateInputChannel(IPixelStreamingInputChannel::FCreateInputChannelFunc& InCreateInputChannel) override;
-
 	private:
 		/** IModuleInterface implementation */
 		void StartupModule() override;
 		void ShutdownModule() override;
 		/** End IModuleInterface implementation */
 
+		/** IInputDeviceModule implementation */
+		virtual TSharedPtr<IInputDevice> CreateInputDevice(const TSharedRef<FGenericApplicationMessageHandler>& InMessageHandler) override;
+		/** End IInputDeviceModule implementation */
+
 		// Own methods
 		void InitDefaultStreamer();
 		bool IsPlatformCompatible() const;
-
-		virtual TSharedPtr<IInputDevice> CreateInputDevice(const TSharedRef<FGenericApplicationMessageHandler>& InMessageHandler) override;
-
 		void PopulateProtocol();
 
 	private:
@@ -81,6 +80,5 @@ namespace UE::PixelStreaming
 		mutable FCriticalSection StreamersCS;
 		TMap<FString, TSharedPtr<IPixelStreamingStreamer>> Streamers;
 		Protocol::FPixelStreamingProtocol MessageProtocol;
-		TSharedPtr<FStreamerInputChannels> StreamerInputChannels;
 	};
 } // namespace UE::PixelStreaming
diff --git a/Engine/Plugins/Media/PixelStreaming/Source/PixelStreaming/Private/PixelStreamingVideoInputViewport.cpp b/Engine/Plugins/Media/PixelStreaming/Source/PixelStreaming/Private/PixelStreamingVideoInputViewport.cpp
--- a/Engine/Plugins/Media/PixelStreaming/Source/PixelStreaming/Private/PixelStreamingVideoInputViewport.cpp
+++ b/Engine/Plugins/Media/PixelStreaming/Source/PixelStreaming/Private/PixelStreamingVideoInputViewport.cpp
@@ -5,17 +5,16 @@
 #include "Utils.h"
 #include "PixelStreamingModule.h"
 #include "PixelStreamingPrivate.h"
-
 #include "PixelCaptureInputFrameRHI.h"
 #include "PixelCaptureBufferFormat.h"
 #include "PixelCaptureCapturerRHI.h"
 #include "PixelCaptureCapturerRHIToI420CPU.h"
 #include "PixelCaptureCapturerRHIToI420Compute.h"
-
 #include "UnrealClient.h"
 #include "Engine/GameViewportClient.h"
 #include "RenderingThread.h"
 #include "Framework/Application/SlateApplication.h"
+#include "Slate/SceneViewport.h"
 
 namespace UE::PixelStreaming
 {
diff --git a/Engine/Plugins/Media/PixelStreaming/Source/PixelStreaming/Private/RTCStatsCollector.cpp b/Engine/Plugins/Media/PixelStreaming/Source/PixelStreaming/Private/RTCStatsCollector.cpp
--- a/Engine/Plugins/Media/PixelStreaming/Source/PixelStreaming/Private/RTCStatsCollector.cpp
+++ b/Engine/Plugins/Media/PixelStreaming/Source/PixelStreaming/Private/RTCStatsCollector.cpp
@@ -24,12 +24,12 @@ namespace UE::PixelStreaming
 		Add(PixelStreamingStatNames::FrameWidth, 0);
 		Add(PixelStreamingStatNames::FrameHeight, 0);
 		Add(PixelStreamingStatNames::HugeFramesSent, 0);
+		Add(PixelStreamingStatNames::AvgSendDelay, 0);
 
 		AddForCalculation(PixelStreamingStatNames::FramesSent);
 		AddForCalculation(PixelStreamingStatNames::BytesSent);
 		AddForCalculation(PixelStreamingStatNames::QPSum);
 		AddForCalculation(PixelStreamingStatNames::TotalEncodeTime);
-		AddForCalculation(PixelStreamingStatNames::TotalPacketSendDelay);
 		AddForCalculation(PixelStreamingStatNames::FramesEncoded);
 	}
 
diff --git a/Engine/Plugins/Media/PixelStreaming/Source/PixelStreaming/Private/Streamer.cpp b/Engine/Plugins/Media/PixelStreaming/Source/PixelStreaming/Private/Streamer.cpp
--- a/Engine/Plugins/Media/PixelStreaming/Source/PixelStreaming/Private/Streamer.cpp
+++ b/Engine/Plugins/Media/PixelStreaming/Source/PixelStreaming/Private/Streamer.cpp
@@ -28,6 +28,8 @@
 #include "PixelStreamingModule.h"
 #include "PixelStreamingProtocol.h"
 #include "IPixelStreamingModule.h"
+#include "PixelStreamingApplicationWrapper.h"
+#include "PixelStreamingInputHandler.h"
 #include "PixelCaptureBufferFormat.h"
 #include "PixelCaptureOutputFrameRHI.h"
 #include "PixelCaptureInputFrameRHI.h"
@@ -48,6 +50,9 @@ namespace UE::PixelStreaming
 		SignallingServerConnection = MakeUnique<FPixelStreamingSignallingConnection>(WebSocketFactory, *this, InStreamerId);
 		SignallingServerConnection->SetAutoReconnect(true);
 
+		TSharedPtr<FPixelStreamingApplicationWrapper> PixelStreamerApplicationWrapper = MakeShareable(new FPixelStreamingApplicationWrapper(FSlateApplication::Get().GetPlatformApplication()));
+		TSharedPtr<FGenericApplicationMessageHandler> BaseHandler = FSlateApplication::Get().GetPlatformApplication()->GetMessageHandler();
+		InputHandler = MakeShared<FPixelStreamingInputHandler>(PixelStreamerApplicationWrapper, BaseHandler);
 		Module.OnProtocolUpdated.AddRaw(this, &FStreamer::OnProtocolUpdated);
 	}
 
@@ -112,32 +117,32 @@ namespace UE::PixelStreaming
 
 	void FStreamer::SetTargetViewport(TWeakPtr<SViewport> InTargetViewport)
 	{
-		InputChannel->SetTargetViewport(InTargetViewport);
+		InputHandler->SetTargetViewport(InTargetViewport);
 	}
 
 	TWeakPtr<SViewport> FStreamer::GetTargetViewport()
 	{
-		return InputChannel ? InputChannel->GetTargetViewport() : nullptr;
+		return InputHandler ? InputHandler->GetTargetViewport() : nullptr;
 	}
 
 	void FStreamer::SetTargetWindow(TWeakPtr<SWindow> InTargetWindow)
 	{
-		InputChannel->SetTargetWindow(InTargetWindow);
+		InputHandler->SetTargetWindow(InTargetWindow);
 	}
 
 	TWeakPtr<SWindow> FStreamer::GetTargetWindow()
 	{
-		return InputChannel->GetTargetWindow();
+		return InputHandler->GetTargetWindow();
 	}
 
 	void FStreamer::SetTargetScreenSize(TWeakPtr<FIntPoint> InTargetScreenSize)
 	{
-		InputChannel->SetTargetScreenSize(InTargetScreenSize);
+		InputHandler->SetTargetScreenSize(InTargetScreenSize);
 	}
 
 	TWeakPtr<FIntPoint> FStreamer::GetTargetScreenSize()
 	{
-		return InputChannel->GetTargetScreenSize();
+		return InputHandler->GetTargetScreenSize();
 	}
 
 	void FStreamer::SetSignallingServerURL(const FString& InSignallingServerURL)
@@ -168,6 +173,7 @@ namespace UE::PixelStreaming
 		if (UPixelStreamingDelegates* Delegates = UPixelStreamingDelegates::GetPixelStreamingDelegates())
 		{
 			ConsumeStatsHandle = Delegates->OnStatChangedNative.AddSP(AsShared(), &FStreamer::ConsumeStats);
+			AllConnectionsClosedHandle = Delegates->OnAllConnectionsClosedNative.AddSP(AsShared(), &FStreamer::TriggerMouseLeave);
 		}
 
 		VideoSourceGroup->Start();
@@ -180,10 +186,12 @@ namespace UE::PixelStreaming
 		if (UPixelStreamingDelegates* Delegates = UPixelStreamingDelegates::GetPixelStreamingDelegates())
 		{
 			Delegates->OnStatChangedNative.Remove(ConsumeStatsHandle);
+			Delegates->OnAllConnectionsClosedNative.Remove(AllConnectionsClosedHandle);
 		}
 
 		SignallingServerConnection->Disconnect();
 		VideoSourceGroup->Stop();
+		TriggerMouseLeave(StreamerId);
 
 		if (bStreamingStarted)
 		{
@@ -344,8 +352,8 @@ namespace UE::PixelStreaming
 
 	void FStreamer::AddPlayerConfig(TSharedRef<FJsonObject>& JsonObject)
 	{
-		checkf(InputChannel.IsValid(), TEXT("No Input Device available when populating Player Config"));
-		JsonObject->SetBoolField(TEXT("FakingTouchEvents"), InputChannel->IsFakingTouchEvents());
+		checkf(InputHandler.IsValid(), TEXT("No Input Device available when populating Player Config"));
+		JsonObject->SetBoolField(TEXT("FakingTouchEvents"), InputHandler->IsFakingTouchEvents());
 		FString PixelStreamingControlScheme;
 		if (Settings::GetControlScheme(PixelStreamingControlScheme))
 		{
@@ -391,6 +399,12 @@ namespace UE::PixelStreaming
 					SetQualityController(PlayerId);
 				}
 
+				if (UPixelStreamingDelegates* Delegates = UPixelStreamingDelegates::GetPixelStreamingDelegates())
+				{
+					Delegates->OnNewConnection.Broadcast(StreamerId, PlayerId, !PlayerContext->Config.IsSFU);
+					Delegates->OnNewConnectionNative.Broadcast(StreamerId, PlayerId, !PlayerContext->Config.IsSFU);
+				}
+
 				return true;
 			}
 		}
@@ -637,7 +651,8 @@ namespace UE::PixelStreaming
 			SFUPlayerId = INVALID_PLAYER_ID;
 		}
 
-		if (PlayerId == QualityControllingId)
+		bool bWasQualityController = PlayerId == QualityControllingId;
+		if (bWasQualityController)
 		{
 			SetQualityController(INVALID_PLAYER_ID);
 
@@ -659,6 +674,17 @@ namespace UE::PixelStreaming
 		{
 			Stats->RemovePeerStats(PlayerId);
 		}
+
+		if (UPixelStreamingDelegates* Delegates = UPixelStreamingDelegates::GetPixelStreamingDelegates())
+		{
+			Delegates->OnClosedConnection.Broadcast(StreamerId, PlayerId, bWasQualityController);
+			Delegates->OnClosedConnectionNative.Broadcast(StreamerId, PlayerId, bWasQualityController);
+			if(Players.IsEmpty())
+			{
+				Delegates->OnAllConnectionsClosed.Broadcast(StreamerId);
+				Delegates->OnAllConnectionsClosedNative.Broadcast(StreamerId);
+			}
+		}
 	}
 
 	void FStreamer::DeleteAllPlayerSessions()
@@ -708,6 +734,12 @@ namespace UE::PixelStreaming
 			InputControllingId = PlayerId;
 		}
 
+		if (UPixelStreamingDelegates* Delegates = UPixelStreamingDelegates::GetPixelStreamingDelegates())
+		{
+			FPlayerContext* PlayerContext = Players.Find(PlayerId);
+			Delegates->OnDataChannelOpenNative.Broadcast(StreamerId, PlayerId, PlayerContext->DataChannel.Get());
+		}
+
 		// When data channel is open
 		SendProtocol(PlayerId);
 		// Try to send cached freeze frame (if we have one)
@@ -738,6 +770,11 @@ namespace UE::PixelStreaming
 					return false;
 				});
 			}
+
+		if (UPixelStreamingDelegates* Delegates = UPixelStreamingDelegates::GetPixelStreamingDelegates())
+		{
+			Delegates->OnDataChannelClosedNative.Broadcast(StreamerId, PlayerId);
+		}
 		}
 	}
 
@@ -787,9 +824,9 @@ namespace UE::PixelStreaming
 			TArray<uint8> MessageData(RawBuffer.data.data(), RawBuffer.data.size());
 			OnInputReceived.Broadcast(PlayerId, Type, MessageData);
 
-			if (InputChannel)
+			if (InputHandler)
 			{
-				InputChannel->OnMessage(RawBuffer);
+				InputHandler->OnMessage(RawBuffer);
 			}
 		}
 	}
@@ -955,7 +992,7 @@ namespace UE::PixelStreaming
 				{
 					// bool QueryPeerStat(FPixelStreamingPlayerId PlayerId, FName StatToQuery, double& OutStatValue)
 					Stats->QueryPeerStat(PlayerId, PixelStreamingStatNames::MeanEncodeTime, EncodeMs);
-					Stats->QueryPeerStat(PlayerId, PixelStreamingStatNames::MeanSendDelay, CaptureToSendMs);
+					Stats->QueryPeerStat(PlayerId, PixelStreamingStatNames::AvgSendDelay, CaptureToSendMs);
 				}
 
 				double TransmissionTimeMs = FPlatformTime::ToMilliseconds64(FPlatformTime::Cycles64());
@@ -1053,5 +1090,21 @@ namespace UE::PixelStreaming
 				PlayerContext.DataChannel->SendMessage(Module.GetProtocol().FromStreamerProtocol.Find("QualityControlOwnership")->Id, IsController);
 			}
 		});
+		if (UPixelStreamingDelegates* Delegates = UPixelStreamingDelegates::GetPixelStreamingDelegates())
+		{
+			Delegates->OnQualityControllerChangedNative.Broadcast(StreamerId, QualityControllingId);
+		}
+	}
+
+	void FStreamer::TriggerMouseLeave(FString InStreamerId)
+	{
+		if (!IsEngineExitRequested() && StreamerId == InStreamerId)
+		{	
+			// Force a MouseLeave event. This prevents the PixelStreamingApplicationWrapper from
+			// still wrapping the base FSlateApplication after we stop streaming
+			TArray<uint8> EmptyArray;
+			TFunction<void(FMemoryReader)> MouseLeaveHandler = IPixelStreamingModule::Get().FindMessageHandler("MouseLeave");
+			MouseLeaveHandler(FMemoryReader(EmptyArray));
+		}
 	}
 } // namespace UE::PixelStreaming
diff --git a/Engine/Plugins/Media/PixelStreaming/Source/PixelStreaming/Private/Streamer.h b/Engine/Plugins/Media/PixelStreaming/Source/PixelStreaming/Private/Streamer.h
--- a/Engine/Plugins/Media/PixelStreaming/Source/PixelStreaming/Private/Streamer.h
+++ b/Engine/Plugins/Media/PixelStreaming/Source/PixelStreaming/Private/Streamer.h
@@ -8,7 +8,7 @@
 #include "VideoSourceGroup.h"
 #include "ThreadSafeMap.h"
 #include "Dom/JsonObject.h"
-#include "IPixelStreamingInputChannel.h"
+#include "IPixelStreamingInputHandler.h"
 #include "PixelStreamingSignallingConnection.h"
 #include "Templates/SharedPointer.h"
 
@@ -37,6 +37,7 @@ namespace UE::PixelStreaming
 
 		virtual void SetSignallingServerURL(const FString& InSignallingServerURL) override;
 		virtual FString GetSignallingServerURL() override;
+		virtual FString GetId() override { return StreamerId; };
 		virtual bool IsSignallingConnected() override;
 		virtual void StartStreaming() override;
 		virtual void StopStreaming() override;
@@ -55,8 +56,8 @@ namespace UE::PixelStreaming
 		virtual void SendFileData(const TArray64<uint8>& ByteData, FString& MimeType, FString& FileExtension) override;
 		virtual void KickPlayer(FPixelStreamingPlayerId PlayerId) override;
 
-		virtual void SetInputChannel(TSharedPtr<IPixelStreamingInputChannel> InInputChannel) override { InputChannel = InInputChannel; }
-		virtual TWeakPtr<IPixelStreamingInputChannel> GetInputChannel() override { return InputChannel; }
+		virtual void SetInputHandler(TSharedPtr<IPixelStreamingInputHandler> InInputHandler) override { InputHandler = InInputHandler; }
+		virtual TWeakPtr<IPixelStreamingInputHandler> GetInputHandler() override { return InputHandler; }
 
 		virtual IPixelStreamingAudioSink* GetPeerAudioSink(FPixelStreamingPlayerId PlayerId) override;
 		virtual IPixelStreamingAudioSink* GetUnlistenedAudioSink() override;
@@ -104,12 +105,13 @@ namespace UE::PixelStreaming
 		bool ShouldPeerGenerateFrames(FPixelStreamingPlayerId PlayerId) const;
 
 		void SetQualityController(FPixelStreamingPlayerId PlayerId);
+		void TriggerMouseLeave(FString InStreamerId);
 
 	private:
 		FString StreamerId;
 		FString CurrentSignallingServerURL;
 
-		TSharedPtr<IPixelStreamingInputChannel> InputChannel;
+		TSharedPtr<IPixelStreamingInputHandler> InputHandler;
 		TUniquePtr<FPixelStreamingSignallingConnection> SignallingServerConnection;
 		double LastSignallingServerConnectionAttemptTimestamp = 0;
 
@@ -143,6 +145,7 @@ namespace UE::PixelStreaming
 		TSharedPtr<FVideoSourceGroup> VideoSourceGroup;
 
 		FDelegateHandle ConsumeStatsHandle;
+		FDelegateHandle AllConnectionsClosedHandle;
 
 		IPixelStreamingModule& Module;
 	};
diff --git a/Engine/Plugins/Media/PixelStreaming/Source/PixelStreaming/Private/StreamerInputChannels.cpp b/Engine/Plugins/Media/PixelStreaming/Source/PixelStreaming/Private/StreamerInputChannels.cpp
deleted file mode 100644
diff --git a/Engine/Plugins/Media/PixelStreaming/Source/PixelStreaming/Private/StreamerInputChannels.h b/Engine/Plugins/Media/PixelStreaming/Source/PixelStreaming/Private/StreamerInputChannels.h
deleted file mode 100644
diff --git a/Engine/Plugins/Media/PixelStreaming/Source/PixelStreaming/Private/Tests/ProtocolTest.cpp b/Engine/Plugins/Media/PixelStreaming/Source/PixelStreaming/Private/Tests/ProtocolTest.cpp
--- a/Engine/Plugins/Media/PixelStreaming/Source/PixelStreaming/Private/Tests/ProtocolTest.cpp
+++ b/Engine/Plugins/Media/PixelStreaming/Source/PixelStreaming/Private/Tests/ProtocolTest.cpp
@@ -6,7 +6,6 @@
 #include "PixelStreamingVideoInputBackBuffer.h"
 #include "WebRTCIncludes.h"
 #include "IPixelStreamingModule.h"
-#include "StreamerInputChannels.h"
 #include "IPixelStreamingStreamer.h"
 
 #if WITH_DEV_AUTOMATION_TESTS
diff --git a/Engine/Plugins/Media/PixelStreaming/Source/PixelStreaming/Private/VideoSource.cpp b/Engine/Plugins/Media/PixelStreaming/Source/PixelStreaming/Private/VideoSource.cpp
--- a/Engine/Plugins/Media/PixelStreaming/Source/PixelStreaming/Private/VideoSource.cpp
+++ b/Engine/Plugins/Media/PixelStreaming/Source/PixelStreaming/Private/VideoSource.cpp
@@ -26,13 +26,12 @@ namespace UE::PixelStreaming
 		static int32 FrameId = 1;
 
 		CurrentState = webrtc::MediaSourceInterface::SourceState::kLive;
-		const int64 TimestampUs = rtc::TimeMicros();
 
 		rtc::scoped_refptr<webrtc::VideoFrameBuffer> FrameBuffer = VideoInput->GetFrameBuffer();
 		check(FrameBuffer->width() != 0 && FrameBuffer->height() != 0);
 		webrtc::VideoFrame Frame = webrtc::VideoFrame::Builder()
 									   .set_video_frame_buffer(FrameBuffer)
-									   .set_timestamp_us(TimestampUs)
+									   .set_timestamp_us(webrtc::Clock::GetRealTimeClock()->TimeInMicroseconds())
 									   .set_rotation(webrtc::VideoRotation::kVideoRotation_0)
 									   .set_id(FrameId++)
 									   .build();
diff --git a/Engine/Plugins/Media/PixelStreaming/Source/PixelStreaming/Public/IPixelStreamingInputChannel.h b/Engine/Plugins/Media/PixelStreaming/Source/PixelStreaming/Public/IPixelStreamingInputChannel.h
deleted file mode 100644
diff --git a/Engine/Plugins/Media/PixelStreaming/Source/PixelStreaming/Public/IPixelStreamingInputHandler.h b/Engine/Plugins/Media/PixelStreaming/Source/PixelStreaming/Public/IPixelStreamingInputHandler.h
new file mode 100644
--- /dev/null
+++ b/Engine/Plugins/Media/PixelStreaming/Source/PixelStreaming/Public/IPixelStreamingInputHandler.h
@@ -0,0 +1,87 @@
+// Copyright Epic Games, Inc. All Rights Reserved.
+
+#pragma once
+
+#include "GenericPlatform/GenericPlatformMisc.h"
+#include "IInputDevice.h"
+#include "Widgets/SViewport.h"
+#include "Widgets/SWindow.h"
+#include "Templates/SharedPointer.h"
+#include "PixelStreamingWebRTCIncludes.h"
+#include "Serialization/MemoryReader.h"
+
+/**
+ * IPixelStreamingInputHandler extends the IInputDevice interface. Setting the target viewport allows for
+ * scaling of input from browser to application, and setting the target window ensure that if windows are tiled (eg editor)
+ * that the streamed input only affect the target window.
+ */
+class PIXELSTREAMING_API IPixelStreamingInputHandler : public IInputDevice
+{
+public:
+	/**
+	 * @brief Handle the message from the WebRTC data channel.
+	 * @param Buffer The data channel message
+	 */
+	virtual void OnMessage(const webrtc::DataBuffer& Buffer) = 0;
+
+	/**
+	 * @brief Set the viewport this input device is associated with.
+	 * @param InTargetViewport The viewport to set
+	 */
+	virtual void SetTargetViewport(TWeakPtr<SViewport> InTargetViewport) = 0;
+
+	/**
+	 * @brief Get the viewport this input device is associated with.
+	 * @return The viewport this input device is associated with
+	 */
+	virtual TWeakPtr<SViewport> GetTargetViewport() = 0;
+
+	/**
+	 * @brief Set the viewport this input device is associated with.
+	 * @param InTargetWindow The viewport to set
+	 */
+	virtual void SetTargetWindow(TWeakPtr<SWindow> InTargetWindow) = 0;
+
+	/**
+	 * @brief Get the viewport this input device is associated with.
+	 * @return The viewport this input device is associated with
+	 */
+	virtual TWeakPtr<SWindow> GetTargetWindow() = 0;
+
+	/**
+	 * @brief Set the target screen size for this streamer. This is used to when the streamer doesn't have a singular target window / viewport 
+	 * and as such we just use the manual scale
+	 * @param InTargetWindow The target screen size
+	 */
+	virtual void SetTargetScreenSize(TWeakPtr<FIntPoint> InTargetScreenSize) = 0;
+
+	/**
+	 * @brief Get the target screen size for this streamer
+	 * @return The target screen size
+	 */
+	virtual TWeakPtr<FIntPoint> GetTargetScreenSize() = 0;
+
+	/**
+	 * @brief Set whether the input devices is faking touch events using keyboard and mouse this can be useful for debugging.
+	 * @return true
+	 * @return false
+	 */
+	virtual bool IsFakingTouchEvents() const = 0;
+
+	/**
+	 * @brief Register a function to be called whenever the specified message type is received.
+	 * 
+	 * @param MessageType The human readable identifier for the message
+	 * @param Handler The function called when this message type is received. This handler must take a single parameter (an FMemoryReader) and have a return type of void
+	 */
+	virtual void RegisterMessageHandler(const FString& MessageType, const TFunction<void(FMemoryReader)>& Handler) = 0;
+
+
+	/**
+	 * @brief Find the function to be called whenever the specified message type is received.
+	 * 
+	 * @param MessageType The human readable identifier for the message
+	 * @return TFunction<void(FMemoryReader)> The function called when this message type is received.
+	 */
+	virtual TFunction<void(FMemoryReader)> FindMessageHandler(const FString& MessageType) = 0;
+};
\ No newline at end of file
diff --git a/Engine/Plugins/Media/PixelStreaming/Source/PixelStreaming/Public/IPixelStreamingModule.h b/Engine/Plugins/Media/PixelStreaming/Source/PixelStreaming/Public/IPixelStreamingModule.h
--- a/Engine/Plugins/Media/PixelStreaming/Source/PixelStreaming/Public/IPixelStreamingModule.h
+++ b/Engine/Plugins/Media/PixelStreaming/Source/PixelStreaming/Public/IPixelStreamingModule.h
@@ -9,7 +9,7 @@
 #include "Templates/RefCounting.h"
 #include "PixelStreamingPlayerId.h"
 #include "IPixelStreamingAudioSink.h"
-#include "IPixelStreamingInputChannel.h"
+#include "IPixelStreamingInputHandler.h"
 #include "PixelStreamingWebRTCIncludes.h"
 #include "IPixelStreamingStreamer.h"
 #include "IInputDeviceModule.h"
@@ -123,6 +123,14 @@ public:
 	 */
 	virtual void RegisterMessage(Protocol::EPixelStreamingMessageDirection MessageDirection, const FString& MessageType, Protocol::FPixelStreamingInputMessage Message, const TFunction<void(FMemoryReader)>& Handler) = 0;
 
+	/**
+	 * @brief Find the function to be called whenever the specified message type is received.
+	 * 
+	 * @param MessageType The human readable identifier for the message
+	 * @return TFunction<void(FMemoryReader)> The function called when this message type is received.
+	 */
+	virtual TFunction<void(FMemoryReader)> FindMessageHandler(const FString& MessageType) = 0;
+
 	/**
 	 * Sets the target FPS for Externally Consumed video Tracks
 	 * @param InFPS new FPS for the ExternalVideoSource to output at. 
@@ -182,12 +190,6 @@ public:
 	 */
 	virtual void ForEachStreamer(const TFunction<void(TSharedPtr<IPixelStreamingStreamer>)>& Func) = 0;
 
-	/**
-	 * Register a lambda that returns a IInputDevice
-	 * @param InCreateInputeChannel - A lambda that will return input Channel
-	*/
-	virtual void RegisterCreateInputChannel(IPixelStreamingInputChannel::FCreateInputChannelFunc& InCreateInputChannel) = 0;
-
 	DECLARE_MULTICAST_DELEGATE(FOnProtocolUpdated);
 	FOnProtocolUpdated OnProtocolUpdated;
 };
diff --git a/Engine/Plugins/Media/PixelStreaming/Source/PixelStreaming/Public/IPixelStreamingStreamer.h b/Engine/Plugins/Media/PixelStreaming/Source/PixelStreaming/Public/IPixelStreamingStreamer.h
--- a/Engine/Plugins/Media/PixelStreaming/Source/PixelStreaming/Public/IPixelStreamingStreamer.h
+++ b/Engine/Plugins/Media/PixelStreaming/Source/PixelStreaming/Public/IPixelStreamingStreamer.h
@@ -10,8 +10,8 @@
 #include "Widgets/SWindow.h"
 #include "Templates/SharedPointer.h"
 #include "IPixelStreamingAudioSink.h"
-#include "IPixelStreamingInputChannel.h"
 #include "IPixelStreamingAudioInput.h"
+#include "IPixelStreamingInputHandler.h"
 
 class UTexture2D;
 
@@ -102,6 +102,12 @@ public:
 	 */
 	virtual FString GetSignallingServerURL() = 0;
 
+	/**
+	 * @brief Get this streamer's ID
+	 * 
+	 */
+	virtual FString GetId() = 0;
+
 	/**
 	 * @return True if the streamer is connected to the signalling mechanism (e.g. websocket for signalling server).
 	 */
@@ -192,16 +198,16 @@ public:
 	// virtual bool IsQualityController(FPixelStreamingPlayerId PlayerId) = 0;
 
 	/**
-	 * @brief Set the streamer's input device
-	 * @param InInputChannel
+	 * @brief Set the streamer's input handler
+	 * @param InInputHandler
 	 */
-	virtual void SetInputChannel(TSharedPtr<IPixelStreamingInputChannel> InInputChannel) = 0;
+	virtual void SetInputHandler(TSharedPtr<IPixelStreamingInputHandler> InInputHandler) = 0;
 
 	/**
-	 * @brief Get the streamer's input device
-	 * @return The streamer's input device
+	 * @brief Get the streamer's input handler
+	 * @return The streamer's input handler
 	 */
-	virtual TWeakPtr<IPixelStreamingInputChannel> GetInputChannel() = 0;
+	virtual TWeakPtr<IPixelStreamingInputHandler> GetInputHandler() = 0;
 
 	/**
 	 * @brief Get the audio sink associated with a specific peer/player.
diff --git a/Engine/Plugins/Media/PixelStreaming/Source/PixelStreaming/Public/PixelStreamingDelegates.h b/Engine/Plugins/Media/PixelStreaming/Source/PixelStreaming/Public/PixelStreamingDelegates.h
--- a/Engine/Plugins/Media/PixelStreaming/Source/PixelStreaming/Public/PixelStreamingDelegates.h
+++ b/Engine/Plugins/Media/PixelStreaming/Source/PixelStreaming/Public/PixelStreamingDelegates.h
@@ -3,6 +3,7 @@
 #pragma once
 
 #include "PixelStreamingPlayerId.h"
+#include "PixelStreamingDataChannel.h"
 #include "Containers/UnrealString.h"
 #include "Delegates/DelegateSignatureImpl.inl"
 #include "PixelStreamingDelegates.generated.h"
@@ -44,22 +45,22 @@ public:
 	 * A new connection has been made to the session.
 	 */
 	// BP Delegate
-	DECLARE_DYNAMIC_MULTICAST_DELEGATE_TwoParams(FNewConnection, FString, PlayerId, bool, QualityController);
+	DECLARE_DYNAMIC_MULTICAST_DELEGATE_ThreeParams(FNewConnection, FString, StreamerId, FString, PlayerId, bool, QualityController);
 	UPROPERTY(BlueprintAssignable, Category = "Pixel Streaming Delegates")
 	FNewConnection OnNewConnection;
 	// C++ Delegate
-	DECLARE_MULTICAST_DELEGATE_TwoParams(FNewConnectionNative, FString, bool);
+	DECLARE_MULTICAST_DELEGATE_ThreeParams(FNewConnectionNative, FString, FPixelStreamingPlayerId, bool);
 	FNewConnectionNative OnNewConnectionNative;
 
 	/**
 	 * A connection to a player was lost.
 	 */
 	// BP Delegate
-	DECLARE_DYNAMIC_MULTICAST_DELEGATE_TwoParams(FClosedConnection, FString, PlayerId, bool, WasQualityController);
+	DECLARE_DYNAMIC_MULTICAST_DELEGATE_ThreeParams(FClosedConnection, FString, StreamerId, FString, PlayerId, bool, WasQualityController);
 	UPROPERTY(BlueprintAssignable, Category = "Pixel Streaming Delegates")
 	FClosedConnection OnClosedConnection;
 	// C++ Delegate
-	DECLARE_MULTICAST_DELEGATE_TwoParams(FClosedConnectionNative, FPixelStreamingPlayerId, bool);
+	DECLARE_MULTICAST_DELEGATE_ThreeParams(FClosedConnectionNative, FString, FPixelStreamingPlayerId, bool);
 	FClosedConnectionNative OnClosedConnectionNative;
 
 	/**
@@ -67,32 +68,32 @@ public:
 	 * the app. This is an opportunity to reset the app.
 	 */
 	// BP Delegate
-	DECLARE_DYNAMIC_MULTICAST_DELEGATE(FAllConnectionsClosed);
+	DECLARE_DYNAMIC_MULTICAST_DELEGATE_OneParam(FAllConnectionsClosed, FString, StreamerId);
 	UPROPERTY(BlueprintAssignable, Category = "Pixel Streaming Delegates")
 	FAllConnectionsClosed OnAllConnectionsClosed;
 	// C++ Delegate
-	DECLARE_MULTICAST_DELEGATE(FAllConnectionsClosedNative);
+	DECLARE_MULTICAST_DELEGATE_OneParam(FAllConnectionsClosedNative, FString);
 	FAllConnectionsClosedNative OnAllConnectionsClosedNative;
 
 	/**
 	 * The quality controlling peer has changed
 	 */
 	// C++ Delegate
-	DECLARE_MULTICAST_DELEGATE_OneParam(FQualityControllerChanged, FPixelStreamingPlayerId);
+	DECLARE_MULTICAST_DELEGATE_TwoParams(FQualityControllerChanged, FString, FPixelStreamingPlayerId);
 	FQualityControllerChanged OnQualityControllerChangedNative;
 
 	/**
 	 * A new data channel has been opened
 	 */
 	// C++ Delegate
-	DECLARE_MULTICAST_DELEGATE_TwoParams(FDataChannelOpen, FPixelStreamingPlayerId, webrtc::DataChannelInterface*);
+	DECLARE_MULTICAST_DELEGATE_ThreeParams(FDataChannelOpen, FString, FPixelStreamingPlayerId, FPixelStreamingDataChannel*);
 	FDataChannelOpen OnDataChannelOpenNative;
 
 	/**
 	 * An existing data channel has been closed
 	 */
 	// C++ Delegate
-	DECLARE_MULTICAST_DELEGATE_TwoParams(FDataChannelClosed, FPixelStreamingPlayerId, webrtc::DataChannelInterface*);
+	DECLARE_MULTICAST_DELEGATE_TwoParams(FDataChannelClosed, FString, FPixelStreamingPlayerId);
 	FDataChannelClosed OnDataChannelClosedNative;
 
 	/**
diff --git a/Engine/Plugins/Media/PixelStreaming/Source/PixelStreaming/Public/PixelStreamingStatNames.h b/Engine/Plugins/Media/PixelStreaming/Source/PixelStreaming/Public/PixelStreamingStatNames.h
--- a/Engine/Plugins/Media/PixelStreaming/Source/PixelStreaming/Public/PixelStreamingStatNames.h
+++ b/Engine/Plugins/Media/PixelStreaming/Source/PixelStreaming/Public/PixelStreamingStatNames.h
@@ -41,6 +41,7 @@ namespace PixelStreamingStatNames
 	const FName FramesEncoded			  = FName(TEXT("framesEncoded"));
 	const FName QualityController		  = FName(TEXT("qualityController"));
 	const FName InputController			  = FName(TEXT("inputController"));
+	const FName AvgSendDelay			  = FName(TEXT("packetSendDelay"));
 
 	// Calculated stats
 	const FName FramesSentPerSecond	   = FName(TEXT("transmitFps"));
diff --git a/Engine/Plugins/Media/PixelStreaming/Source/PixelStreamingEditor/PixelStreamingEditor.build.cs b/Engine/Plugins/Media/PixelStreaming/Source/PixelStreamingEditor/PixelStreamingEditor.build.cs
--- a/Engine/Plugins/Media/PixelStreaming/Source/PixelStreamingEditor/PixelStreamingEditor.build.cs
+++ b/Engine/Plugins/Media/PixelStreaming/Source/PixelStreamingEditor/PixelStreamingEditor.build.cs
@@ -13,8 +13,11 @@ namespace UnrealBuildTool.Rules
 		{
 			var EngineDir = Path.GetFullPath(Target.RelativeEnginePath);
 
-			// This is so for game projects using our public headers don't have to include extra modules they might not know about.
-			PublicDependencyModuleNames.AddRange(new string[] {});
+			// NOTE: General rule is not to access the private folder of another module
+			PrivateIncludePaths.AddRange(new string[]
+			{
+				Path.Combine(EngineDir, "Plugins/Media/PixelStreaming/Source/PixelStreaming/Private")
+			});
 
 			PrivateDependencyModuleNames.AddRange(new string[]
 			{
@@ -29,7 +32,8 @@ namespace UnrealBuildTool.Rules
 				"Slate",
 				"SlateCore",
 				"EngineSettings",
-				"InputCore"
+				"InputCore",
+				"Json"
 			});
 
 			if(Target.bBuildEditor)
@@ -40,7 +44,8 @@ namespace UnrealBuildTool.Rules
 					"ToolMenus",
 					"EditorStyle",
 					"DesktopPlatform",
-					"LevelEditor"
+					"LevelEditor",
+					"MainFrame"
 				});
 			}
 		}
diff --git a/Engine/Plugins/Media/PixelStreaming/Source/PixelStreamingEditor/Private/PixelStreamingEditorModule.cpp b/Engine/Plugins/Media/PixelStreaming/Source/PixelStreamingEditor/Private/PixelStreamingEditorModule.cpp
--- a/Engine/Plugins/Media/PixelStreaming/Source/PixelStreamingEditor/Private/PixelStreamingEditorModule.cpp
+++ b/Engine/Plugins/Media/PixelStreaming/Source/PixelStreamingEditor/Private/PixelStreamingEditorModule.cpp
@@ -1,24 +1,17 @@
 // Copyright Epic Games, Inc. All Rights Reserved.
 
 #include "PixelStreamingEditorModule.h"
-#include "SLevelViewport.h"
-#include "Slate/SceneViewport.h"
-#include "LevelEditorViewport.h"
 #include "Editor/EditorPerformanceSettings.h"
 #include "PixelStreamingToolbar.h"
-#include "PixelStreamingVideoInputViewport.h"
 #include "IPixelStreamingModule.h"
 #include "PixelStreamingStyle.h"
 #include "IPixelStreamingStreamer.h"
-#include "Math/Vector4.h"
-#include "Math/Matrix.h"
-#include "Math/TransformNonVectorized.h"
-#include "Serialization/MemoryReader.h"
-#include "PixelStreamingPlayerId.h"
-#include "Kismet/KismetMathLibrary.h"
-#include "PixelStreamingVideoInputBackBuffer.h"
 #include "PixelStreamingVideoInputBackBufferComposited.h"
 #include "Settings.h"
+#include "UnrealEngine.h"
+#include "Interfaces/IMainFrameModule.h"
+#include "PixelStreamingProtocol.h"
+#include "Utils.h"
 
 UE::PixelStreaming::FPixelStreamingEditorModule* UE::PixelStreaming::FPixelStreamingEditorModule::PixelStreamingEditorModule = nullptr;
 
@@ -42,7 +35,7 @@ namespace UE::PixelStreaming
 		Settings::Editor::InitialiseSettings();
 		
 		IPixelStreamingModule& Module = IPixelStreamingModule::Get();
-		Module.OnReady().AddRaw(this, &FPixelStreamingEditorModule::InitEditorStreamer);
+		Module.OnReady().AddRaw(this, &FPixelStreamingEditorModule::InitEditorStreaming);
 	}
 
 	void FPixelStreamingEditorModule::ShutdownModule()
@@ -50,7 +43,7 @@ namespace UE::PixelStreaming
 		StopStreaming();
 	}
 
-	void FPixelStreamingEditorModule::InitEditorStreamer(IPixelStreamingModule& Module)
+	void FPixelStreamingEditorModule::InitEditorStreaming(IPixelStreamingModule& Module)
 	{
 		TSharedPtr<IPixelStreamingStreamer> Streamer = Module.GetStreamer(Module.GetDefaultStreamerID());
 		if (!Streamer.IsValid())
@@ -58,6 +51,49 @@ namespace UE::PixelStreaming
 			return;
 		}
 
+		// The current handler is the function that will currently be executed if a message with type "Command" is received
+		TFunction<void(FMemoryReader)> CurrentHandler = Module.FindMessageHandler("Command");
+		TFunction<void(FMemoryReader)> ExtendedHandler = [this, CurrentHandler](FMemoryReader Ar)
+		{
+			// We then create our new handler which will execute the "base" handler
+			CurrentHandler(Ar);
+			// and then perform out extended functionality after.
+			// equivalent to the super::DoSomeFunc pattern
+			FString Res;
+			Res.GetCharArray().SetNumUninitialized(Ar.TotalSize() / 2 + 1);
+			Ar.Serialize(Res.GetCharArray().GetData(), Ar.TotalSize());
+			FString Descriptor = Res.Mid(1);
+			FString WidthString;
+			FString HeightString;
+			bool bSuccess;
+			UE::PixelStreaming::ExtractJsonFromDescriptor(Descriptor, TEXT("Resolution.Width"), WidthString, bSuccess);
+			if (bSuccess)
+			{
+				UE::PixelStreaming::ExtractJsonFromDescriptor(Descriptor, TEXT("Resolution.Height"), HeightString, bSuccess);
+				if (bSuccess)
+				{
+					int Width = FCString::Atoi(*WidthString);
+					int Height = FCString::Atoi(*HeightString);
+					if (Width < 1 || Height < 1)
+					{
+						return;
+					}
+
+					TSharedPtr<SWindow> ParentWindow = IMainFrameModule::Get().GetParentWindow();
+					ParentWindow->Resize(FVector2D(Width, Height));
+					FSlateApplication::Get().OnSizeChanged(ParentWindow->GetNativeWindow().ToSharedRef(), Width, Height);
+				}
+			}
+		};
+		Module.RegisterMessage( Protocol::EPixelStreamingMessageDirection::ToStreamer, 
+								"Command", 
+								Protocol::FPixelStreamingInputMessage(51, 0, {}),
+								[ExtendedHandler](FMemoryReader Ar) {
+									ExtendedHandler(Ar);
+								});
+
+
+
 		Streamer->SetVideoInput(FPixelStreamingVideoInputBackBufferComposited::Create());
 		// Give the editor streamer the default url if the user hasn't specified one when launching the editor
 		if (Streamer->GetSignallingServerURL().IsEmpty())
@@ -65,10 +101,16 @@ namespace UE::PixelStreaming
 			Streamer->SetSignallingServerURL(Module.GetDefaultSignallingURL());
 		}
 
-		if(Settings::Editor::CVarEditorPixelStreamingStartOnLaunch.GetValueOnAnyThread())
+
+		IMainFrameModule::Get().OnMainFrameCreationFinished().AddLambda([&](TSharedPtr<SWindow> RootWindow, bool bIsNewProjectWindow)
 		{
-			StartStreaming();
-		}
+			MaybeResizeEditor(RootWindow);
+
+			if(Settings::Editor::CVarEditorPixelStreamingStartOnLaunch.GetValueOnAnyThread())
+			{
+				StartStreaming();
+			}
+		});
 	}
 
 	void FPixelStreamingEditorModule::StartStreaming()
@@ -112,6 +154,78 @@ namespace UE::PixelStreaming
 		}
 		return PixelStreamingEditorModule;
 	}
+
+	bool FPixelStreamingEditorModule::ParseResolution(const TCHAR* InResolution, uint32& OutX, uint32& OutY)
+	{
+		if(*InResolution)
+		{
+			FString CmdString(InResolution);
+			CmdString = CmdString.TrimStartAndEnd().ToLower();
+
+			// Retrieve the X dimensional value
+			const uint32 X = FMath::Max(FCString::Atof(*CmdString), 0.0f);
+
+			// Determine whether the user has entered a resolution and extract the Y dimension.
+			FString YString;
+
+			// Find separator between values (Example of expected format: 1280x768)
+			const TCHAR* YValue = NULL;
+			if(FCString::Strchr(*CmdString,'x'))
+			{
+				YValue = const_cast<TCHAR*> (FCString::Strchr(*CmdString,'x')+1);
+				YString = YValue;
+				// Remove any whitespace from the end of the string
+				YString = YString.TrimStartAndEnd();
+			}
+
+			// If the Y dimensional value exists then setup to use the specified resolution.
+			uint32 Y = 0;
+			if ( YValue && YString.Len() > 0 )
+			{
+				if (YString.IsNumeric())
+				{
+					Y = FMath::Max(FCString::Atof(YValue), 0.0f);
+					OutX = X;
+					OutY = Y;
+					return true;
+				}
+			}
+		}
+		return false;
+	}
+
+	void FPixelStreamingEditorModule::MaybeResizeEditor(TSharedPtr<SWindow> RootWindow)
+	{
+		uint32 ResolutionX, ResolutionY = 0;
+		FString ResolutionStr;
+		bool bSuccess = FParse::Value(FCommandLine::Get(), TEXT("EditorPixelStreamingRes="), ResolutionStr);
+		if (bSuccess)
+		{
+			bSuccess = ParseResolution(*ResolutionStr, ResolutionX, ResolutionY);
+		}
+		else
+		{
+			bool UserSpecifiedWidth = FParse::Value(FCommandLine::Get(), TEXT("EditorPixelStreamingResX="), ResolutionX);
+			bool UserSpecifiedHeight = FParse::Value(FCommandLine::Get(), TEXT("EditorPixelStreamingResY="), ResolutionY);
+			bSuccess = UserSpecifiedWidth | UserSpecifiedHeight;
+
+			const float AspectRatio = 16.0 / 9.0;
+			if (UserSpecifiedWidth && !UserSpecifiedHeight)
+			{
+				ResolutionY = int32(ResolutionX / AspectRatio);
+			}
+			else if (UserSpecifiedHeight && !UserSpecifiedWidth)
+			{
+				ResolutionX = int32(ResolutionY * AspectRatio);
+			}
+		}
+
+		if(bSuccess)
+		{
+			RootWindow->Resize(FVector2D(ResolutionX, ResolutionY));
+			FSlateApplication::Get().OnSizeChanged(RootWindow->GetNativeWindow().ToSharedRef(), ResolutionX, ResolutionY);
+		}
+	}
 } // namespace UE::PixelStreaming
 
 IMPLEMENT_MODULE(UE::PixelStreaming::FPixelStreamingEditorModule, PixelStreamingEditor)
\ No newline at end of file
diff --git a/Engine/Plugins/Media/PixelStreaming/Source/PixelStreamingEditor/Public/PixelStreamingEditorModule.h b/Engine/Plugins/Media/PixelStreaming/Source/PixelStreamingEditor/Public/PixelStreamingEditorModule.h
--- a/Engine/Plugins/Media/PixelStreaming/Source/PixelStreamingEditor/Public/PixelStreamingEditorModule.h
+++ b/Engine/Plugins/Media/PixelStreaming/Source/PixelStreamingEditor/Public/PixelStreamingEditorModule.h
@@ -18,7 +18,9 @@ namespace UE::PixelStreaming
         
         static FPixelStreamingEditorModule* GetModule();
     private:
-        void InitEditorStreamer(IPixelStreamingModule& Module);
+        void InitEditorStreaming(IPixelStreamingModule& Module);
+        bool ParseResolution(const TCHAR* InResolution, uint32& OutX, uint32& OutY);
+        void MaybeResizeEditor(TSharedPtr<SWindow> RootWindow);
 
         TSharedPtr<class FPixelStreamingToolbar> Toolbar;
 
