Cherrypicked 21558229 from Dev-TensorWorks:

QOL: Move the Editor specific PixelStreamingVideoInputs to the appropriate PixelStreamingEditor module. Leaky namespaces have also been fixed

diff --git a/Engine/Plugins/Media/PixelStreaming/Source/PixelStreaming/Private/PixelStreamingVideoInputBackBufferComposited.cpp b/Engine/Plugins/Media/PixelStreaming/Source/PixelStreaming/Private/PixelStreamingVideoInputBackBufferComposited.cpp
deleted file mode 100644
diff --git a/Engine/Plugins/Media/PixelStreaming/Source/PixelStreaming/Private/PixelStreamingVideoInputViewport.cpp b/Engine/Plugins/Media/PixelStreaming/Source/PixelStreaming/Private/PixelStreamingVideoInputViewport.cpp
deleted file mode 100644
diff --git a/Engine/Plugins/Media/PixelStreaming/Source/PixelStreaming/Private/Settings.cpp b/Engine/Plugins/Media/PixelStreaming/Source/PixelStreaming/Private/Settings.cpp
--- a/Engine/Plugins/Media/PixelStreaming/Source/PixelStreaming/Private/Settings.cpp
+++ b/Engine/Plugins/Media/PixelStreaming/Source/PixelStreaming/Private/Settings.cpp
@@ -9,43 +9,7 @@
 #include "Async/Async.h"
 #include "IPixelStreamingModule.h"
 #include "PixelStreamingDelegates.h"
-
-template <typename T>
-void PSCommandLineParseValue(const TCHAR* Match, TAutoConsoleVariable<T>& CVar)
-{
-	T Value;
-	if (FParse::Value(FCommandLine::Get(), Match, Value))
-		CVar->Set(Value, ECVF_SetByCommandline);
-};
-
-void PSCommandLineParseValue(const TCHAR* Match, TAutoConsoleVariable<FString>& CVar, bool bStopOnSeparator = false)
-{
-	FString Value;
-	if (FParse::Value(FCommandLine::Get(), Match, Value, bStopOnSeparator))
-		CVar->Set(*Value, ECVF_SetByCommandline);
-};
-
-void PSCommandLineParseOption(const TCHAR* Match, TAutoConsoleVariable<bool>& CVar)
-{
-	FString ValueMatch(Match);
-	ValueMatch.Append(TEXT("="));
-	FString Value;
-	if (FParse::Value(FCommandLine::Get(), *ValueMatch, Value))
-	{
-		if (Value.Equals(FString(TEXT("true")), ESearchCase::IgnoreCase))
-		{
-			CVar->Set(true, ECVF_SetByCommandline);
-		}
-		else if (Value.Equals(FString(TEXT("false")), ESearchCase::IgnoreCase))
-		{
-			CVar->Set(false, ECVF_SetByCommandline);
-		}
-	}
-	else if (FParse::Param(FCommandLine::Get(), Match))
-	{
-		CVar->Set(true, ECVF_SetByCommandline);
-	}
-}
+#include "Utils.h"
 
 namespace UE::PixelStreaming::Settings
 {
@@ -588,11 +552,12 @@ namespace UE::PixelStreaming::Settings
 	void OnStreamerReady(IPixelStreamingModule& Module)
 	{
 		CVarPixelStreamingLogStats.AsVariable()->SetOnChangedCallback(FConsoleVariableDelegate::CreateStatic(&OnLogStatsChanged));
-		PSCommandLineParseOption(TEXT("PixelStreamingLogStats"), CVarPixelStreamingLogStats);
+		CommandLineParseOption(TEXT("PixelStreamingLogStats"), CVarPixelStreamingLogStats);
 	}
 
 	void InitialiseSettings()
 	{
+		using namespace UE::PixelStreaming;
 		UE_LOG(LogPixelStreaming, Log, TEXT("Initialising Pixel Streaming settings."));
 
 		FString StringOptions;
@@ -609,44 +574,44 @@ namespace UE::PixelStreaming::Settings
 		CVarPixelStreamingEncoderKeyframeInterval.AsVariable()->SetOnChangedCallback(FConsoleVariableDelegate::CreateStatic(&OnKeyframeIntervalChanged));
 
 		// Values parse from commands line
-		PSCommandLineParseValue(TEXT("PixelStreamingEncoderKeyframeInterval="), CVarPixelStreamingEncoderKeyframeInterval);
-		PSCommandLineParseValue(TEXT("PixelStreamingEncoderTargetBitrate="), CVarPixelStreamingEncoderTargetBitrate);
-		PSCommandLineParseValue(TEXT("PixelStreamingEncoderMaxBitrate="), CVarPixelStreamingEncoderMaxBitrate);
-		PSCommandLineParseValue(TEXT("PixelStreamingEncoderMinQP="), CVarPixelStreamingEncoderMinQP);
-		PSCommandLineParseValue(TEXT("PixelStreamingEncoderMaxQP="), CVarPixelStreamingEncoderMaxQP);
-		PSCommandLineParseValue(TEXT("PixelStreamingEncoderRateControl="), CVarPixelStreamingEncoderRateControl);
-		PSCommandLineParseValue(TEXT("PixelStreamingEncoderMultipass="), CVarPixelStreamingEncoderMultipass);
-		PSCommandLineParseValue(TEXT("PixelStreamingEncoderCodec="), CVarPixelStreamingEncoderCodec);
-		PSCommandLineParseValue(TEXT("PixelStreamingH264Profile="), CVarPixelStreamingH264Profile);
-		PSCommandLineParseValue(TEXT("PixelStreamingDegradationPreference="), CVarPixelStreamingDegradationPreference);
-		PSCommandLineParseValue(TEXT("PixelStreamingWebRTCDegradationPreference="), CVarPixelStreamingDegradationPreference);
-		PSCommandLineParseValue(TEXT("PixelStreamingWebRTCFps="), CVarPixelStreamingWebRTCFps);
-		PSCommandLineParseValue(TEXT("PixelStreamingWebRTCStartBitrate="), CVarPixelStreamingWebRTCStartBitrate);
-		PSCommandLineParseValue(TEXT("PixelStreamingWebRTCMinBitrate="), CVarPixelStreamingWebRTCMinBitrate);
-		PSCommandLineParseValue(TEXT("PixelStreamingWebRTCMaxBitrate="), CVarPixelStreamingWebRTCMaxBitrate);
-		PSCommandLineParseValue(TEXT("PixelStreamingWebRTCLowQpThreshold="), CVarPixelStreamingWebRTCLowQpThreshold);
-		PSCommandLineParseValue(TEXT("PixelStreamingWebRTCHighQpThreshold="), CVarPixelStreamingWebRTCHighQpThreshold);
-		PSCommandLineParseValue(TEXT("PixelStreamingFreezeFrameQuality"), CVarPixelStreamingFreezeFrameQuality);
-		PSCommandLineParseValue(TEXT("PixelStreamingKeyFilter="), CVarPixelStreamingKeyFilter);
-		PSCommandLineParseValue(TEXT("PixelStreamingVideoTracks="), CVarPixelStreamingVideoTracks);
-		PSCommandLineParseValue(TEXT("PixelStreamingFrameScale="), CVarPixelStreamingFrameScale);
-		PSCommandLineParseValue(TEXT("PixelStreamingInputController="), CVarPixelStreamingInputController);
-		PSCommandLineParseValue(TEXT("PixelStreamingSignalingReconnectInterval="), CVarPixelStreamingSignalingReconnectInterval);
+		CommandLineParseValue(TEXT("PixelStreamingEncoderKeyframeInterval="), CVarPixelStreamingEncoderKeyframeInterval);
+		CommandLineParseValue(TEXT("PixelStreamingEncoderTargetBitrate="), CVarPixelStreamingEncoderTargetBitrate);
+		CommandLineParseValue(TEXT("PixelStreamingEncoderMaxBitrate="), CVarPixelStreamingEncoderMaxBitrate);
+		CommandLineParseValue(TEXT("PixelStreamingEncoderMinQP="), CVarPixelStreamingEncoderMinQP);
+		CommandLineParseValue(TEXT("PixelStreamingEncoderMaxQP="), CVarPixelStreamingEncoderMaxQP);
+		CommandLineParseValue(TEXT("PixelStreamingEncoderRateControl="), CVarPixelStreamingEncoderRateControl);
+		CommandLineParseValue(TEXT("PixelStreamingEncoderMultipass="), CVarPixelStreamingEncoderMultipass);
+		CommandLineParseValue(TEXT("PixelStreamingEncoderCodec="), CVarPixelStreamingEncoderCodec);
+		CommandLineParseValue(TEXT("PixelStreamingH264Profile="), CVarPixelStreamingH264Profile);
+		CommandLineParseValue(TEXT("PixelStreamingDegradationPreference="), CVarPixelStreamingDegradationPreference);
+		CommandLineParseValue(TEXT("PixelStreamingWebRTCDegradationPreference="), CVarPixelStreamingDegradationPreference);
+		CommandLineParseValue(TEXT("PixelStreamingWebRTCFps="), CVarPixelStreamingWebRTCFps);
+		CommandLineParseValue(TEXT("PixelStreamingWebRTCStartBitrate="), CVarPixelStreamingWebRTCStartBitrate);
+		CommandLineParseValue(TEXT("PixelStreamingWebRTCMinBitrate="), CVarPixelStreamingWebRTCMinBitrate);
+		CommandLineParseValue(TEXT("PixelStreamingWebRTCMaxBitrate="), CVarPixelStreamingWebRTCMaxBitrate);
+		CommandLineParseValue(TEXT("PixelStreamingWebRTCLowQpThreshold="), CVarPixelStreamingWebRTCLowQpThreshold);
+		CommandLineParseValue(TEXT("PixelStreamingWebRTCHighQpThreshold="), CVarPixelStreamingWebRTCHighQpThreshold);
+		CommandLineParseValue(TEXT("PixelStreamingFreezeFrameQuality"), CVarPixelStreamingFreezeFrameQuality);
+		CommandLineParseValue(TEXT("PixelStreamingKeyFilter="), CVarPixelStreamingKeyFilter);
+		CommandLineParseValue(TEXT("PixelStreamingVideoTracks="), CVarPixelStreamingVideoTracks);
+		CommandLineParseValue(TEXT("PixelStreamingFrameScale="), CVarPixelStreamingFrameScale);
+		CommandLineParseValue(TEXT("PixelStreamingInputController="), CVarPixelStreamingInputController);
+		CommandLineParseValue(TEXT("PixelStreamingSignalingReconnectInterval="), CVarPixelStreamingSignalingReconnectInterval);
 
 		// Options parse (if these exist they are set to true)
-		PSCommandLineParseOption(TEXT("AllowPixelStreamingCommands"), CVarPixelStreamingAllowConsoleCommands);
-		PSCommandLineParseOption(TEXT("PixelStreamingOnScreenStats"), CVarPixelStreamingOnScreenStats);
-		PSCommandLineParseOption(TEXT("PixelStreamingHudStats"), CVarPixelStreamingOnScreenStats);
-		PSCommandLineParseOption(TEXT("PixelStreamingDebugDumpFrame"), CVarPixelStreamingDebugDumpFrame);
-		PSCommandLineParseOption(TEXT("PixelStreamingEnableFillerData"), CVarPixelStreamingEnableFillerData);
-		PSCommandLineParseOption(TEXT("PixelStreamingWebRTCDisableStats"), CVarPixelStreamingWebRTCDisableStats);
-		PSCommandLineParseOption(TEXT("PixelStreamingWebRTCDisableReceiveAudio"), CVarPixelStreamingWebRTCDisableReceiveAudio);
-		PSCommandLineParseOption(TEXT("PixelStreamingWebRTCDisableTransmitAudio"), CVarPixelStreamingWebRTCDisableTransmitAudio);
-		PSCommandLineParseOption(TEXT("PixelStreamingWebRTCDisableAudioSync"), CVarPixelStreamingWebRTCDisableAudioSync);
-		PSCommandLineParseOption(TEXT("PixelStreamingSendPlayerIdAsInteger"), CVarSendPlayerIdAsInteger);
-		PSCommandLineParseOption(TEXT("PixelStreamingWebRTCUseLegacyAudioDevice"), CVarPixelStreamingWebRTCUseLegacyAudioDevice);
-		PSCommandLineParseOption(TEXT("PixelStreamingDisableLatencyTester"), CVarPixelStreamingDisableLatencyTester);
-		PSCommandLineParseOption(TEXT("PixelStreamingVPXUseCompute"), CVarPixelStreamingVPXUseCompute);
+		CommandLineParseOption(TEXT("AllowPixelStreamingCommands"), CVarPixelStreamingAllowConsoleCommands);
+		CommandLineParseOption(TEXT("PixelStreamingOnScreenStats"), CVarPixelStreamingOnScreenStats);
+		CommandLineParseOption(TEXT("PixelStreamingHudStats"), CVarPixelStreamingOnScreenStats);
+		CommandLineParseOption(TEXT("PixelStreamingDebugDumpFrame"), CVarPixelStreamingDebugDumpFrame);
+		CommandLineParseOption(TEXT("PixelStreamingEnableFillerData"), CVarPixelStreamingEnableFillerData);
+		CommandLineParseOption(TEXT("PixelStreamingWebRTCDisableStats"), CVarPixelStreamingWebRTCDisableStats);
+		CommandLineParseOption(TEXT("PixelStreamingWebRTCDisableReceiveAudio"), CVarPixelStreamingWebRTCDisableReceiveAudio);
+		CommandLineParseOption(TEXT("PixelStreamingWebRTCDisableTransmitAudio"), CVarPixelStreamingWebRTCDisableTransmitAudio);
+		CommandLineParseOption(TEXT("PixelStreamingWebRTCDisableAudioSync"), CVarPixelStreamingWebRTCDisableAudioSync);
+		CommandLineParseOption(TEXT("PixelStreamingSendPlayerIdAsInteger"), CVarSendPlayerIdAsInteger);
+		CommandLineParseOption(TEXT("PixelStreamingWebRTCUseLegacyAudioDevice"), CVarPixelStreamingWebRTCUseLegacyAudioDevice);
+		CommandLineParseOption(TEXT("PixelStreamingDisableLatencyTester"), CVarPixelStreamingDisableLatencyTester);
+		CommandLineParseOption(TEXT("PixelStreamingVPXUseCompute"), CVarPixelStreamingVPXUseCompute);
 
 		ReadSimulcastParameters();
 
diff --git a/Engine/Plugins/Media/PixelStreaming/Source/PixelStreaming/Private/Utils.h b/Engine/Plugins/Media/PixelStreaming/Source/PixelStreaming/Private/Utils.h
--- a/Engine/Plugins/Media/PixelStreaming/Source/PixelStreaming/Private/Utils.h
+++ b/Engine/Plugins/Media/PixelStreaming/Source/PixelStreaming/Private/Utils.h
@@ -9,9 +9,49 @@
 #include "Serialization/JsonWriter.h"
 #include "Serialization/JsonSerializer.h"
 #include "Policies/CondensedJsonPrintPolicy.h"
+#include "CoreMinimal.h"
+#include "HAL/IConsoleManager.h"
+#include "Misc/CommandLine.h"
 
 namespace UE::PixelStreaming
 {
+	template <typename T>
+	inline void CommandLineParseValue(const TCHAR* Match, TAutoConsoleVariable<T>& CVar)
+	{
+		T Value;
+		if (FParse::Value(FCommandLine::Get(), Match, Value))
+			CVar->Set(Value, ECVF_SetByCommandline);
+	};
+
+	inline void CommandLineParseValue(const TCHAR* Match, TAutoConsoleVariable<FString>& CVar, bool bStopOnSeparator = false)
+	{
+		FString Value;
+		if (FParse::Value(FCommandLine::Get(), Match, Value, bStopOnSeparator))
+			CVar->Set(*Value, ECVF_SetByCommandline);
+	};
+
+	inline void CommandLineParseOption(const TCHAR* Match, TAutoConsoleVariable<bool>& CVar)
+	{
+		FString ValueMatch(Match);
+		ValueMatch.Append(TEXT("="));
+		FString Value;
+		if (FParse::Value(FCommandLine::Get(), *ValueMatch, Value))
+		{
+			if (Value.Equals(FString(TEXT("true")), ESearchCase::IgnoreCase))
+			{
+				CVar->Set(true, ECVF_SetByCommandline);
+			}
+			else if (Value.Equals(FString(TEXT("false")), ESearchCase::IgnoreCase))
+			{
+				CVar->Set(false, ECVF_SetByCommandline);
+			}
+		}
+		else if (FParse::Param(FCommandLine::Get(), Match))
+		{
+			CVar->Set(true, ECVF_SetByCommandline);
+		}
+	}
+
 	template <typename T>
 	void DoOnGameThread(T&& Func)
 	{
diff --git a/Engine/Plugins/Media/PixelStreaming/Source/PixelStreaming/Public/PixelStreamingVideoInputBackBufferComposited.h b/Engine/Plugins/Media/PixelStreaming/Source/PixelStreaming/Public/PixelStreamingVideoInputBackBufferComposited.h
deleted file mode 100644
diff --git a/Engine/Plugins/Media/PixelStreaming/Source/PixelStreaming/Public/PixelStreamingVideoInputViewport.h b/Engine/Plugins/Media/PixelStreaming/Source/PixelStreaming/Public/PixelStreamingVideoInputViewport.h
deleted file mode 100644
diff --git a/Engine/Plugins/Media/PixelStreaming/Source/PixelStreamingEditor/PixelStreamingEditor.build.cs b/Engine/Plugins/Media/PixelStreaming/Source/PixelStreamingEditor/PixelStreamingEditor.build.cs
--- a/Engine/Plugins/Media/PixelStreaming/Source/PixelStreamingEditor/PixelStreamingEditor.build.cs
+++ b/Engine/Plugins/Media/PixelStreaming/Source/PixelStreamingEditor/PixelStreamingEditor.build.cs
@@ -33,7 +33,8 @@ namespace UnrealBuildTool.Rules
 				"SlateCore",
 				"EngineSettings",
 				"InputCore",
-				"Json"
+				"Json",
+				"PixelCapture"
 			});
 
 			if(Target.bBuildEditor)
diff --git a/Engine/Plugins/Media/PixelStreaming/Source/PixelStreamingEditor/Private/PixelStreamingCommands.cpp b/Engine/Plugins/Media/PixelStreaming/Source/PixelStreamingEditor/Private/PixelStreamingCommands.cpp
--- a/Engine/Plugins/Media/PixelStreaming/Source/PixelStreamingEditor/Private/PixelStreamingCommands.cpp
+++ b/Engine/Plugins/Media/PixelStreaming/Source/PixelStreamingEditor/Private/PixelStreamingCommands.cpp
@@ -2,7 +2,7 @@
 
 #include "PixelStreamingCommands.h"
 
-namespace UE::PixelStreaming
+namespace UE::EditorPixelStreaming
 {
 #define LOCTEXT_NAMESPACE "PixelStreamingToolBar"
     void FPixelStreamingCommands::RegisterCommands()
diff --git a/Engine/Plugins/Media/PixelStreaming/Source/PixelStreamingEditor/Private/PixelStreamingCommands.h b/Engine/Plugins/Media/PixelStreaming/Source/PixelStreamingEditor/Private/PixelStreamingCommands.h
--- a/Engine/Plugins/Media/PixelStreaming/Source/PixelStreamingEditor/Private/PixelStreamingCommands.h
+++ b/Engine/Plugins/Media/PixelStreaming/Source/PixelStreamingEditor/Private/PixelStreamingCommands.h
@@ -9,7 +9,7 @@
  * These are the commands used in the toolbar visible in the editor
  * 
  */
-namespace UE::PixelStreaming
+namespace UE::EditorPixelStreaming
 {
     class FPixelStreamingCommands : public TCommands<FPixelStreamingCommands>
     {
diff --git a/Engine/Plugins/Media/PixelStreaming/Source/PixelStreamingEditor/Private/PixelStreamingEditorModule.cpp b/Engine/Plugins/Media/PixelStreaming/Source/PixelStreamingEditor/Private/PixelStreamingEditorModule.cpp
--- a/Engine/Plugins/Media/PixelStreaming/Source/PixelStreamingEditor/Private/PixelStreamingEditorModule.cpp
+++ b/Engine/Plugins/Media/PixelStreaming/Source/PixelStreamingEditor/Private/PixelStreamingEditorModule.cpp
@@ -13,219 +13,214 @@
 #include "PixelStreamingProtocol.h"
 #include "Utils.h"
 
-UE::PixelStreaming::FPixelStreamingEditorModule* UE::PixelStreaming::FPixelStreamingEditorModule::PixelStreamingEditorModule = nullptr;
+FPixelStreamingEditorModule* FPixelStreamingEditorModule::PixelStreamingEditorModule = nullptr;
 
-namespace UE::PixelStreaming
+/**
+ * IModuleInterface implementation
+ */
+void FPixelStreamingEditorModule::StartupModule()
 {
-	/**
-	 * IModuleInterface implementation
-	 */
-	void FPixelStreamingEditorModule::StartupModule()
-	{
-		// Initialize the editor toolbar
-		FPixelStreamingStyle::Initialize();
-		FPixelStreamingStyle::ReloadTextures();
-		Toolbar = MakeShared<FPixelStreamingToolbar>();
-
-		// Update editor settings so that editor won't slow down if not in focus
-		UEditorPerformanceSettings* Settings = GetMutableDefault<UEditorPerformanceSettings>();
-		Settings->bThrottleCPUWhenNotForeground = false;
-		Settings->PostEditChange();
-
-		Settings::Editor::InitialiseSettings();
+	using namespace UE::EditorPixelStreaming;
+	// Initialize the editor toolbar
+	FPixelStreamingStyle::Initialize();
+	FPixelStreamingStyle::ReloadTextures();
+	Toolbar = MakeShared<FPixelStreamingToolbar>();
+
+	// Update editor settings so that editor won't slow down if not in focus
+	UEditorPerformanceSettings* Settings = GetMutableDefault<UEditorPerformanceSettings>();
+	Settings->bThrottleCPUWhenNotForeground = false;
+	Settings->PostEditChange();
+
+	Settings::InitialiseSettings();
 		
-		IPixelStreamingModule& Module = IPixelStreamingModule::Get();
-		Module.OnReady().AddRaw(this, &FPixelStreamingEditorModule::InitEditorStreaming);
-	}
+	IPixelStreamingModule& Module = IPixelStreamingModule::Get();
+	Module.OnReady().AddRaw(this, &FPixelStreamingEditorModule::InitEditorStreaming);
+}
 
-	void FPixelStreamingEditorModule::ShutdownModule()
+void FPixelStreamingEditorModule::ShutdownModule()
+{
+	StopStreaming();
+}
+
+void FPixelStreamingEditorModule::InitEditorStreaming(IPixelStreamingModule& Module)
+{
+	TSharedPtr<IPixelStreamingStreamer> Streamer = Module.GetStreamer(Module.GetDefaultStreamerID());
+	if (!Streamer.IsValid())
 	{
-		StopStreaming();
+		return;
 	}
 
-	void FPixelStreamingEditorModule::InitEditorStreaming(IPixelStreamingModule& Module)
+	// The current handler is the function that will currently be executed if a message with type "Command" is received
+	TFunction<void(FMemoryReader)> CurrentHandler = Module.FindMessageHandler("Command");
+	TFunction<void(FMemoryReader)> ExtendedHandler = [CurrentHandler](FMemoryReader Ar)
 	{
-		TSharedPtr<IPixelStreamingStreamer> Streamer = Module.GetStreamer(Module.GetDefaultStreamerID());
-		if (!Streamer.IsValid())
-		{
-			return;
-		}
-
-		// The current handler is the function that will currently be executed if a message with type "Command" is received
-		TFunction<void(FMemoryReader)> CurrentHandler = Module.FindMessageHandler("Command");
-		TFunction<void(FMemoryReader)> ExtendedHandler = [this, CurrentHandler](FMemoryReader Ar)
+		// We then create our new handler which will execute the "base" handler
+		CurrentHandler(Ar);
+		// and then perform out extended functionality after.
+		// equivalent to the super::DoSomeFunc pattern
+		FString Res;
+		Res.GetCharArray().SetNumUninitialized(Ar.TotalSize() / 2 + 1);
+		Ar.Serialize(Res.GetCharArray().GetData(), Ar.TotalSize());
+		FString Descriptor = Res.Mid(1);
+		FString WidthString;
+		FString HeightString;
+		bool bSuccess;
+		UE::PixelStreaming::ExtractJsonFromDescriptor(Descriptor, TEXT("Resolution.Width"), WidthString, bSuccess);
+		if (bSuccess)
 		{
-			// We then create our new handler which will execute the "base" handler
-			CurrentHandler(Ar);
-			// and then perform out extended functionality after.
-			// equivalent to the super::DoSomeFunc pattern
-			FString Res;
-			Res.GetCharArray().SetNumUninitialized(Ar.TotalSize() / 2 + 1);
-			Ar.Serialize(Res.GetCharArray().GetData(), Ar.TotalSize());
-			FString Descriptor = Res.Mid(1);
-			FString WidthString;
-			FString HeightString;
-			bool bSuccess;
-			UE::PixelStreaming::ExtractJsonFromDescriptor(Descriptor, TEXT("Resolution.Width"), WidthString, bSuccess);
+			UE::PixelStreaming::ExtractJsonFromDescriptor(Descriptor, TEXT("Resolution.Height"), HeightString, bSuccess);
 			if (bSuccess)
 			{
-				UE::PixelStreaming::ExtractJsonFromDescriptor(Descriptor, TEXT("Resolution.Height"), HeightString, bSuccess);
-				if (bSuccess)
+				int Width = FCString::Atoi(*WidthString);
+				int Height = FCString::Atoi(*HeightString);
+				if (Width < 1 || Height < 1)
 				{
-					int Width = FCString::Atoi(*WidthString);
-					int Height = FCString::Atoi(*HeightString);
-					if (Width < 1 || Height < 1)
-					{
-						return;
-					}
-
-					TSharedPtr<SWindow> ParentWindow = IMainFrameModule::Get().GetParentWindow();
-					ParentWindow->Resize(FVector2D(Width, Height));
-					FSlateApplication::Get().OnSizeChanged(ParentWindow->GetNativeWindow().ToSharedRef(), Width, Height);
+					return;
 				}
-			}
-		};
-		Module.RegisterMessage( Protocol::EPixelStreamingMessageDirection::ToStreamer, 
-								"Command", 
-								Protocol::FPixelStreamingInputMessage(51, 0, {}),
-								[ExtendedHandler](FMemoryReader Ar) {
-									ExtendedHandler(Ar);
-								});
 
-
-
-		Streamer->SetVideoInput(FPixelStreamingVideoInputBackBufferComposited::Create());
-		// Give the editor streamer the default url if the user hasn't specified one when launching the editor
-		if (Streamer->GetSignallingServerURL().IsEmpty())
-		{
-			Streamer->SetSignallingServerURL(Module.GetDefaultSignallingURL());
-		}
-
-
-		IMainFrameModule::Get().OnMainFrameCreationFinished().AddLambda([&](TSharedPtr<SWindow> RootWindow, bool bIsNewProjectWindow)
-		{
-			MaybeResizeEditor(RootWindow);
-
-			if(Settings::Editor::CVarEditorPixelStreamingStartOnLaunch.GetValueOnAnyThread())
-			{
-				StartStreaming();
+				TSharedPtr<SWindow> ParentWindow = IMainFrameModule::Get().GetParentWindow();
+				ParentWindow->Resize(FVector2D(Width, Height));
+				FSlateApplication::Get().OnSizeChanged(ParentWindow->GetNativeWindow().ToSharedRef(), Width, Height);
 			}
-		});
+		}
+	};
+	Module.RegisterMessage( Protocol::EPixelStreamingMessageDirection::ToStreamer, 
+							"Command", 
+							Protocol::FPixelStreamingInputMessage(51, 0, {}),
+							[ExtendedHandler](FMemoryReader Ar) {
+								ExtendedHandler(Ar);
+							});
+
+	Streamer->SetVideoInput(FPixelStreamingVideoInputBackBufferComposited::Create());
+	// Give the editor streamer the default url if the user hasn't specified one when launching the editor
+	if (Streamer->GetSignallingServerURL().IsEmpty())
+	{
+		Streamer->SetSignallingServerURL(Module.GetDefaultSignallingURL());
 	}
 
-	void FPixelStreamingEditorModule::StartStreaming()
+	IMainFrameModule::Get().OnMainFrameCreationFinished().AddLambda([&](TSharedPtr<SWindow> RootWindow, bool bIsNewProjectWindow)
 	{
-		// Activate our level editor streamer
-		IPixelStreamingModule& Module = IPixelStreamingModule::Get();
-		TSharedPtr<IPixelStreamingStreamer> Streamer = Module.GetStreamer(Module.GetDefaultStreamerID());
-		if (!Streamer.IsValid())
+		MaybeResizeEditor(RootWindow);
+
+		if(UE::EditorPixelStreaming::Settings::CVarEditorPixelStreamingStartOnLaunch.GetValueOnAnyThread())
 		{
-			return;
+			StartStreaming();
 		}
+	});
+}
 
-		// use the modules start streaming method to start all streamers
-		Module.StartStreaming();
+void FPixelStreamingEditorModule::StartStreaming()
+{
+	// Activate our level editor streamer
+	IPixelStreamingModule& Module = IPixelStreamingModule::Get();
+	TSharedPtr<IPixelStreamingStreamer> Streamer = Module.GetStreamer(Module.GetDefaultStreamerID());
+	if (!Streamer.IsValid())
+	{
+		return;
 	}
 
-	void FPixelStreamingEditorModule::StopStreaming()
+	// Use the Pixel Streaming module's start streaming method to start all streamers
+	Module.StartStreaming();
+}
+
+void FPixelStreamingEditorModule::StopStreaming()
+{
+	IPixelStreamingModule& Module = IPixelStreamingModule::Get();
+	TSharedPtr<IPixelStreamingStreamer> Streamer = Module.GetStreamer(Module.GetDefaultStreamerID());
+	if (!Streamer.IsValid())
 	{
-		// De-activate our level editor streamer
-		IPixelStreamingModule& Module = IPixelStreamingModule::Get();
-		TSharedPtr<IPixelStreamingStreamer> Streamer = Module.GetStreamer(Module.GetDefaultStreamerID());
-		if (!Streamer.IsValid())
-		{
-			return;
-		}
-		Streamer->SetTargetViewport(nullptr);
-		// use the modules stop streaming method to stop all streamers
-		Module.StopStreaming();
+		return;
 	}
 
-	FPixelStreamingEditorModule* FPixelStreamingEditorModule::GetModule()
+	Streamer->SetTargetViewport(nullptr);
+	// Use the Pixel Streaming module's stop streaming method to stop all streamers
+	Module.StopStreaming();
+}
+
+FPixelStreamingEditorModule* FPixelStreamingEditorModule::GetModule()
+{
+	if (PixelStreamingEditorModule)
 	{
-		if (PixelStreamingEditorModule)
-		{
-			return PixelStreamingEditorModule;
-		}
-		FPixelStreamingEditorModule* Module = FModuleManager::Get().LoadModulePtr<FPixelStreamingEditorModule>("PixelStreamingEditor");
-		if (Module)
-		{
-			PixelStreamingEditorModule = Module;
-		}
 		return PixelStreamingEditorModule;
 	}
+	FPixelStreamingEditorModule* Module = FModuleManager::Get().LoadModulePtr<FPixelStreamingEditorModule>("PixelStreamingEditor");
+	if (Module)
+	{
+		PixelStreamingEditorModule = Module;
+	}
+	return PixelStreamingEditorModule;
+}
 
-	bool FPixelStreamingEditorModule::ParseResolution(const TCHAR* InResolution, uint32& OutX, uint32& OutY)
+bool FPixelStreamingEditorModule::ParseResolution(const TCHAR* InResolution, uint32& OutX, uint32& OutY)
+{
+	if(*InResolution)
 	{
-		if(*InResolution)
-		{
-			FString CmdString(InResolution);
-			CmdString = CmdString.TrimStartAndEnd().ToLower();
+		FString CmdString(InResolution);
+		CmdString = CmdString.TrimStartAndEnd().ToLower();
 
-			// Retrieve the X dimensional value
-			const uint32 X = FMath::Max(FCString::Atof(*CmdString), 0.0f);
+		// Retrieve the X dimensional value
+		const uint32 X = FMath::Max(FCString::Atof(*CmdString), 0.0f);
 
-			// Determine whether the user has entered a resolution and extract the Y dimension.
-			FString YString;
+		// Determine whether the user has entered a resolution and extract the Y dimension.
+		FString YString;
 
-			// Find separator between values (Example of expected format: 1280x768)
-			const TCHAR* YValue = NULL;
-			if(FCString::Strchr(*CmdString,'x'))
-			{
-				YValue = const_cast<TCHAR*> (FCString::Strchr(*CmdString,'x')+1);
-				YString = YValue;
-				// Remove any whitespace from the end of the string
-				YString = YString.TrimStartAndEnd();
-			}
+		// Find separator between values (Example of expected format: 1280x768)
+		const TCHAR* YValue = NULL;
+		if(FCString::Strchr(*CmdString,'x'))
+		{
+			YValue = const_cast<TCHAR*> (FCString::Strchr(*CmdString,'x')+1);
+			YString = YValue;
+			// Remove any whitespace from the end of the string
+			YString = YString.TrimStartAndEnd();
+		}
 
-			// If the Y dimensional value exists then setup to use the specified resolution.
-			uint32 Y = 0;
-			if ( YValue && YString.Len() > 0 )
+		// If the Y dimensional value exists then setup to use the specified resolution.
+		uint32 Y = 0;
+		if ( YValue && YString.Len() > 0 )
+		{
+			if (YString.IsNumeric())
 			{
-				if (YString.IsNumeric())
-				{
-					Y = FMath::Max(FCString::Atof(YValue), 0.0f);
-					OutX = X;
-					OutY = Y;
-					return true;
-				}
+				Y = FMath::Max(FCString::Atof(YValue), 0.0f);
+				OutX = X;
+				OutY = Y;
+				return true;
 			}
 		}
-		return false;
 	}
+	return false;
+}
 
-	void FPixelStreamingEditorModule::MaybeResizeEditor(TSharedPtr<SWindow> RootWindow)
+void FPixelStreamingEditorModule::MaybeResizeEditor(TSharedPtr<SWindow> RootWindow)
+{
+	uint32 ResolutionX, ResolutionY = 0;
+	FString ResolutionStr;
+	bool bSuccess = FParse::Value(FCommandLine::Get(), TEXT("EditorPixelStreamingRes="), ResolutionStr);
+	if (bSuccess)
 	{
-		uint32 ResolutionX, ResolutionY = 0;
-		FString ResolutionStr;
-		bool bSuccess = FParse::Value(FCommandLine::Get(), TEXT("EditorPixelStreamingRes="), ResolutionStr);
-		if (bSuccess)
+		bSuccess = ParseResolution(*ResolutionStr, ResolutionX, ResolutionY);
+	}
+	else
+	{
+		bool UserSpecifiedWidth = FParse::Value(FCommandLine::Get(), TEXT("EditorPixelStreamingResX="), ResolutionX);
+		bool UserSpecifiedHeight = FParse::Value(FCommandLine::Get(), TEXT("EditorPixelStreamingResY="), ResolutionY);
+		bSuccess = UserSpecifiedWidth | UserSpecifiedHeight;
+
+		const float AspectRatio = 16.0 / 9.0;
+		if (UserSpecifiedWidth && !UserSpecifiedHeight)
 		{
-			bSuccess = ParseResolution(*ResolutionStr, ResolutionX, ResolutionY);
+			ResolutionY = int32(ResolutionX / AspectRatio);
 		}
-		else
+		else if (UserSpecifiedHeight && !UserSpecifiedWidth)
 		{
-			bool UserSpecifiedWidth = FParse::Value(FCommandLine::Get(), TEXT("EditorPixelStreamingResX="), ResolutionX);
-			bool UserSpecifiedHeight = FParse::Value(FCommandLine::Get(), TEXT("EditorPixelStreamingResY="), ResolutionY);
-			bSuccess = UserSpecifiedWidth | UserSpecifiedHeight;
-
-			const float AspectRatio = 16.0 / 9.0;
-			if (UserSpecifiedWidth && !UserSpecifiedHeight)
-			{
-				ResolutionY = int32(ResolutionX / AspectRatio);
-			}
-			else if (UserSpecifiedHeight && !UserSpecifiedWidth)
-			{
-				ResolutionX = int32(ResolutionY * AspectRatio);
-			}
+			ResolutionX = int32(ResolutionY * AspectRatio);
 		}
+	}
 
-		if(bSuccess)
-		{
-			RootWindow->Resize(FVector2D(ResolutionX, ResolutionY));
-			FSlateApplication::Get().OnSizeChanged(RootWindow->GetNativeWindow().ToSharedRef(), ResolutionX, ResolutionY);
-		}
+	if(bSuccess)
+	{
+		RootWindow->Resize(FVector2D(ResolutionX, ResolutionY));
+		FSlateApplication::Get().OnSizeChanged(RootWindow->GetNativeWindow().ToSharedRef(), ResolutionX, ResolutionY);
 	}
-} // namespace UE::PixelStreaming
+}
 
-IMPLEMENT_MODULE(UE::PixelStreaming::FPixelStreamingEditorModule, PixelStreamingEditor)
\ No newline at end of file
+IMPLEMENT_MODULE(FPixelStreamingEditorModule, PixelStreamingEditor)
\ No newline at end of file
diff --git a/Engine/Plugins/Media/PixelStreaming/Source/PixelStreamingEditor/Private/PixelStreamingStyle.cpp b/Engine/Plugins/Media/PixelStreaming/Source/PixelStreamingEditor/Private/PixelStreamingStyle.cpp
--- a/Engine/Plugins/Media/PixelStreaming/Source/PixelStreamingEditor/Private/PixelStreamingStyle.cpp
+++ b/Engine/Plugins/Media/PixelStreaming/Source/PixelStreamingEditor/Private/PixelStreamingStyle.cpp
@@ -7,7 +7,7 @@
 #include "Interfaces/IPluginManager.h"
 #include "Styling/SlateStyle.h"
 
-namespace UE::PixelStreaming
+namespace UE::EditorPixelStreaming
 {
     TSharedPtr<FSlateStyleSet> FPixelStreamingStyle::StyleInstance = NULL;
     void FPixelStreamingStyle::Initialize()
diff --git a/Engine/Plugins/Media/PixelStreaming/Source/PixelStreamingEditor/Private/PixelStreamingStyle.h b/Engine/Plugins/Media/PixelStreaming/Source/PixelStreamingEditor/Private/PixelStreamingStyle.h
--- a/Engine/Plugins/Media/PixelStreaming/Source/PixelStreamingEditor/Private/PixelStreamingStyle.h
+++ b/Engine/Plugins/Media/PixelStreaming/Source/PixelStreamingEditor/Private/PixelStreamingStyle.h
@@ -4,7 +4,7 @@
 #include "CoreMinimal.h"
 #include "Styling/SlateStyle.h"
 
-namespace UE::PixelStreaming
+namespace UE::EditorPixelStreaming
 {
     class FPixelStreamingStyle
     {
diff --git a/Engine/Plugins/Media/PixelStreaming/Source/PixelStreamingEditor/Private/PixelStreamingToolbar.cpp b/Engine/Plugins/Media/PixelStreaming/Source/PixelStreamingEditor/Private/PixelStreamingToolbar.cpp
--- a/Engine/Plugins/Media/PixelStreaming/Source/PixelStreamingEditor/Private/PixelStreamingToolbar.cpp
+++ b/Engine/Plugins/Media/PixelStreaming/Source/PixelStreamingEditor/Private/PixelStreamingToolbar.cpp
@@ -30,7 +30,7 @@
 
 #define LOCTEXT_NAMESPACE "PixelStreamingEditor"
 
-namespace UE::PixelStreaming
+namespace UE::EditorPixelStreaming
 {
     FPixelStreamingToolbar::FPixelStreamingToolbar()
     : PixelStreamingModule(IPixelStreamingModule::Get())
diff --git a/Engine/Plugins/Media/PixelStreaming/Source/PixelStreamingEditor/Private/PixelStreamingToolbar.h b/Engine/Plugins/Media/PixelStreaming/Source/PixelStreamingEditor/Private/PixelStreamingToolbar.h
--- a/Engine/Plugins/Media/PixelStreaming/Source/PixelStreamingEditor/Private/PixelStreamingToolbar.h
+++ b/Engine/Plugins/Media/PixelStreaming/Source/PixelStreamingEditor/Private/PixelStreamingToolbar.h
@@ -9,7 +9,7 @@
 #include "ToolMenus.h"
 #include "IPixelStreamingModule.h"
 
-namespace UE::PixelStreaming
+namespace UE::EditorPixelStreaming
 {
     class FPixelStreamingToolbar
     {
diff --git a/Engine/Plugins/Media/PixelStreaming/Source/PixelStreamingEditor/Private/PixelStreamingVideoInputBackBufferComposited.cpp b/Engine/Plugins/Media/PixelStreaming/Source/PixelStreamingEditor/Private/PixelStreamingVideoInputBackBufferComposited.cpp
new file mode 100644
--- /dev/null
+++ b/Engine/Plugins/Media/PixelStreaming/Source/PixelStreamingEditor/Private/PixelStreamingVideoInputBackBufferComposited.cpp
@@ -0,0 +1,165 @@
+// Copyright Epic Games, Inc. All Rights Reserved.
+
+#include "PixelStreamingVideoInputBackBufferComposited.h"
+#include "Settings.h"
+#include "Utils.h"
+#include "UtilsRender.h"
+#include "ToStringExtensions.h"
+#include "IPixelStreamingModule.h"
+#include "PixelStreamingPrivate.h"
+#include "PixelCaptureInputFrameRHI.h"
+#include "PixelCaptureCapturerRHI.h"
+#include "PixelCaptureCapturerRHIToI420CPU.h"
+#include "PixelCaptureCapturerRHIToI420Compute.h"
+#include "PixelCaptureBufferFormat.h"
+#include "Framework/Application/SlateApplication.h"
+#include "GenericPlatform/GenericWindowDefinition.h"
+#include "Application/SlateApplicationBase.h"
+#include "Widgets/SWindow.h"
+
+
+TSharedPtr<FPixelStreamingVideoInputBackBufferComposited> FPixelStreamingVideoInputBackBufferComposited::Create()
+{
+	TSharedPtr<FPixelStreamingVideoInputBackBufferComposited> NewInput = TSharedPtr<FPixelStreamingVideoInputBackBufferComposited>(new FPixelStreamingVideoInputBackBufferComposited());
+	TWeakPtr<FPixelStreamingVideoInputBackBufferComposited> WeakInput = NewInput;
+	// Set up the callback on the game thread since FSlateApplication::Get() can only be used there
+	UE::PixelStreaming::DoOnGameThread([WeakInput]() {
+		if (TSharedPtr<FPixelStreamingVideoInputBackBufferComposited> Input = WeakInput.Pin())
+		{
+			FSlateRenderer* Renderer = FSlateApplication::Get().GetRenderer();
+			Input->DelegateHandle = Renderer->OnBackBufferReadyToPresent().AddSP(Input.ToSharedRef(), &FPixelStreamingVideoInputBackBufferComposited::OnBackBufferReady);
+		}
+	});
+
+	return NewInput;
+}
+
+FPixelStreamingVideoInputBackBufferComposited::FPixelStreamingVideoInputBackBufferComposited()
+{
+	// Initialize our composited frame texture to a default size. It'll be enlarged / shrunk as needed
+	CompositedFrame = UE::PixelStreaming::CreateRHITexture(DefaultSize.X, DefaultSize.Y);
+	CompositedFrameSize = MakeShared<FIntPoint>(DefaultSize);
+	UE::PixelStreaming::DoOnGameThread([this]() {
+		FSlateApplication::Get().OnPreTick().AddLambda([this](float DeltaTime) {
+			TopLevelWindowsCriticalSection.Lock();
+			TopLevelWindows.Empty();
+			FSlateApplication::Get().GetAllVisibleWindowsOrdered(TopLevelWindows);
+			TopLevelWindowsCriticalSection.Unlock();
+		});
+	});
+}
+
+FPixelStreamingVideoInputBackBufferComposited::~FPixelStreamingVideoInputBackBufferComposited()
+{
+	if (!IsEngineExitRequested())
+	{
+		UE::PixelStreaming::DoOnGameThread([HandleCopy = DelegateHandle]() {
+			FSlateApplication::Get().GetRenderer()->OnBackBufferReadyToPresent().Remove(HandleCopy);
+		});
+	}
+}
+
+void FPixelStreamingVideoInputBackBufferComposited::OnBackBufferReady(SWindow& SlateWindow, const FTexture2DRHIRef& FrameBuffer)
+{
+	{
+		FScopeLock Lock(&TopLevelWindowsCriticalSection);
+		if (TopLevelWindows.IsEmpty())
+		{
+			return;
+		}
+	}
+
+	FRHICommandListImmediate& RHICmdList = FRHICommandListExecutor::GetImmediateCommandList();
+	FTextureRHIRef StagingTexture = UE::PixelStreaming::CreateRHITexture(FrameBuffer->GetSizeXY().X, FrameBuffer->GetSizeXY().Y);
+	// Re-render FrameBuffer to StagingTexture (ensure format match)
+	UE::PixelStreaming::CopyTexture(RHICmdList, FrameBuffer, StagingTexture, nullptr);
+	
+	TopLevelWindowTextures.Add(&SlateWindow, StagingTexture);
+	TopLevelWindowsCriticalSection.Lock();
+	// Check that we have received a texture from every window in the TopLevelWindows array
+	uint8 WindowsRendered = 0;
+	for (TSharedRef<SWindow> Window : TopLevelWindows)
+	{
+		if (TopLevelWindowTextures.FindRef(&Window.Get()))
+		{
+			WindowsRendered += 1;
+		}
+	}
+
+	if (WindowsRendered == TopLevelWindows.Num())
+	{
+		// We have all the textures needed, let's composite them
+		// Need to iterate over the TopLevelWindows array as this array contains the ordered windows
+		for (TSharedRef<SWindow> Window : TopLevelWindows)
+		{
+			FTextureRHIRef Texture = TopLevelWindowTextures.FindRef(&Window.Get());
+			if (Window->GetOpacity() == 0.0f || Window->GetType() == EWindowType::CursorDecorator)
+			{
+				continue;
+			}
+
+			uint32 CompositedX = CompositedFrame->GetSizeXY().X;
+			uint32 CompositedY = CompositedFrame->GetSizeXY().Y;
+			// Use the absolute value for window positions. This is due to main windows having a position of (-8, -8) which can lead to crashes when copying textures
+			uint32 TotalX = FMath::Min(Texture->GetSizeXY().X, Window->GetSizeInScreen().IntPoint().X) + (Window->GetPositionInScreen().X > 0 ? Window->GetPositionInScreen().X : 0);
+			uint32 TotalY = FMath::Min(Texture->GetSizeXY().Y, Window->GetSizeInScreen().IntPoint().Y) + (Window->GetPositionInScreen().Y > 0 ? Window->GetPositionInScreen().Y : 0);
+
+			// Check if the location + size of our texture will be outside the bounds of the compositedframe
+			if (CompositedX < TotalX || CompositedY < TotalY)
+			{
+				// If so, create a new compositedframe large enough to accomodate our new texture + its position
+				// and render the contents of the old compositedframe to the new one
+				uint32 SizeX = TotalX > CompositedX ? TotalX : CompositedX;
+				uint32 SizeY = TotalY > CompositedY ? TotalY : CompositedY;
+
+				StagingTexture = UE::PixelStreaming::CreateRHITexture(SizeX, SizeY);
+				FRHICopyTextureInfo CopyInfo;
+				CopyInfo.Size.X = CompositedFrame->GetSizeXY().X;
+				CopyInfo.Size.Y = CompositedFrame->GetSizeXY().Y;
+				CopyInfo.Size.Z = 1;
+				CopyInfo.DestPosition.X = 0;
+				CopyInfo.DestPosition.Y = 0;
+				CopyInfo.DestPosition.Z = 0;
+				// source and dest are the same. simple copy
+				RHICmdList.Transition(FRHITransitionInfo(CompositedFrame, ERHIAccess::Unknown, ERHIAccess::CopySrc));
+				RHICmdList.Transition(FRHITransitionInfo(StagingTexture, ERHIAccess::Unknown, ERHIAccess::CopyDest));
+				RHICmdList.CopyTexture(CompositedFrame, StagingTexture, CopyInfo);
+				// Update our composited frame with the new, larger texture for later use
+				CompositedFrame = StagingTexture;
+			}
+
+			// Copy our new texture to the compositedframe texture
+			FRHICopyTextureInfo CopyInfo;
+			// There is only ever one tooltip and as such UE keeps the same texture for each and just rerenders the content
+			// this can lead to small tooltips having a large texture from a previously displayed long tooltip
+			// so we use the tooltips window size which is guaranteed to be correct
+			CopyInfo.Size.X = FMath::Min(Texture->GetSizeXY().X, Window->GetSizeInScreen().IntPoint().X);
+			CopyInfo.Size.Y = FMath::Min(Texture->GetSizeXY().Y, Window->GetSizeInScreen().IntPoint().Y);
+			CopyInfo.Size.Z = 1;
+			// Sometimes a windows position in screen is negative (eg the main window has a position of -8,-8).
+			// This causes issues when copying the texture so we'll just ignore negative values
+			CopyInfo.DestPosition.X = Window->GetPositionInScreen().X > 0 ? Window->GetPositionInScreen().X : 0;
+			CopyInfo.DestPosition.Y = Window->GetPositionInScreen().Y > 0 ? Window->GetPositionInScreen().Y : 0;
+			CopyInfo.DestPosition.Z = 0;
+			// source and dest are the same. simple copy
+			RHICmdList.Transition(FRHITransitionInfo(Texture, ERHIAccess::Unknown, ERHIAccess::CopySrc));
+			RHICmdList.Transition(FRHITransitionInfo(CompositedFrame, ERHIAccess::Unknown, ERHIAccess::CopyDest));
+			RHICmdList.CopyTexture(Texture, CompositedFrame, CopyInfo);
+		}
+
+		// Our composition is complete, send it to along the pipeline
+		OnFrame(FPixelCaptureInputFrameRHI(CompositedFrame));
+		// Update the default streamer to let it know out compositedframe size. This way it can correctly scale input from the browser
+		*CompositedFrameSize.Get() = CompositedFrame->GetSizeXY();
+		IPixelStreamingModule& Module = IPixelStreamingModule::Get();
+		Module.GetStreamer(Module.GetDefaultStreamerID())->SetTargetScreenSize(CompositedFrameSize);
+		bRecreateTexture = true;
+	}
+	TopLevelWindowsCriticalSection.Unlock();
+
+	if (bRecreateTexture)
+	{
+		CompositedFrame = UE::PixelStreaming::CreateRHITexture(DefaultSize.X, DefaultSize.Y);
+		bRecreateTexture = false;
+	}
+}
diff --git a/Engine/Plugins/Media/PixelStreaming/Source/PixelStreamingEditor/Private/PixelStreamingVideoInputViewport.cpp b/Engine/Plugins/Media/PixelStreaming/Source/PixelStreamingEditor/Private/PixelStreamingVideoInputViewport.cpp
new file mode 100644
--- /dev/null
+++ b/Engine/Plugins/Media/PixelStreaming/Source/PixelStreamingEditor/Private/PixelStreamingVideoInputViewport.cpp
@@ -0,0 +1,89 @@
+// Copyright Epic Games, Inc. All Rights Reserved.
+
+#include "PixelStreamingVideoInputViewport.h"
+#include "Settings.h"
+#include "Utils.h"
+#include "PixelCaptureInputFrameRHI.h"
+#include "PixelCaptureBufferFormat.h"
+#include "PixelCaptureCapturerRHI.h"
+#include "PixelCaptureCapturerRHIToI420CPU.h"
+#include "PixelCaptureCapturerRHIToI420Compute.h"
+#include "UnrealClient.h"
+#include "Engine/GameViewportClient.h"
+#include "RenderingThread.h"
+#include "Framework/Application/SlateApplication.h"
+#include "Slate/SceneViewport.h"
+#include "IPixelStreamingModule.h"
+
+
+TSharedPtr<FPixelStreamingVideoInputViewport> FPixelStreamingVideoInputViewport::Create()
+{
+	TSharedPtr<FPixelStreamingVideoInputViewport> NewInput = TSharedPtr<FPixelStreamingVideoInputViewport>(new FPixelStreamingVideoInputViewport());
+	TWeakPtr<FPixelStreamingVideoInputViewport> WeakInput = NewInput;
+
+	// Set up the callback on the game thread since FSlateApplication::Get() can only be used there
+	UE::PixelStreaming::DoOnGameThread([WeakInput]() {
+		if (TSharedPtr<FPixelStreamingVideoInputViewport> Input = WeakInput.Pin())
+		{
+			Input->DelegateHandle = UGameViewportClient::OnViewportRendered().AddSP(Input.ToSharedRef(), &FPixelStreamingVideoInputViewport::OnViewportRendered);
+		}
+	});
+
+	return NewInput;
+}
+
+FPixelStreamingVideoInputViewport::~FPixelStreamingVideoInputViewport()
+{
+	if (!IsEngineExitRequested())
+	{
+		UE::PixelStreaming::DoOnGameThread([HandleCopy = DelegateHandle]() {
+			UGameViewportClient::OnViewportRendered().Remove(HandleCopy);
+		});
+	}
+}
+
+void FPixelStreamingVideoInputViewport::OnViewportRendered(FViewport* InViewport)
+{
+	IPixelStreamingModule& Module = IPixelStreamingModule::Get();
+	TSharedPtr<IPixelStreamingStreamer> Streamer = Module.GetStreamer(Module.GetDefaultStreamerID());
+	if (!Streamer.IsValid() || !Streamer->IsStreaming())
+	{
+		return;
+	}
+
+	TSharedPtr<SViewport> TargetScene = Streamer->GetTargetViewport().Pin();
+	if (!TargetScene.IsValid())
+	{
+		return;
+	}
+
+	if (InViewport == nullptr)
+	{
+		return;
+	}
+
+	if (InViewport->GetViewportType() != TargetViewportType)
+	{
+		return;
+	}
+
+	// Bit dirty to do a static cast here, but we check viewport type just above so it is somewhat "safe".
+	TSharedPtr<SViewport> InScene = static_cast<FSceneViewport*>(InViewport)->GetViewportWidget().Pin();
+
+	// If the viewport we were passed is not our target viewport we are not interested in getting its texture.
+	if (TargetScene != InScene)
+	{
+		return;
+	}
+
+	const FTextureRHIRef& FrameBuffer = InViewport->GetRenderTargetTexture();
+	if (!FrameBuffer)
+	{
+		return;
+	}
+
+	ENQUEUE_RENDER_COMMAND(StreamViewportTextureCommand)
+	([&, FrameBuffer](FRHICommandList& RHICmdList) {
+		OnFrame(FPixelCaptureInputFrameRHI(FrameBuffer));
+	});
+}
\ No newline at end of file
diff --git a/Engine/Plugins/Media/PixelStreaming/Source/PixelStreamingEditor/Private/Settings.cpp b/Engine/Plugins/Media/PixelStreaming/Source/PixelStreamingEditor/Private/Settings.cpp
--- a/Engine/Plugins/Media/PixelStreaming/Source/PixelStreamingEditor/Private/Settings.cpp
+++ b/Engine/Plugins/Media/PixelStreaming/Source/PixelStreamingEditor/Private/Settings.cpp
@@ -3,44 +3,7 @@
 #include "Settings.h"
 #include "Utils.h"
 
-template <typename T>
-void PSEditorCommandLineParseValue(const TCHAR* Match, TAutoConsoleVariable<T>& CVar)
-{
-	T Value;
-	if (FParse::Value(FCommandLine::Get(), Match, Value))
-		CVar->Set(Value, ECVF_SetByCommandline);
-};
-
-void PSEditorCommandLineParseValue(const TCHAR* Match, TAutoConsoleVariable<FString>& CVar, bool bStopOnSeparator = false)
-{
-	FString Value;
-	if (FParse::Value(FCommandLine::Get(), Match, Value, bStopOnSeparator))
-		CVar->Set(*Value, ECVF_SetByCommandline);
-};
-
-void PSEditorCommandLineParseOption(const TCHAR* Match, TAutoConsoleVariable<bool>& CVar)
-{
-	FString ValueMatch(Match);
-	ValueMatch.Append(TEXT("="));
-	FString Value;
-	if (FParse::Value(FCommandLine::Get(), *ValueMatch, Value))
-	{
-		if (Value.Equals(FString(TEXT("true")), ESearchCase::IgnoreCase))
-		{
-			CVar->Set(true, ECVF_SetByCommandline);
-		}
-		else if (Value.Equals(FString(TEXT("false")), ESearchCase::IgnoreCase))
-		{
-			CVar->Set(false, ECVF_SetByCommandline);
-		}
-	}
-	else if (FParse::Param(FCommandLine::Get(), Match))
-	{
-		CVar->Set(true, ECVF_SetByCommandline);
-	}
-}
-
-namespace UE::PixelStreaming::Settings::Editor
+namespace UE::EditorPixelStreaming::Settings
 {
     TAutoConsoleVariable<bool> CVarEditorPixelStreamingStartOnLaunch(
 		TEXT("PixelStreaming.Editor.StartOnLaunch"),
@@ -50,7 +13,8 @@ namespace UE::PixelStreaming::Settings::Editor
 
     void InitialiseSettings()
     {
+		using namespace UE::PixelStreaming;
         // Options parse (if these exist they are set to true)
-		PSEditorCommandLineParseOption(TEXT("EditorPixelStreamingStartOnLaunch"), CVarEditorPixelStreamingStartOnLaunch);
+		CommandLineParseOption(TEXT("EditorPixelStreamingStartOnLaunch"), CVarEditorPixelStreamingStartOnLaunch);
     }
 }
\ No newline at end of file
diff --git a/Engine/Plugins/Media/PixelStreaming/Source/PixelStreamingEditor/Private/Settings.h b/Engine/Plugins/Media/PixelStreaming/Source/PixelStreamingEditor/Private/Settings.h
--- a/Engine/Plugins/Media/PixelStreaming/Source/PixelStreamingEditor/Private/Settings.h
+++ b/Engine/Plugins/Media/PixelStreaming/Source/PixelStreamingEditor/Private/Settings.h
@@ -6,7 +6,7 @@
 #include "HAL/IConsoleManager.h"
 #include "Misc/CommandLine.h"
 
-namespace UE::PixelStreaming::Settings::Editor
+namespace UE::EditorPixelStreaming::Settings
 {
     extern void InitialiseSettings();
 
diff --git a/Engine/Plugins/Media/PixelStreaming/Source/PixelStreamingEditor/Public/PixelStreamingEditorModule.h b/Engine/Plugins/Media/PixelStreaming/Source/PixelStreamingEditor/Public/PixelStreamingEditorModule.h
--- a/Engine/Plugins/Media/PixelStreaming/Source/PixelStreamingEditor/Public/PixelStreamingEditorModule.h
+++ b/Engine/Plugins/Media/PixelStreaming/Source/PixelStreamingEditor/Public/PixelStreamingEditorModule.h
@@ -4,26 +4,27 @@
 
 #include "IPixelStreamingModule.h"
 
-namespace UE::PixelStreaming
+namespace UE::EditorPixelStreaming
 {
-    class PIXELSTREAMINGEDITOR_API FPixelStreamingEditorModule : public IModuleInterface
-    {
-    public:
-        /** IModuleInterface implementation */
-        virtual void StartupModule() override;
-        virtual void ShutdownModule() override;
+    class FPixelStreamingToolbar;
+}
 
-        void StartStreaming();
-        void StopStreaming();
-        
-        static FPixelStreamingEditorModule* GetModule();
-    private:
-        void InitEditorStreaming(IPixelStreamingModule& Module);
-        bool ParseResolution(const TCHAR* InResolution, uint32& OutX, uint32& OutY);
-        void MaybeResizeEditor(TSharedPtr<SWindow> RootWindow);
+class PIXELSTREAMINGEDITOR_API FPixelStreamingEditorModule : public IModuleInterface
+{
+public:
+    /** IModuleInterface implementation */
+    virtual void StartupModule() override;
+    virtual void ShutdownModule() override;
 
-        TSharedPtr<class FPixelStreamingToolbar> Toolbar;
+    void StartStreaming();
+    void StopStreaming();
+       
+    static FPixelStreamingEditorModule* GetModule();
+private:
+    void InitEditorStreaming(IPixelStreamingModule& Module);
+	bool ParseResolution(const TCHAR* InResolution, uint32& OutX, uint32& OutY);
+    void MaybeResizeEditor(TSharedPtr<SWindow> RootWindow);
 
-        static FPixelStreamingEditorModule* PixelStreamingEditorModule;	
-    };
-} // namespace UE::PixelStreaming
+    TSharedPtr<UE::EditorPixelStreaming::FPixelStreamingToolbar> Toolbar;
+    static FPixelStreamingEditorModule* PixelStreamingEditorModule;	
+};
diff --git a/Engine/Plugins/Media/PixelStreaming/Source/PixelStreamingEditor/Public/PixelStreamingVideoInputBackBufferComposited.h b/Engine/Plugins/Media/PixelStreaming/Source/PixelStreamingEditor/Public/PixelStreamingVideoInputBackBufferComposited.h
new file mode 100644
--- /dev/null
+++ b/Engine/Plugins/Media/PixelStreaming/Source/PixelStreamingEditor/Public/PixelStreamingVideoInputBackBufferComposited.h
@@ -0,0 +1,37 @@
+// Copyright Epic Games, Inc. All Rights Reserved.
+
+#pragma once
+
+#include "PixelStreamingVideoInputRHI.h"
+#include "Widgets/SWindow.h"
+#include "RHI.h"
+#include "Delegates/IDelegateInstance.h"
+#include "Widgets/SWindow.h"
+#include "GenericPlatform/GenericWindowDefinition.h"
+
+
+/*
+ * Use this if you want to send the full UE editor as video input.
+ */
+class PIXELSTREAMINGEDITOR_API FPixelStreamingVideoInputBackBufferComposited : public FPixelStreamingVideoInputRHI
+{
+public:
+	static TSharedPtr<FPixelStreamingVideoInputBackBufferComposited> Create();
+	virtual ~FPixelStreamingVideoInputBackBufferComposited();
+
+private:
+	FPixelStreamingVideoInputBackBufferComposited();
+
+	void OnBackBufferReady(SWindow& SlateWindow, const FTexture2DRHIRef& FrameBuffer);
+
+	FDelegateHandle DelegateHandle;
+
+	TArray<TSharedRef<SWindow>> TopLevelWindows;
+	TMap<SWindow*, FTextureRHIRef> TopLevelWindowTextures;
+
+	FCriticalSection TopLevelWindowsCriticalSection;
+	FTexture2DRHIRef CompositedFrame;
+	TSharedPtr<FIntPoint> CompositedFrameSize;
+	bool bRecreateTexture = false;
+	FIntPoint DefaultSize = FIntPoint(1, 1);
+};
\ No newline at end of file
diff --git a/Engine/Plugins/Media/PixelStreaming/Source/PixelStreamingEditor/Public/PixelStreamingVideoInputViewport.h b/Engine/Plugins/Media/PixelStreaming/Source/PixelStreamingEditor/Public/PixelStreamingVideoInputViewport.h
new file mode 100644
--- /dev/null
+++ b/Engine/Plugins/Media/PixelStreaming/Source/PixelStreamingEditor/Public/PixelStreamingVideoInputViewport.h
@@ -0,0 +1,27 @@
+// Copyright Epic Games, Inc. All Rights Reserved.
+
+#pragma once
+
+#include "PixelStreamingVideoInputRHI.h"
+#include "Delegates/IDelegateInstance.h"
+#include "UnrealClient.h"
+
+
+/*
+ * Use this if you want to send the UE primary scene viewport as video input - will only work in editor.
+ */
+class PIXELSTREAMINGEDITOR_API FPixelStreamingVideoInputViewport : public FPixelStreamingVideoInputRHI
+{
+public:
+	static TSharedPtr<FPixelStreamingVideoInputViewport> Create();
+	virtual ~FPixelStreamingVideoInputViewport();
+
+private:
+	FPixelStreamingVideoInputViewport() = default;
+
+	void OnViewportRendered(FViewport* Viewport);
+
+	FDelegateHandle DelegateHandle;
+
+	FName TargetViewportType = FName(FString(TEXT("SceneViewport")));
+};
