Cherrypicked from 21693843:

Fix: Remove the overhead from just having the pixel streaming plugin enabled.
Previously, video inputs were created as soon as the module was ready and would capture and adapt as soon as they were created as the streamer was requesting an RHI format for its freeze frames. Now, the video inputs are created if the user specifies a URL, meaning they intend on streaming and hence require a video input.

diff --git a/Engine/Plugins/Media/PixelStreaming/Source/PixelStreaming/Private/PixelStreamingModule.cpp b/Engine/Plugins/Media/PixelStreaming/Source/PixelStreaming/Private/PixelStreamingModule.cpp
--- a/Engine/Plugins/Media/PixelStreaming/Source/PixelStreaming/Private/PixelStreamingModule.cpp
+++ b/Engine/Plugins/Media/PixelStreaming/Source/PixelStreaming/Private/PixelStreamingModule.cpp
@@ -406,7 +406,6 @@ namespace UE::PixelStreaming
 		// The PixelStreamingEditorModule handles setting video input in the editor
 		if (!GIsEditor)
 		{
-			Streamer->SetVideoInput(FPixelStreamingVideoInputBackBuffer::Create());
 			// default to the scene viewport if we have a game engine
 			if (UGameEngine* GameEngine = Cast<UGameEngine>(GEngine))
 			{
@@ -425,6 +424,12 @@ namespace UE::PixelStreaming
 
 		if (!SignallingServerURL.IsEmpty())
 		{
+			// The user has specified a URL on the command line meaning their intention is to start streaming immediately
+			// in that case, set up the video input for them (as long as we're not in editor)
+			if(!GIsEditor)
+			{
+				Streamer->SetVideoInput(FPixelStreamingVideoInputBackBuffer::Create());
+			}
 			Streamer->SetSignallingServerURL(SignallingServerURL);
 		}
 	}
diff --git a/Engine/Plugins/Media/PixelStreaming/Source/PixelStreamingEditor/Private/PixelStreamingEditorModule.cpp b/Engine/Plugins/Media/PixelStreaming/Source/PixelStreamingEditor/Private/PixelStreamingEditorModule.cpp
--- a/Engine/Plugins/Media/PixelStreaming/Source/PixelStreamingEditor/Private/PixelStreamingEditorModule.cpp
+++ b/Engine/Plugins/Media/PixelStreaming/Source/PixelStreamingEditor/Private/PixelStreamingEditorModule.cpp
@@ -91,7 +91,6 @@ void FPixelStreamingEditorModule::InitEditorStreaming(IPixelStreamingModule& Mod
 								ExtendedHandler(Ar);
 							});
 
-	Streamer->SetVideoInput(FPixelStreamingVideoInputBackBufferComposited::Create());
 	// Give the editor streamer the default url if the user hasn't specified one when launching the editor
 	if (Streamer->GetSignallingServerURL().IsEmpty())
 	{
@@ -119,6 +118,13 @@ void FPixelStreamingEditorModule::StartStreaming()
 		return;
 	}
 
+	// The streamer doesn't currently have a video input so create one
+	TWeakPtr<FPixelStreamingVideoInput> VideoInput = Streamer->GetVideoInput();
+	if(!VideoInput.Pin())
+	{
+		Streamer->SetVideoInput(FPixelStreamingVideoInputBackBufferComposited::Create());
+	}
+
 	// Use the Pixel Streaming module's start streaming method to start all streamers
 	Module.StartStreaming();
 }
