Cherrypicked from 21693927:

QOL: Disable the rendering of tool tips if the editor is rendering off screen.
This is due to the tool tips currently being broken so while that fix is in the works we'll just disable them.

diff --git a/Engine/Plugins/Media/PixelStreaming/Source/PixelStreamingEditor/Private/PixelStreamingEditorModule.cpp b/Engine/Plugins/Media/PixelStreaming/Source/PixelStreamingEditor/Private/PixelStreamingEditorModule.cpp
--- a/Engine/Plugins/Media/PixelStreaming/Source/PixelStreamingEditor/Private/PixelStreamingEditorModule.cpp
+++ b/Engine/Plugins/Media/PixelStreaming/Source/PixelStreamingEditor/Private/PixelStreamingEditorModule.cpp
@@ -101,6 +101,10 @@ void FPixelStreamingEditorModule::InitEditorStreaming(IPixelStreamingModule& Mod
 	{
 		MaybeResizeEditor(RootWindow);
 
+		// We don't want to show tooltips in render off screen as they're currently broken
+		bool bIsRenderingOffScreen = FParse::Param(FCommandLine::Get(), TEXT("RenderOffScreen"));
+		FSlateApplication::Get().SetAllowTooltips(!bIsRenderingOffScreen);
+
 		if(UE::EditorPixelStreaming::Settings::CVarEditorPixelStreamingStartOnLaunch.GetValueOnAnyThread())
 		{
 			StartStreaming();
