Fix usage of inline functions in PixelStreamingServers which is unsafe for some build configurations

diff --git a/Engine/Plugins/Media/PixelStreaming/Source/PixelStreamingServers/Private/ServerUtils.h b/Engine/Plugins/Media/PixelStreaming/Source/PixelStreamingServers/Private/ServerUtils.h
--- a/Engine/Plugins/Media/PixelStreaming/Source/PixelStreamingServers/Private/ServerUtils.h
+++ b/Engine/Plugins/Media/PixelStreaming/Source/PixelStreamingServers/Private/ServerUtils.h
@@ -20,7 +20,7 @@ namespace UE::PixelStreamingServers
 		 * @param FallbackArgValue If the key does not exist in the `LaunchArgs.ProcessArgs` this value to set for it in-place.
 		 * @return The extracted/set process arg.
 		 **/
-		inline FString QueryOrSetProcessArgs(FLaunchArgs& LaunchArgs, FString ArgKey, FString FallbackArgValue);
+		FString QueryOrSetProcessArgs(FLaunchArgs& LaunchArgs, FString ArgKey, FString FallbackArgValue);
 
 		/**
 		 * Launches an executable or script as a child process of Unreal Engine.
@@ -30,7 +30,7 @@ namespace UE::PixelStreamingServers
 		 * @param LogPrefix If the child process has output the prefix to put on the UE log message for that output.
 		 * @param bRunAsScript If true the `ExecutableAbsPath` will be passed to cmd/bash.
 		 **/
-		inline TSharedPtr<FMonitoredProcess> LaunchChildProcess(FString ExecutableAbsPath, FString Args, FString LogPrefix, bool bRunAsScript);
+		TSharedPtr<FMonitoredProcess> LaunchChildProcess(FString ExecutableAbsPath, FString Args, FString LogPrefix, bool bRunAsScript);
 
 		/**
 		 * Get the absolute path to the Pixel Streaming resources dir (changes based on whether called in editor or not.)
@@ -38,14 +38,14 @@ namespace UE::PixelStreamingServers
 		 * @param OutResourcesDir The absolute path to "resources" dir.
 		 * @return True if the returned path exists.
 		 **/
-		inline bool GetResourcesDir(FString& OutResourcesDir);
+		bool GetResourcesDir(FString& OutResourcesDir);
 
 		/**
 		 * Get the absolute path to the webservers (changes based on whether called in editor or not.)
 		 * @param OutWebServersAbsPath The absolute path to the webservers.
 		 * @return True if the returned path exists.
 		 **/
-		inline bool GetWebServersDir(FString& OutWebServersAbsPath);
+		bool GetWebServersDir(FString& OutWebServersAbsPath);
 
 		/**
 		 * Attempts to get the absolute path to a downloaded server residing in the /Resources directory
@@ -54,42 +54,42 @@ namespace UE::PixelStreamingServers
 		 * @param ServerDirectoryName The directory name of the server we should look for under /Resources
 		 * @return True if the scripts to launch the server exist at this location.
 		 **/
-		inline bool GetDownloadedServer(FString& OutAbsPath, FString ServerDirectoryName);
+		bool GetDownloadedServer(FString& OutAbsPath, FString ServerDirectoryName);
 
 		/**
 		 * Download the Pixel Streaming servers using the `get_ps_servers` scripts.
 		 * @param bSkipIfPresent Servers will not be downloaded if they are already present.
 		 * @return The child process that is used to download the servers.
 		 **/
-		inline TSharedPtr<FMonitoredProcess> DownloadPixelStreamingServers(bool bSkipIfPresent);
+		TSharedPtr<FMonitoredProcess> DownloadPixelStreamingServers(bool bSkipIfPresent);
 
 		/**
 		 * Populates the cirrus endpoints by extracting endpoint information from the `InLaunchArgs.ProcessArgs`, and if not found, using defaults.
 		 * @param InLaunchArgs The launch args containing the process args that we will be inspecting.
 		 * @param OutEndPoints The end points for the server if we launch it with these launch args.
 		 **/
-		inline void PopulateCirrusEndPoints(FLaunchArgs& InLaunchArgs, TMap<EEndpoint, FURL>& OutEndPoints);
+		void PopulateCirrusEndPoints(FLaunchArgs& InLaunchArgs, TMap<EEndpoint, FURL>& OutEndPoints);
 
 		/**
 		 * Convert a `FURL` into a `FString`. `FURL::ToString()` already exists; however, it appends a problematic trailing slash to the output string.
 		 * @param Url The url to format into a string.
 		 * @return The url formatted without a trailing slash (the trailing slash messes up the Websocket clients).
 		 **/
-		inline FString ToString(FURL Url);
+		FString ToString(FURL Url);
 
 		/**
 		 * Convert a UTF8 byte array to an FString
 		 * @param UTF8Bytes Byte array of UTF8 characters.
 		 * @return A string representation of the byte array.
 		 **/
-		inline FString ToString(TArrayView<uint8> UTF8Bytes);
+		FString ToString(TArrayView<uint8> UTF8Bytes);
 
 		/**
 		 * Convert JSON object to FString.
 		 * @param JSONObj A JSON object.
 		 * @return The stringified JSON object.
 		 **/
-		inline FString ToString(TSharedRef<FJsonObject> JSONObj);
+		FString ToString(TSharedRef<FJsonObject> JSONObj);
 
 		/** 
 		 * Convert a JSON formatted string into a JSON object, if possible.
@@ -97,7 +97,7 @@ namespace UE::PixelStreamingServers
 		 * @param OutJSON The output json object.
 		 * @return True if we were able to deserialize the string into a JSON object.
 		 **/
-		inline bool Jsonify(FString InJSONString, TSharedPtr<FJsonObject>& OutJSON);
+		bool Jsonify(FString InJSONString, TSharedPtr<FJsonObject>& OutJSON);
 
 	} // Utils
 } // UE::PixelStreamingServers
