Fix intermittent failure of PixelStreaming.HandshakePlayerOffer and similar tests due to requirement that streamer connects to the signalling server before the player.
It was possible for this to happen out of order, so logic was modified to ensure streamer connects first in tests.
Also disable Pixel Streaming module warning under a test environment.

diff --git a/Engine/Plugins/Media/PixelStreaming/Source/PixelStreaming/Private/PixelStreamingModule.cpp b/Engine/Plugins/Media/PixelStreaming/Source/PixelStreaming/Private/PixelStreamingModule.cpp
--- a/Engine/Plugins/Media/PixelStreaming/Source/PixelStreaming/Private/PixelStreamingModule.cpp
+++ b/Engine/Plugins/Media/PixelStreaming/Source/PixelStreaming/Private/PixelStreamingModule.cpp
@@ -127,7 +127,9 @@ namespace UE::PixelStreaming
 		}
 		else
 		{
-			UE_LOG(LogPixelStreaming, Warning, TEXT("Only D3D11/D3D12/Vulkan Dynamic RHI is supported. Detected %s"), GDynamicRHI != nullptr ? GDynamicRHI->GetName() : TEXT("[null]"));
+			#if !WITH_DEV_AUTOMATION_TESTS
+				UE_LOG(LogPixelStreaming, Warning, TEXT("Only D3D11/D3D12/Vulkan Dynamic RHI is supported. Detected %s"), GDynamicRHI != nullptr ? GDynamicRHI->GetName() : TEXT("[null]"));
+			#endif
 		}
 
 		rtc::InitializeSSL();
diff --git a/Engine/Plugins/Media/PixelStreaming/Source/PixelStreaming/Private/Tests/CodecTest.cpp b/Engine/Plugins/Media/PixelStreaming/Source/PixelStreaming/Private/Tests/CodecTest.cpp
--- a/Engine/Plugins/Media/PixelStreaming/Source/PixelStreaming/Private/Tests/CodecTest.cpp
+++ b/Engine/Plugins/Media/PixelStreaming/Source/PixelStreaming/Private/Tests/CodecTest.cpp
@@ -31,9 +31,8 @@ namespace UE::PixelStreaming
 		Player->VideoSink = VideoSink;
 
 		Streamer->StartStreaming();
-		Player->Connect(PlayerPort);
 
-		ADD_LATENT_AUTOMATION_COMMAND(FWaitForStreamerConnectedOrTimeout(5.0, Streamer))
+		ADD_LATENT_AUTOMATION_COMMAND(FConnectPlayerAfterStreamerConnectedOrTimeout(5.0, Streamer, Player, PlayerPort))
 		ADD_LATENT_AUTOMATION_COMMAND(FWaitForPlayerConnectedOrTimeout(5.0, Player, PlayerPort))
 		ADD_LATENT_AUTOMATION_COMMAND(FWaitForICEConnectedOrTimeout(5.0, Player))
 		ADD_LATENT_AUTOMATION_COMMAND(FSendSolidColorFrame(VideoInput, FrameConfig))
@@ -60,9 +59,8 @@ namespace UE::PixelStreaming
 		// Note: Important to couple framerate as we are manually passing frames and don't want any cached frames
 		Streamer->SetCoupleFramerate(true);
 		Streamer->StartStreaming();
-		Player->Connect(PlayerPort);
 
-		ADD_LATENT_AUTOMATION_COMMAND(FWaitForStreamerConnectedOrTimeout(5.0, Streamer))
+		ADD_LATENT_AUTOMATION_COMMAND(FConnectPlayerAfterStreamerConnectedOrTimeout(5.0, Streamer, Player, PlayerPort))
 		ADD_LATENT_AUTOMATION_COMMAND(FWaitForPlayerConnectedOrTimeout(5.0, Player, PlayerPort))
 		ADD_LATENT_AUTOMATION_COMMAND(FWaitForICEConnectedOrTimeout(5.0, Player))
 
diff --git a/Engine/Plugins/Media/PixelStreaming/Source/PixelStreaming/Private/Tests/DataChannelTest.cpp b/Engine/Plugins/Media/PixelStreaming/Source/PixelStreaming/Private/Tests/DataChannelTest.cpp
--- a/Engine/Plugins/Media/PixelStreaming/Source/PixelStreaming/Private/Tests/DataChannelTest.cpp
+++ b/Engine/Plugins/Media/PixelStreaming/Source/PixelStreaming/Private/Tests/DataChannelTest.cpp
@@ -45,9 +45,8 @@ namespace UE::PixelStreaming
         };
 
 		Streamer->StartStreaming();
-		Player->Connect(PlayerPort);
 
-		ADD_LATENT_AUTOMATION_COMMAND(FWaitForStreamerConnectedOrTimeout(5.0, Streamer))
+		ADD_LATENT_AUTOMATION_COMMAND(FConnectPlayerAfterStreamerConnectedOrTimeout(5.0, Streamer, Player, PlayerPort))
 		ADD_LATENT_AUTOMATION_COMMAND(FWaitForPlayerConnectedOrTimeout(5.0, Player, PlayerPort))
 		ADD_LATENT_AUTOMATION_COMMAND(FWaitForICEConnectedOrTimeout(5.0, Player))
         ADD_LATENT_AUTOMATION_COMMAND(FSendDataChannelMessageToStreamer(Player, ToStreamerEchoId, EchoContent))
diff --git a/Engine/Plugins/Media/PixelStreaming/Source/PixelStreaming/Private/Tests/HandshakeTest.cpp b/Engine/Plugins/Media/PixelStreaming/Source/PixelStreaming/Private/Tests/HandshakeTest.cpp
--- a/Engine/Plugins/Media/PixelStreaming/Source/PixelStreaming/Private/Tests/HandshakeTest.cpp
+++ b/Engine/Plugins/Media/PixelStreaming/Source/PixelStreaming/Private/Tests/HandshakeTest.cpp
@@ -22,9 +22,8 @@ namespace UE::PixelStreaming
 		});
 
 		OutStreamer->StartStreaming();
-		OutPlayer->Connect(PlayerPort);
 
-		ADD_LATENT_AUTOMATION_COMMAND(FWaitForStreamerConnectedOrTimeout(5.0, OutStreamer))
+		ADD_LATENT_AUTOMATION_COMMAND(FConnectPlayerAfterStreamerConnectedOrTimeout(5.0, OutStreamer, OutPlayer, PlayerPort))
 		ADD_LATENT_AUTOMATION_COMMAND(FWaitForPlayerConnectedOrTimeout(5.0, OutPlayer, PlayerPort))
 		ADD_LATENT_AUTOMATION_COMMAND(FWaitForICEConnectedOrTimeout(5.0, OutPlayer))
 		ADD_LATENT_AUTOMATION_COMMAND(FCleanupAll(OutSignallingServer, OutStreamer, OutPlayer))
diff --git a/Engine/Plugins/Media/PixelStreaming/Source/PixelStreaming/Private/Tests/ProtocolTest.cpp b/Engine/Plugins/Media/PixelStreaming/Source/PixelStreaming/Private/Tests/ProtocolTest.cpp
--- a/Engine/Plugins/Media/PixelStreaming/Source/PixelStreaming/Private/Tests/ProtocolTest.cpp
+++ b/Engine/Plugins/Media/PixelStreaming/Source/PixelStreaming/Private/Tests/ProtocolTest.cpp
@@ -71,9 +71,8 @@ namespace UE::PixelStreaming
         };
 
 		Streamer->StartStreaming();
-		Player->Connect(PlayerPort);
 
-		ADD_LATENT_AUTOMATION_COMMAND(FWaitForStreamerConnectedOrTimeout(5.0, Streamer))
+		ADD_LATENT_AUTOMATION_COMMAND(FConnectPlayerAfterStreamerConnectedOrTimeout(5.0, Streamer, Player, PlayerPort))
 		ADD_LATENT_AUTOMATION_COMMAND(FWaitForPlayerConnectedOrTimeout(5.0, Player, PlayerPort))
 		ADD_LATENT_AUTOMATION_COMMAND(FWaitForICEConnectedOrTimeout(5.0, Player))
         ADD_LATENT_AUTOMATION_COMMAND(FWaitForDataChannelMessageOrTimeout(15.0, Player, Callback, bComplete))
@@ -116,9 +115,8 @@ namespace UE::PixelStreaming
 		Player->VideoSink = VideoSink;
 
 		Streamer->StartStreaming();
-		Player->Connect(PlayerPort);
 
-		ADD_LATENT_AUTOMATION_COMMAND(FWaitForStreamerConnectedOrTimeout(5.0, Streamer))
+		ADD_LATENT_AUTOMATION_COMMAND(FConnectPlayerAfterStreamerConnectedOrTimeout(5.0, Streamer, Player, PlayerPort))
 		ADD_LATENT_AUTOMATION_COMMAND(FWaitForPlayerConnectedOrTimeout(5.0, Player, PlayerPort))
 		ADD_LATENT_AUTOMATION_COMMAND(FWaitForICEConnectedOrTimeout(5.0, Player))
         ADD_LATENT_AUTOMATION_COMMAND(FSendCustomMessageToStreamer(Player, Message.Id, 1337))
diff --git a/Engine/Plugins/Media/PixelStreaming/Source/PixelStreaming/Private/Tests/TestUtils.cpp b/Engine/Plugins/Media/PixelStreaming/Source/PixelStreaming/Private/Tests/TestUtils.cpp
--- a/Engine/Plugins/Media/PixelStreaming/Source/PixelStreaming/Private/Tests/TestUtils.cpp
+++ b/Engine/Plugins/Media/PixelStreaming/Source/PixelStreaming/Private/Tests/TestUtils.cpp
@@ -166,6 +166,24 @@ namespace UE::PixelStreaming
 		return false;
 	}
 
+	bool FConnectPlayerAfterStreamerConnectedOrTimeout::Update()
+	{
+		if (OutStreamer->IsSignallingConnected())
+		{
+			UE_LOG(LogPixelStreaming, Log, TEXT("Streamer connected to signalling server. Attempting to connect player..."));
+			OutPlayer->Connect(PlayerPort);
+			return true;
+		}
+
+		double DeltaTime = FPlatformTime::Seconds() - StartTime;
+		if (DeltaTime > TimeoutSeconds)
+		{
+			UE_LOG(LogPixelStreaming, Error, TEXT("Timed out waiting for streamer to connect to signalling server."));
+			return true;
+		}
+		return false;		
+	}
+
 	bool FWaitForPlayerConnectedOrTimeout::Update()
 	{
 		if (OutPlayer->IsSignallingConnected())
diff --git a/Engine/Plugins/Media/PixelStreaming/Source/PixelStreaming/Private/Tests/TestUtils.h b/Engine/Plugins/Media/PixelStreaming/Source/PixelStreaming/Private/Tests/TestUtils.h
--- a/Engine/Plugins/Media/PixelStreaming/Source/PixelStreaming/Private/Tests/TestUtils.h
+++ b/Engine/Plugins/Media/PixelStreaming/Source/PixelStreaming/Private/Tests/TestUtils.h
@@ -43,6 +43,7 @@ namespace UE::PixelStreaming
 	DEFINE_LATENT_AUTOMATION_COMMAND_THREE_PARAMETER(FWaitForFrameReceived, double, TimeoutSeconds, TSharedPtr<FMockVideoSink>, VideoSink, FMockVideoFrameConfig, FrameConfig);
 	DEFINE_LATENT_AUTOMATION_COMMAND_TWO_PARAMETER(FWaitForICEConnectedOrTimeout, double, TimeoutSeconds, TSharedPtr<FMockPlayer>, OutPlayer);
 	DEFINE_LATENT_AUTOMATION_COMMAND_TWO_PARAMETER(FWaitForStreamerConnectedOrTimeout, double, TimeoutSeconds, TSharedPtr<IPixelStreamingStreamer>, OutStreamer);
+	DEFINE_LATENT_AUTOMATION_COMMAND_FOUR_PARAMETER(FConnectPlayerAfterStreamerConnectedOrTimeout, double, TimeoutSeconds, TSharedPtr<IPixelStreamingStreamer>, OutStreamer, TSharedPtr<FMockPlayer>, OutPlayer, int, PlayerPort);
 	DEFINE_LATENT_AUTOMATION_COMMAND_THREE_PARAMETER(FWaitForPlayerConnectedOrTimeout, double, TimeoutSeconds, TSharedPtr<FMockPlayer>, OutPlayer, int, PlayerPort);
 	DEFINE_LATENT_AUTOMATION_COMMAND_THREE_PARAMETER(FCleanupAll, TSharedPtr<UE::PixelStreamingServers::IServer>, OutSignallingServer, TSharedPtr<IPixelStreamingStreamer>, OutStreamer, TSharedPtr<FMockPlayer>, OutPlayer);
 
