Fixed a compile issue with PixelStreaming not finding a private header in Renderer.

diff --git a/Engine/Plugins/Media/PixelStreaming/Source/PixelStreaming/PixelStreaming.Build.cs b/Engine/Plugins/Media/PixelStreaming/Source/PixelStreaming/PixelStreaming.Build.cs
--- a/Engine/Plugins/Media/PixelStreaming/Source/PixelStreaming/PixelStreaming.Build.cs
+++ b/Engine/Plugins/Media/PixelStreaming/Source/PixelStreaming/PixelStreaming.Build.cs
@@ -72,6 +72,7 @@ namespace UnrealBuildTool.Rules
 			PrivateIncludePaths.AddRange(new string[]
 			{
 				Path.Combine(EngineDir, "Source/Runtime/AudioMixer/Private"),
+				Path.Combine(EngineDir, "Source/Runtime/Renderer/Private"),
 			});
 
 			// WebRTC third party includes (just libyuv for colour format conversions for now)
