Fixed linux editor compile issues.

diff --git a/Engine/Plugins/Media/PixelCapture/Source/PixelCapture/PixelCapture.build.cs b/Engine/Plugins/Media/PixelCapture/Source/PixelCapture/PixelCapture.build.cs
--- a/Engine/Plugins/Media/PixelCapture/Source/PixelCapture/PixelCapture.build.cs
+++ b/Engine/Plugins/Media/PixelCapture/Source/PixelCapture/PixelCapture.build.cs
@@ -18,7 +18,8 @@ namespace UnrealBuildTool.Rules
 
 			// NOTE: General rule is not to access the private folder of another module
 			PrivateIncludePaths.AddRange(new string[] {
-				Path.Combine(EngineDir, "Source/ThirdParty/WebRTC/4147/Include/third_party/libyuv/include")
+				Path.Combine(EngineDir, "Source/ThirdParty/WebRTC/4147/Include/third_party/libyuv/include"),
+				Path.Combine(EngineDir, "Source/Runtime/Renderer/Private"),
 			});
 
 			PrivateDependencyModuleNames.AddRange(new string[] {
diff --git a/Engine/Plugins/Media/PixelStreaming/Source/PixelStreamingEditor/PixelStreamingEditor.build.cs b/Engine/Plugins/Media/PixelStreaming/Source/PixelStreamingEditor/PixelStreamingEditor.build.cs
--- a/Engine/Plugins/Media/PixelStreaming/Source/PixelStreamingEditor/PixelStreamingEditor.build.cs
+++ b/Engine/Plugins/Media/PixelStreaming/Source/PixelStreamingEditor/PixelStreamingEditor.build.cs
@@ -16,7 +16,8 @@ namespace UnrealBuildTool.Rules
 			// NOTE: General rule is not to access the private folder of another module
 			PrivateIncludePaths.AddRange(new string[]
 			{
-				Path.Combine(EngineDir, "Plugins/Media/PixelStreaming/Source/PixelStreaming/Private")
+				Path.Combine(EngineDir, "Plugins/Media/PixelStreaming/Source/PixelStreaming/Private"),
+				Path.Combine(EngineDir, "Source/Runtime/Renderer/Private"),
 			});
 
 			PrivateDependencyModuleNames.AddRange(new string[]
diff --git a/Engine/Plugins/Runtime/nDisplay/Source/DisplayClusterShaders/DisplayClusterShaders.Build.cs b/Engine/Plugins/Runtime/nDisplay/Source/DisplayClusterShaders/DisplayClusterShaders.Build.cs
--- a/Engine/Plugins/Runtime/nDisplay/Source/DisplayClusterShaders/DisplayClusterShaders.Build.cs
+++ b/Engine/Plugins/Runtime/nDisplay/Source/DisplayClusterShaders/DisplayClusterShaders.Build.cs
@@ -31,6 +31,11 @@ public class DisplayClusterShaders : ModuleRules
 		{
 			PrivateDependencyModuleNames.Add("UnrealEd");
 		}
+		
+		// NOTE: General rule is not to access the private folder of another module
+		PrivateIncludePaths.AddRange(new string[] {
+			Path.Combine(Path.GetFullPath(Target.RelativeEnginePath), "Source/Runtime/Renderer/Private"),
+		});
 
 		AddEngineThirdPartyPrivateStaticDependencies(Target, "UElibPNG");
 		AddEngineThirdPartyPrivateStaticDependencies(Target, "zlib");
