Fix race condition causing encoders to be created in parallel when only one should be created.

diff --git a/Engine/Plugins/Media/PixelStreaming/Source/PixelStreaming/Private/VideoEncoderFactorySingleLayer.cpp b/Engine/Plugins/Media/PixelStreaming/Source/PixelStreaming/Private/VideoEncoderFactorySingleLayer.cpp
--- a/Engine/Plugins/Media/PixelStreaming/Source/PixelStreaming/Private/VideoEncoderFactorySingleLayer.cpp
+++ b/Engine/Plugins/Media/PixelStreaming/Source/PixelStreaming/Private/VideoEncoderFactorySingleLayer.cpp
@@ -149,7 +149,13 @@ namespace UE::PixelStreaming
 
 	FVideoEncoderWrapperHardware* FVideoEncoderFactorySingleLayer::GetOrCreateHardwareEncoder(int Width, int Height, int MaxBitrate, int TargetBitrate, int MaxFramerate)
 	{
-		if (HardwareEncoder == nullptr)
+		FScopeLock InitLock(&InitEncoderGuard);
+
+		if (HardwareEncoder != nullptr)
+ 		{
+			return HardwareEncoder.Get();
+		}
+		else
 		{
 			// Make AVEncoder frame factory.
 			TUniquePtr<FEncoderFrameFactory> FrameFactory = MakeUnique<FEncoderFrameFactory>();
@@ -183,10 +189,6 @@ namespace UE::PixelStreaming
 				return nullptr;
 			}
 		}
-		else
-		{
-			return HardwareEncoder.Get();
-		}
 	}
 
 	FVideoEncoderWrapperHardware* FVideoEncoderFactorySingleLayer::GetHardwareEncoder()
diff --git a/Engine/Plugins/Media/PixelStreaming/Source/PixelStreaming/Private/VideoEncoderFactorySingleLayer.h b/Engine/Plugins/Media/PixelStreaming/Source/PixelStreaming/Private/VideoEncoderFactorySingleLayer.h
--- a/Engine/Plugins/Media/PixelStreaming/Source/PixelStreaming/Private/VideoEncoderFactorySingleLayer.h
+++ b/Engine/Plugins/Media/PixelStreaming/Source/PixelStreaming/Private/VideoEncoderFactorySingleLayer.h
@@ -53,5 +53,8 @@ namespace UE::PixelStreaming
 		// raw pointers here because the factory interface wants a unique_ptr here so we cant own the object.
 		TArray<FVideoEncoderSingleLayerH264*> ActiveEncoders;
 		FCriticalSection ActiveEncodersGuard;
+
+		// Init encoder guard
+		FCriticalSection InitEncoderGuard;
 	};
 } // namespace UE::PixelStreaming
