Cherrypicked from 21754624, 21960483: Adding a blueprint wrapper to streamer.
Fix include file casing.

diff --git a/Engine/Plugins/Media/PixelStreaming/PixelStreaming.uplugin b/Engine/Plugins/Media/PixelStreaming/PixelStreaming.uplugin
--- a/Engine/Plugins/Media/PixelStreaming/PixelStreaming.uplugin
+++ b/Engine/Plugins/Media/PixelStreaming/PixelStreaming.uplugin
@@ -35,6 +35,20 @@
             "LoadingPhase": "PreDefault",
             "PlatformAllowList": [ "Win64", "Linux" ],
             "TargetDenyList": [ "Server" ]
+		},
+		{
+			"Name": "PixelStreamingBlueprint",
+			"Type": "Runtime",
+			"LoadingPhase": "PreDefault",
+			"PlatformAllowList": [ "Win64", "Linux" ],
+			"TargetDenyList": [ "Server" ]
+		},
+		{
+			"Name": "PixelStreamingBlueprintEditor",
+			"Type": "Editor",
+			"LoadingPhase": "PostEngineInit",
+			"PlatformAllowList": [ "Win64", "Linux" ],
+			"TargetDenyList": [ "Server" ]
 		}
 	],
 	"Plugins": [
diff --git a/Engine/Plugins/Media/PixelStreaming/Source/PixelStreamingBlueprint/PixelStreamingBlueprint.Build.cs b/Engine/Plugins/Media/PixelStreaming/Source/PixelStreamingBlueprint/PixelStreamingBlueprint.Build.cs
new file mode 100644
--- /dev/null
+++ b/Engine/Plugins/Media/PixelStreaming/Source/PixelStreamingBlueprint/PixelStreamingBlueprint.Build.cs
@@ -0,0 +1,40 @@
+// Copyright Epic Games, Inc. All Rights Reserved.
+
+using System;
+using System.IO;
+using System.Collections.Generic;
+using EpicGames.Core;
+
+namespace UnrealBuildTool.Rules
+{
+	public class PixelStreamingBlueprint : ModuleRules
+	{
+		public PixelStreamingBlueprint(ReadOnlyTargetRules Target) : base(Target)
+		{
+			//PCHUsage = PCHUsageMode.UseExplicitOrSharedPCHs;
+			//PrivatePCHHeaderFile = "Private/PCH.h";
+
+			var EngineDir = Path.GetFullPath(Target.RelativeEnginePath);
+
+			PublicIncludePaths.AddRange(
+				new string[] {
+				});
+
+			PrivateIncludePaths.AddRange(
+				new string[] {
+				});
+
+			PublicDependencyModuleNames.AddRange(
+				new string[] {
+				});
+
+			PrivateDependencyModuleNames.AddRange(
+				new string[] {
+					"Core",
+					"CoreUObject",
+					"Engine",
+					"PixelStreaming",
+				});
+		}
+	}
+}
diff --git a/Engine/Plugins/Media/PixelStreaming/Source/PixelStreamingBlueprint/Private/PixelStreamingBlueprintModule.cpp b/Engine/Plugins/Media/PixelStreaming/Source/PixelStreamingBlueprint/Private/PixelStreamingBlueprintModule.cpp
new file mode 100644
--- /dev/null
+++ b/Engine/Plugins/Media/PixelStreaming/Source/PixelStreamingBlueprint/Private/PixelStreamingBlueprintModule.cpp
@@ -0,0 +1,17 @@
+// Copyright Epic Games, Inc. All Rights Reserved.
+
+#include "PixelStreamingBlueprintPrivate.h"
+#include "Modules/ModuleManager.h"
+
+DEFINE_LOG_CATEGORY(LogPixelStreamingBlueprint);
+
+class FPixelStreamingBlueprintModule : public IModuleInterface
+{
+public:
+private:
+	/** IModuleInterface implementation */
+	void StartupModule() override {}
+	void ShutdownModule() override {}
+};
+
+IMPLEMENT_MODULE(FPixelStreamingBlueprintModule, PixelStreamingBlueprint)
diff --git a/Engine/Plugins/Media/PixelStreaming/Source/PixelStreamingBlueprint/Private/PixelStreamingBlueprintPrivate.h b/Engine/Plugins/Media/PixelStreaming/Source/PixelStreamingBlueprint/Private/PixelStreamingBlueprintPrivate.h
new file mode 100644
--- /dev/null
+++ b/Engine/Plugins/Media/PixelStreaming/Source/PixelStreamingBlueprint/Private/PixelStreamingBlueprintPrivate.h
@@ -0,0 +1,7 @@
+// Copyright Epic Games, Inc. All Rights Reserved.
+
+#pragma once
+
+#include "Logging/LogMacros.h"
+
+DECLARE_LOG_CATEGORY_EXTERN(LogPixelStreamingBlueprint, Log, All);
diff --git a/Engine/Plugins/Media/PixelStreaming/Source/PixelStreamingBlueprint/Private/PixelStreamingStreamerComponent.cpp b/Engine/Plugins/Media/PixelStreaming/Source/PixelStreamingBlueprint/Private/PixelStreamingStreamerComponent.cpp
new file mode 100644
--- /dev/null
+++ b/Engine/Plugins/Media/PixelStreaming/Source/PixelStreamingBlueprint/Private/PixelStreamingStreamerComponent.cpp
@@ -0,0 +1,112 @@
+// Copyright Epic Games, Inc. All Rights Reserved.
+
+#include "PixelStreamingStreamerComponent.h"
+#include "IPixelStreamingModule.h"
+#include "PixelStreamingVideoInputBackBuffer.h"
+#include "Engine/GameEngine.h"
+#include "Slate/SceneViewport.h"
+
+UPixelStreamingStreamerComponent::UPixelStreamingStreamerComponent(const FObjectInitializer& ObjectInitializer)
+	: Super(ObjectInitializer)
+{
+	FCoreDelegates::OnFEngineLoopInitComplete.AddLambda([this]() {
+		OnPostEngineInit();
+	});
+}
+
+FString UPixelStreamingStreamerComponent::GetId()
+{
+	return Streamer->GetId();
+}
+
+bool UPixelStreamingStreamerComponent::IsSignallingConnected()
+{
+	return Streamer->IsSignallingConnected();
+}
+
+void UPixelStreamingStreamerComponent::StartStreaming()
+{
+	if (!Streamer)
+	{
+		CreateStreamer();
+	}
+
+	if (StreamerInput && StreamerInput->GetVideoInput())
+	{
+		Streamer->SetVideoInput(StreamerInput->GetVideoInput());
+	}
+
+	Streamer->SetCoupleFramerate(CoupleFramerate);
+	Streamer->SetStreamFPS(StreamFPS);
+	Streamer->SetSignallingServerURL(SignallingServerURL);
+	Streamer->StartStreaming();
+}
+
+void UPixelStreamingStreamerComponent::StopStreaming()
+{
+	Streamer->StopStreaming();
+}
+
+bool UPixelStreamingStreamerComponent::IsStreaming()
+{
+	return Streamer->IsStreaming();
+}
+
+void UPixelStreamingStreamerComponent::ForceKeyFrame()
+{
+	Streamer->ForceKeyFrame();
+}
+
+void UPixelStreamingStreamerComponent::FreezeStream(UTexture2D* Texture)
+{
+	Streamer->FreezeStream(Texture);
+}
+
+void UPixelStreamingStreamerComponent::UnfreezeStream()
+{
+	Streamer->UnfreezeStream();
+}
+
+void UPixelStreamingStreamerComponent::SendPlayerMessage(uint8 Type, const FString& Descriptor)
+{
+	Streamer->SendPlayerMessage(Type, Descriptor);
+}
+
+void UPixelStreamingStreamerComponent::OnPostEngineInit()
+{
+	if (Streamer)
+	{
+		SetupStreamerInput();
+	}
+	bEngineStarted = true;
+}
+
+void UPixelStreamingStreamerComponent::CreateStreamer()
+{
+	Streamer = IPixelStreamingModule::Get().CreateStreamer(StreamerId);
+	Streamer->OnStreamingStarted().AddLambda([&](IPixelStreamingStreamer*) { OnStreamingStarted.Broadcast(); });
+	Streamer->OnStreamingStopped().AddLambda([&](IPixelStreamingStreamer*) { OnStreamingStopped.Broadcast(); });
+	Streamer->OnInputReceived.AddLambda([&](FPixelStreamingPlayerId PlayerId, uint8 Type, TArray<uint8> Data) { OnInputReceived.Broadcast(PlayerId, Type, Data); });
+
+	if (bEngineStarted)
+	{
+		SetupStreamerInput();
+	}
+}
+
+void UPixelStreamingStreamerComponent::SetupStreamerInput()
+{
+	if (!GIsEditor)
+	{
+		// default to the scene viewport if we have a game engine
+		if (UGameEngine* GameEngine = Cast<UGameEngine>(GEngine))
+		{
+			TSharedPtr<FSceneViewport> TargetViewport = GameEngine->SceneViewport;
+			if (TargetViewport.IsValid())
+			{
+				Streamer->SetTargetViewport(TargetViewport->GetViewportWidget());
+				Streamer->SetTargetWindow(TargetViewport->FindWindow());
+			}
+		}
+	}
+}
\ No newline at end of file
diff --git a/Engine/Plugins/Media/PixelStreaming/Source/PixelStreamingBlueprint/Private/PixelStreamingStreamerInput.cpp b/Engine/Plugins/Media/PixelStreaming/Source/PixelStreamingBlueprint/Private/PixelStreamingStreamerInput.cpp
new file mode 100644
--- /dev/null
+++ b/Engine/Plugins/Media/PixelStreaming/Source/PixelStreamingBlueprint/Private/PixelStreamingStreamerInput.cpp
@@ -0,0 +1,8 @@
+// Copyright Epic Games, Inc. All Rights Reserved.
+
+#include "PixelStreamingStreamerInput.h"
+
+UPixelStreamingStreamerInput::UPixelStreamingStreamerInput(const FObjectInitializer& ObjectInitializer)
+	: Super(ObjectInitializer)
+{
+}
diff --git a/Engine/Plugins/Media/PixelStreaming/Source/PixelStreamingBlueprint/Private/PixelStreamingStreamerInputBackBuffer.cpp b/Engine/Plugins/Media/PixelStreaming/Source/PixelStreamingBlueprint/Private/PixelStreamingStreamerInputBackBuffer.cpp
new file mode 100644
--- /dev/null
+++ b/Engine/Plugins/Media/PixelStreaming/Source/PixelStreamingBlueprint/Private/PixelStreamingStreamerInputBackBuffer.cpp
@@ -0,0 +1,10 @@
+// Copyright Epic Games, Inc. All Rights Reserved.
+
+#include "PixelStreamingStreamerInputBackBuffer.h"
+#include "PixelStreamingVideoInputBackBuffer.h"
+
+UPixelStreamingStreamerInputBackBuffer::UPixelStreamingStreamerInputBackBuffer(const FObjectInitializer& ObjectInitializer)
+	: Super(ObjectInitializer)
+{
+	VideoInput = FPixelStreamingVideoInputBackBuffer::Create();
+}
diff --git a/Engine/Plugins/Media/PixelStreaming/Source/PixelStreamingBlueprint/Public/PixelStreamingStreamerComponent.h b/Engine/Plugins/Media/PixelStreaming/Source/PixelStreamingBlueprint/Public/PixelStreamingStreamerComponent.h
new file mode 100644
--- /dev/null
+++ b/Engine/Plugins/Media/PixelStreaming/Source/PixelStreamingBlueprint/Public/PixelStreamingStreamerComponent.h
@@ -0,0 +1,75 @@
+// Copyright Epic Games, Inc. All Rights Reserved.
+
+#pragma once
+
+#include "IPixelStreamingStreamer.h"
+#include "PixelStreamingPlayerId.h"
+#include "Components/ActorComponent.h"
+#include "PixelStreamingStreamerInput.h"
+#include "PixelStreamingStreamerComponent.generated.h"
+
+UCLASS(BlueprintType, Blueprintable, Category = "PixelStreaming", META = (DisplayName = "Streamer Component", BlueprintSpawnableComponent))
+class PIXELSTREAMINGBLUEPRINT_API UPixelStreamingStreamerComponent : public UActorComponent
+{
+	GENERATED_UCLASS_BODY()
+
+public:
+	UFUNCTION(BlueprintCallable, Category = "PixelStreaming")
+	FString GetId();
+
+	UFUNCTION(BlueprintCallable, Category = "PixelStreaming")
+	bool IsSignallingConnected();
+
+	UFUNCTION(BlueprintCallable, Category = "PixelStreaming")
+	void StartStreaming();
+
+	UFUNCTION(BlueprintCallable, Category = "PixelStreaming")
+	void StopStreaming();
+
+	UFUNCTION(BlueprintCallable, Category = "PixelStreaming")
+	bool IsStreaming();
+
+	DECLARE_EVENT(UPixelStreamingStreamerComponent, FStreamingStartedEvent);
+	FStreamingStartedEvent OnStreamingStarted;
+
+	DECLARE_EVENT(UPixelStreamingStreamerComponent, FStreamingStoppedEvent);
+	FStreamingStoppedEvent OnStreamingStopped;
+
+	DECLARE_MULTICAST_DELEGATE_ThreeParams(FOnInputReceived, FPixelStreamingPlayerId, uint8, TArray<uint8>);
+	FOnInputReceived OnInputReceived;
+
+	UFUNCTION(BlueprintCallable, Category = "PixelStreaming")
+	void ForceKeyFrame();
+
+	UFUNCTION(BlueprintCallable, Category = "PixelStreaming")
+	void FreezeStream(UTexture2D* Texture);
+
+	UFUNCTION(BlueprintCallable, Category = "PixelStreaming")
+	void UnfreezeStream();
+
+	UFUNCTION(BlueprintCallable, Category = "PixelStreaming")
+	void SendPlayerMessage(uint8 Type, const FString& Descriptor);
+
+	UPROPERTY(EditAnywhere, Category = "PixelStreaming")
+	FString StreamerId = "Streamer Component";
+
+	UPROPERTY(EditAnywhere, Category = "PixelStreaming")
+	FString SignallingServerURL = "ws://localhost:8888";
+
+	UPROPERTY(EditAnywhere, Category = "PixelStreaming")
+	int32 StreamFPS = 60;
+
+	UPROPERTY(EditAnywhere, Category = "PixelStreaming")
+	bool CoupleFramerate = false;
+
+	UPROPERTY(EditAnywhere, Category = "PixelStreaming")
+	TObjectPtr<UPixelStreamingStreamerInput> StreamerInput = nullptr;
+
+private:
+	TSharedPtr<IPixelStreamingStreamer> Streamer;
+	bool bEngineStarted = false;
+
+	void OnPostEngineInit();
+	void CreateStreamer();
+	void SetupStreamerInput();
+};
diff --git a/Engine/Plugins/Media/PixelStreaming/Source/PixelStreamingBlueprint/Public/PixelStreamingStreamerInput.h b/Engine/Plugins/Media/PixelStreaming/Source/PixelStreamingBlueprint/Public/PixelStreamingStreamerInput.h
new file mode 100644
--- /dev/null
+++ b/Engine/Plugins/Media/PixelStreaming/Source/PixelStreamingBlueprint/Public/PixelStreamingStreamerInput.h
@@ -0,0 +1,19 @@
+// Copyright Epic Games, Inc. All Rights Reserved.
+
+#pragma once
+
+#include "Engine/Texture2DDynamic.h"
+#include "PixelStreamingVideoInput.h"
+#include "PixelStreamingStreamerInput.generated.h"
+
+UCLASS(NotBlueprintType, NotBlueprintable, Category = "PixelStreaming", META = (DisplayName = "Streamer Input"))
+class PIXELSTREAMINGBLUEPRINT_API UPixelStreamingStreamerInput : public UObject
+{
+	GENERATED_UCLASS_BODY()
+
+public:
+	TSharedPtr<FPixelStreamingVideoInput> GetVideoInput() { return VideoInput; }
+
+protected:
+	TSharedPtr<FPixelStreamingVideoInput> VideoInput;
+};
diff --git a/Engine/Plugins/Media/PixelStreaming/Source/PixelStreamingBlueprint/Public/PixelStreamingStreamerInputBackBuffer.h b/Engine/Plugins/Media/PixelStreaming/Source/PixelStreamingBlueprint/Public/PixelStreamingStreamerInputBackBuffer.h
new file mode 100644
--- /dev/null
+++ b/Engine/Plugins/Media/PixelStreaming/Source/PixelStreamingBlueprint/Public/PixelStreamingStreamerInputBackBuffer.h
@@ -0,0 +1,12 @@
+// Copyright Epic Games, Inc. All Rights Reserved.
+
+#pragma once
+
+#include "PixelStreamingStreamerInput.h"
+#include "PixelStreamingStreamerInputBackBuffer.generated.h"
+
+UCLASS(NotBlueprintType, NotBlueprintable, Category = "PixelStreaming", META = (DisplayName = "Back Buffer Streamer Input"))
+class PIXELSTREAMINGBLUEPRINT_API UPixelStreamingStreamerInputBackBuffer : public UPixelStreamingStreamerInput
+{
+	GENERATED_UCLASS_BODY()
+};
diff --git a/Engine/Plugins/Media/PixelStreaming/Source/PixelStreamingBlueprintEditor/PixelStreamingBlueprintEditor.Build.cs b/Engine/Plugins/Media/PixelStreaming/Source/PixelStreamingBlueprintEditor/PixelStreamingBlueprintEditor.Build.cs
new file mode 100644
--- /dev/null
+++ b/Engine/Plugins/Media/PixelStreaming/Source/PixelStreamingBlueprintEditor/PixelStreamingBlueprintEditor.Build.cs
@@ -0,0 +1,38 @@
+// Copyright Epic Games, Inc. All Rights Reserved.
+
+using System;
+using System.IO;
+using System.Collections.Generic;
+using EpicGames.Core;
+
+namespace UnrealBuildTool.Rules
+{
+	public class PixelStreamingBlueprintEditor : ModuleRules
+	{
+		public PixelStreamingBlueprintEditor(ReadOnlyTargetRules Target) : base(Target)
+		{
+			var EngineDir = Path.GetFullPath(Target.RelativeEnginePath);
+
+			PublicIncludePaths.AddRange(
+				new string[] {
+				});
+
+			PrivateIncludePaths.AddRange(
+				new string[] {
+				});
+
+			PublicDependencyModuleNames.AddRange(
+				new string[] {
+				});
+
+			PrivateDependencyModuleNames.AddRange(
+				new string[] {
+					"Core",
+					"CoreUObject",
+					"PixelStreaming",
+					"PixelStreamingBlueprint",
+					"UnrealEd"
+				});
+		}
+	}
+}
diff --git a/Engine/Plugins/Media/PixelStreaming/Source/PixelStreamingBlueprintEditor/Private/AssetTypeActions_StreamerInput.h b/Engine/Plugins/Media/PixelStreaming/Source/PixelStreamingBlueprintEditor/Private/AssetTypeActions_StreamerInput.h
new file mode 100644
--- /dev/null
+++ b/Engine/Plugins/Media/PixelStreaming/Source/PixelStreamingBlueprintEditor/Private/AssetTypeActions_StreamerInput.h
@@ -0,0 +1,21 @@
+// Copyright Epic Games, Inc. All Rights Reserved.
+
+#pragma once
+
+#include "CoreMinimal.h"
+#include "PixelStreamingStreamerInput.h"
+#include "Toolkits/IToolkitHost.h"
+#include "AssetTypeActions_Base.h"
+
+class FAssetTypeActions_StreamerInput : public FAssetTypeActions_Base
+{
+public:
+	// IAssetTypeActions Implementation
+	virtual FText GetName() const override { return NSLOCTEXT("PixelStreaming", "AssetTypeActions_StreamerInput", "Streamer Input Actions"); }
+	virtual FColor GetTypeColor() const override { return FColor(192,64,64); }
+	virtual UClass* GetSupportedClass() const override { return UPixelStreamingStreamerInput::StaticClass(); }
+	virtual bool HasActions ( const TArray<UObject*>& InObjects ) const override { return false; }
+	virtual bool CanFilter() override { return false; }
+	virtual uint32 GetCategories() override { return EAssetTypeCategories::Misc; }
+	virtual bool IsImportedAsset() const override { return false; }
+};
diff --git a/Engine/Plugins/Media/PixelStreaming/Source/PixelStreamingBlueprintEditor/Private/PixelStreamingBlueprintEditorModule.cpp b/Engine/Plugins/Media/PixelStreaming/Source/PixelStreamingBlueprintEditor/Private/PixelStreamingBlueprintEditorModule.cpp
new file mode 100644
--- /dev/null
+++ b/Engine/Plugins/Media/PixelStreaming/Source/PixelStreamingBlueprintEditor/Private/PixelStreamingBlueprintEditorModule.cpp
@@ -0,0 +1,19 @@
+// Copyright Epic Games, Inc. All Rights Reserved.
+
+#include "Modules/ModuleManager.h"
+#include "AssetTypeActions_StreamerInput.h"
+
+class FPixelStreamingBlueprintEditorModule : public IModuleInterface
+{
+public:
+private:
+	/** IModuleInterface implementation */
+	void StartupModule() override
+	{
+		IAssetTools& AssetTools = FModuleManager::LoadModuleChecked<FAssetToolsModule>("AssetTools").Get();
+		AssetTools.RegisterAssetTypeActions(MakeShared<FAssetTypeActions_StreamerInput>());
+	}
+	void ShutdownModule() override {}
+};
+
+IMPLEMENT_MODULE(FPixelStreamingBlueprintEditorModule, PixelStreamingBlueprintEditor)
diff --git a/Engine/Plugins/Media/PixelStreaming/Source/PixelStreamingBlueprintEditor/Private/PixelStreamingStreamerInputFactory.cpp b/Engine/Plugins/Media/PixelStreaming/Source/PixelStreamingBlueprintEditor/Private/PixelStreamingStreamerInputFactory.cpp
new file mode 100644
--- /dev/null
+++ b/Engine/Plugins/Media/PixelStreaming/Source/PixelStreamingBlueprintEditor/Private/PixelStreamingStreamerInputFactory.cpp
@@ -0,0 +1,39 @@
+// Copyright Epic Games, Inc. All Rights Reserved.
+
+#include "PixelStreamingStreamerInputFactory.h"
+
+#include <AssetTypeCategories.h>
+#include "PixelStreamingStreamerInputBackBuffer.h"
+
+#define LOCTEXT_NAMESPACE "PixelStreaming"
+
+UPixelStreamingStreamerInputBackBufferFactory::UPixelStreamingStreamerInputBackBufferFactory(const FObjectInitializer& ObjectInitializer)
+	: Super(ObjectInitializer)
+{
+	bCreateNew = true;
+	bEditAfterNew = true;
+
+	SupportedClass = UPixelStreamingStreamerInputBackBuffer::StaticClass();
+}
+
+FText UPixelStreamingStreamerInputBackBufferFactory::GetDisplayName() const
+{
+	return LOCTEXT("InputBackBufferFactoryDisplayName", "Back Buffer Streamer Input");
+}
+
+uint32 UPixelStreamingStreamerInputBackBufferFactory::GetMenuCategories() const
+{
+	return EAssetTypeCategories::Misc;
+}
+
+UObject* UPixelStreamingStreamerInputBackBufferFactory::FactoryCreateNew(UClass* InClass, UObject* InParent, FName InName, EObjectFlags Flags, UObject* Context, FFeedbackContext* Warn)
+{
+	if (UPixelStreamingStreamerInputBackBuffer* Resource = NewObject<UPixelStreamingStreamerInputBackBuffer>(InParent, InName, Flags | RF_Transactional))
+	{
+		return Resource;
+	}
+
+	return nullptr;
+}
+
+#undef LOCTEXT_NAMESPACE
diff --git a/Engine/Plugins/Media/PixelStreaming/Source/PixelStreamingBlueprintEditor/Private/PixelStreamingStreamerInputFactory.h b/Engine/Plugins/Media/PixelStreaming/Source/PixelStreamingBlueprintEditor/Private/PixelStreamingStreamerInputFactory.h
new file mode 100644
--- /dev/null
+++ b/Engine/Plugins/Media/PixelStreaming/Source/PixelStreamingBlueprintEditor/Private/PixelStreamingStreamerInputFactory.h
@@ -0,0 +1,20 @@
+// Copyright Epic Games, Inc. All Rights Reserved.
+
+#pragma once
+
+#include <CoreMinimal.h>
+#include <Factories/Factory.h>
+#include <UObject/Object.h>
+
+#include "PixelStreamingStreamerInputFactory.generated.h"
+
+UCLASS()
+class UPixelStreamingStreamerInputBackBufferFactory : public UFactory
+{
+	GENERATED_UCLASS_BODY()
+
+	public:
+		virtual FText GetDisplayName() const override;
+		virtual uint32 GetMenuCategories() const override;
+		virtual UObject* FactoryCreateNew(UClass* InClass, UObject* InParent, FName InName, EObjectFlags Flags, UObject* Context, FFeedbackContext* Warn) override;
+};
diff --git a/Engine/Plugins/Media/PixelStreamingPlayer/Source/PixelStreamingPlayerEditor/PixelStreamingPlayerEditor.Build.cs b/Engine/Plugins/Media/PixelStreamingPlayer/Source/PixelStreamingPlayerEditor/PixelStreamingPlayerEditor.Build.cs
--- a/Engine/Plugins/Media/PixelStreamingPlayer/Source/PixelStreamingPlayerEditor/PixelStreamingPlayerEditor.Build.cs
+++ b/Engine/Plugins/Media/PixelStreamingPlayer/Source/PixelStreamingPlayerEditor/PixelStreamingPlayerEditor.Build.cs
@@ -29,6 +29,7 @@ namespace UnrealBuildTool.Rules
 				new string[] {
 					"Core",
 					"CoreUObject",
+					"PixelStreaming",
 					"PixelStreamingPlayer",
 					"UnrealEd"
 				});
diff --git a/Engine/Plugins/Media/PixelStreamingPlayer/Source/PixelStreamingPlayerEditor/Private/PixelStreamingMediaTextureFactory.cpp b/Engine/Plugins/Media/PixelStreamingPlayer/Source/PixelStreamingPlayerEditor/Private/PixelStreamingMediaTextureFactory.cpp
--- a/Engine/Plugins/Media/PixelStreamingPlayer/Source/PixelStreamingPlayerEditor/Private/PixelStreamingMediaTextureFactory.cpp
+++ b/Engine/Plugins/Media/PixelStreamingPlayer/Source/PixelStreamingPlayerEditor/Private/PixelStreamingMediaTextureFactory.cpp
@@ -5,7 +5,7 @@
 #include <AssetTypeCategories.h>
 #include "PixelStreamingMediaTexture.h"
 
-#define LOCTEXT_NAMESPACE "PixelStreamingPlayerEditorMediaTextureFactory"
+#define LOCTEXT_NAMESPACE "PixelStreaming"
 
 UPixelStreamingMediaTextureFactory::UPixelStreamingMediaTextureFactory(const FObjectInitializer& ObjectInitializer)
 	: Super(ObjectInitializer)
@@ -18,7 +18,7 @@ UPixelStreamingMediaTextureFactory::UPixelStreamingMediaTextureFactory(const FOb
 
 FText UPixelStreamingMediaTextureFactory::GetDisplayName() const
 {
-	return LOCTEXT("PixelStreamingMediaTextureFactoryDisplayName", "Pixel Streaming Media Texture");
+	return LOCTEXT("MediaTextureFactoryDisplayName", "Pixel Streaming Media Texture");
 }
 
 uint32 UPixelStreamingMediaTextureFactory::GetMenuCategories() const
