Cherrypicked from 21801996: Fixing some missed cleanup of destroyed components.

diff --git a/Engine/Plugins/Media/PixelStreaming/Source/PixelStreamingBlueprint/Private/PixelStreamingStreamerComponent.cpp b/Engine/Plugins/Media/PixelStreaming/Source/PixelStreamingBlueprint/Private/PixelStreamingStreamerComponent.cpp
--- a/Engine/Plugins/Media/PixelStreaming/Source/PixelStreamingBlueprint/Private/PixelStreamingStreamerComponent.cpp
+++ b/Engine/Plugins/Media/PixelStreaming/Source/PixelStreamingBlueprint/Private/PixelStreamingStreamerComponent.cpp
@@ -84,9 +84,9 @@ void UPixelStreamingStreamerComponent::OnPostEngineInit()
 void UPixelStreamingStreamerComponent::CreateStreamer()
 {
 	Streamer = IPixelStreamingModule::Get().CreateStreamer(StreamerId);
-	Streamer->OnStreamingStarted().AddLambda([&](IPixelStreamingStreamer*) { OnStreamingStarted.Broadcast(); });
-	Streamer->OnStreamingStopped().AddLambda([&](IPixelStreamingStreamer*) { OnStreamingStopped.Broadcast(); });
-	Streamer->OnInputReceived.AddLambda([&](FPixelStreamingPlayerId PlayerId, uint8 Type, TArray<uint8> Data) { OnInputReceived.Broadcast(PlayerId, Type, Data); });
+	Streamer->OnStreamingStarted().AddUObject(this, &UPixelStreamingStreamerComponent::StreamingStarted);
+	Streamer->OnStreamingStopped().AddUObject(this, &UPixelStreamingStreamerComponent::StreamingStopped);
+	Streamer->OnInputReceived.AddUObject(this, &UPixelStreamingStreamerComponent::StreamingInput);
 
 	if (bEngineStarted)
 	{
@@ -94,6 +94,21 @@ void UPixelStreamingStreamerComponent::CreateStreamer()
 	}
 }
 
+void UPixelStreamingStreamerComponent::StreamingStarted(IPixelStreamingStreamer*)
+{
+	OnStreamingStarted.Broadcast();
+}
+
+void UPixelStreamingStreamerComponent::StreamingStopped(IPixelStreamingStreamer*)
+{
+	OnStreamingStopped.Broadcast();
+}
+
+void UPixelStreamingStreamerComponent::StreamingInput(FPixelStreamingPlayerId PlayerId, uint8 Type, TArray<uint8> Data)
+{
+	OnInputReceived.Broadcast(PlayerId, Type, Data);
+}
+
 void UPixelStreamingStreamerComponent::SetupStreamerInput()
 {
 	if (!GIsEditor)
diff --git a/Engine/Plugins/Media/PixelStreaming/Source/PixelStreamingBlueprint/Public/PixelStreamingStreamerComponent.h b/Engine/Plugins/Media/PixelStreaming/Source/PixelStreamingBlueprint/Public/PixelStreamingStreamerComponent.h
--- a/Engine/Plugins/Media/PixelStreaming/Source/PixelStreamingBlueprint/Public/PixelStreamingStreamerComponent.h
+++ b/Engine/Plugins/Media/PixelStreaming/Source/PixelStreamingBlueprint/Public/PixelStreamingStreamerComponent.h
@@ -72,4 +72,8 @@ private:
 	void OnPostEngineInit();
 	void CreateStreamer();
 	void SetupStreamerInput();
+
+	void StreamingStarted(IPixelStreamingStreamer*);
+	void StreamingStopped(IPixelStreamingStreamer*);
+	void StreamingInput(FPixelStreamingPlayerId PlayerId, uint8 Type, TArray<uint8> Data);
 };
