Cherrypicked from 21801997: Fixing packaging with the new streamer blueprints.

diff --git a/Engine/Plugins/Media/PixelStreaming/Source/PixelStreaming/Private/PixelStreamingVideoInputBackBuffer.cpp b/Engine/Plugins/Media/PixelStreaming/Source/PixelStreaming/Private/PixelStreamingVideoInputBackBuffer.cpp
--- a/Engine/Plugins/Media/PixelStreaming/Source/PixelStreaming/Private/PixelStreamingVideoInputBackBuffer.cpp
+++ b/Engine/Plugins/Media/PixelStreaming/Source/PixelStreaming/Private/PixelStreamingVideoInputBackBuffer.cpp
@@ -12,6 +12,12 @@
 
 TSharedPtr<FPixelStreamingVideoInputBackBuffer> FPixelStreamingVideoInputBackBuffer::Create()
 {
+	// this was added to fix packaging
+	if (!FSlateApplication::IsInitialized())
+	{
+		return nullptr;
+	}
+
 	TSharedPtr<FPixelStreamingVideoInputBackBuffer> NewInput = TSharedPtr<FPixelStreamingVideoInputBackBuffer>(new FPixelStreamingVideoInputBackBuffer());
 	TWeakPtr<FPixelStreamingVideoInputBackBuffer> WeakInput = NewInput;
 
