Ensure we fire touch moved events every frame while the touch is down to prevent jerky movement

diff --git a/Engine/Plugins/Media/PixelStreaming/Source/PixelStreaming/Private/PixelStreamingInputHandler.cpp b/Engine/Plugins/Media/PixelStreaming/Source/PixelStreaming/Private/PixelStreamingInputHandler.cpp
--- a/Engine/Plugins/Media/PixelStreaming/Source/PixelStreaming/Private/PixelStreamingInputHandler.cpp
+++ b/Engine/Plugins/Media/PixelStreaming/Source/PixelStreaming/Private/PixelStreamingInputHandler.cpp
@@ -79,12 +79,16 @@ namespace UE::PixelStreaming
 
     void FPixelStreamingInputHandler::Tick(const float InDeltaTime)
     {
+    	TouchIndicesProcessedThisFrame.Reset();
+    	
         FMessage Message;
 		while (Messages.Dequeue(Message))
 		{
             FMemoryReader Ar(Message.Data);
             (*Message.Handler)(Ar);
         }
+
+    	BroadcastActiveTouchMoveEvents();
     }
 
 	void FPixelStreamingInputHandler::OnMessage(const webrtc::DataBuffer& Buffer)
@@ -271,8 +275,15 @@ namespace UE::PixelStreaming
 			{
 				//                                                                           convert range from 0,65536 -> 0,1
 				FVector2D TouchLocation = ConvertFromNormalizedScreenLocation(FVector2D(Touch.Param1 / uint16_MAX, Touch.Param2 / uint16_MAX));
+				const int32 TouchIndex = Touch.Param3;
                 UE_LOG(LogPixelStreamingInputHandler, Verbose, TEXT("TOUCH_MOVE: TouchIndex = %d; Pos = (%d, %d); CursorPos = (%d, %d); Force = %.3f"), Touch.Param3, Touch.Param1, Touch.Param2, static_cast<int>(TouchLocation.X), static_cast<int>(TouchLocation.Y),  Touch.Param4 / 255.0f);
-                MessageHandler->OnTouchMoved(TouchLocation, Touch.Param4 / 255.0f, Touch.Param3, 0); // TODO: ControllerId?
+
+				FCachedTouchEvent& TouchEvent = CachedTouchEvents.FindOrAdd(TouchIndex);
+				TouchEvent.Location = TouchLocation;
+				TouchEvent.Force = Touch.Param4 / 255.0f;
+				TouchEvent.ControllerIndex = 0;
+				MessageHandler->OnTouchMoved(TouchEvent.Location, TouchEvent.Force, TouchIndex, TouchEvent.ControllerIndex); // TODO: ControllerId?
+				TouchIndicesProcessedThisFrame.Add(TouchIndex);
             }
         }
     }
@@ -291,14 +302,15 @@ namespace UE::PixelStreaming
 			{
 				//                                                                           convert range from 0,65536 -> 0,1
 				FVector2D TouchLocation = ConvertFromNormalizedScreenLocation(FVector2D(Touch.Param1 / uint16_MAX, Touch.Param2 / uint16_MAX));
-
+				const int32 TouchIndex = Touch.Param3;
 				UE_LOG(LogPixelStreamingInputHandler, Verbose, TEXT("TOUCH_END: TouchIndex = %d; Pos = (%d, %d); CursorPos = (%d, %d)"), Touch.Param3, Touch.Param1, Touch.Param2, static_cast<int>(TouchLocation.X), static_cast<int>(TouchLocation.Y));
-				MessageHandler->OnTouchEnded(TouchLocation, Touch.Param3, 0); // TODO: ControllerId?
+				MessageHandler->OnTouchEnded(TouchLocation, TouchIndex, 0); // TODO: ControllerId?
+            	CachedTouchEvents.Remove(TouchIndex);
 				NumActiveTouches--;
 			}
 		}
 
-		// If there's no remaing touches, and there is also no mouse over the player window
+		// If there's no remaining touches, and there is also no mouse over the player window
 		// then set the platform application back to its default. We need to set it back to default
 		// so that people using the editor (if editor streaming) can click on buttons outside the target window
 		// and also have the correct cursor (pixel streaming forces default cursor)
@@ -748,6 +760,26 @@ namespace UE::PixelStreaming
 		}
 	}
 
+	void FPixelStreamingInputHandler::BroadcastActiveTouchMoveEvents()
+	{
+    	if (!ensure(MessageHandler))
+    	{
+    		return;
+    	}
+    	
+    	for (TPair<int32, FCachedTouchEvent> CachedTouchEvent : CachedTouchEvents)
+    	{
+    		const int32& TouchIndex = CachedTouchEvent.Key;
+    		const FCachedTouchEvent& TouchEvent = CachedTouchEvent.Value;
+
+    		// Only broadcast events that haven't already been fired this frame
+    		if (!TouchIndicesProcessedThisFrame.Contains(TouchIndex))
+    		{
+    			MessageHandler->OnTouchMoved(TouchEvent.Location, TouchEvent.Force, TouchIndex, TouchEvent.ControllerIndex);
+    		}
+    	}
+	}
+
 	void FPixelStreamingInputHandler::FindFocusedWidget()
 	{
 		FSlateApplication::Get().ForEachUser([this](FSlateUser& User) {
diff --git a/Engine/Plugins/Media/PixelStreaming/Source/PixelStreaming/Private/PixelStreamingInputHandler.h b/Engine/Plugins/Media/PixelStreaming/Source/PixelStreaming/Private/PixelStreamingInputHandler.h
--- a/Engine/Plugins/Media/PixelStreaming/Source/PixelStreaming/Private/PixelStreamingInputHandler.h
+++ b/Engine/Plugins/Media/PixelStreaming/Source/PixelStreaming/Private/PixelStreamingInputHandler.h
@@ -96,6 +96,22 @@ namespace UE::PixelStreaming
 		FGamepadKeyNames::Type ConvertAxisIndexToGamepadAxis(uint8 AnalogAxis);
 		FGamepadKeyNames::Type ConvertButtonIndexToGamepadButton(uint8 ButtonIndex);
 
+		struct FCachedTouchEvent
+		{
+			FVector2D Location;
+			float Force;
+			int32 ControllerIndex;
+		};
+
+		// Keep a cache of the last touch events as we need to fire Touch Moved every frame while touch is down
+		TMap<int32, FCachedTouchEvent> CachedTouchEvents;
+
+		// Track which touch events we processed this frame so we can avoid re-processing them
+		TSet<int32> TouchIndicesProcessedThisFrame;
+
+		// Sends Touch Moved events for any touch index which is currently down but wasn't already updated this frame
+		void BroadcastActiveTouchMoveEvents();
+
 		void FindFocusedWidget();
 		bool FilterKey(const FKey& Key);
 
