The Pixel Streaming toolbar has been updated for a better UX that incorporates the work done to link Pixel Streaming and Virtual cameras.

Cherry pick CLs 22088897, 22107075 and 22177749 to 5.1.

* There is now a single embedded signalling server associated with an instance of the editor that can be configured from the Pixel Streaming toolbar and launched from either a Virtual Camera's details panel or the toolbar itself.
* Users can choose to use a signalling server outside of the editor by toggling the "Use Remote Signalling Server" option in the toolbar and specifying the URL to this server.
* Added the official Pixel Streaming logo to the Plugins window and toolbar.

diff --git a/Engine/Plugins/Media/PixelStreaming/Source/PixelStreamingEditor/PixelStreamingEditor.build.cs b/Engine/Plugins/Media/PixelStreaming/Source/PixelStreamingEditor/PixelStreamingEditor.build.cs
--- a/Engine/Plugins/Media/PixelStreaming/Source/PixelStreamingEditor/PixelStreamingEditor.build.cs
+++ b/Engine/Plugins/Media/PixelStreaming/Source/PixelStreamingEditor/PixelStreamingEditor.build.cs
@@ -35,7 +35,10 @@ namespace UnrealBuildTool.Rules
 				"EngineSettings",
 				"InputCore",
 				"Json",
-				"PixelCapture"
+				"PixelCapture",
+				"PixelStreamingServers",
+				"HTTP",
+				"Sockets"
 			});
 
 			if(Target.bBuildEditor)
diff --git a/Engine/Plugins/Media/PixelStreaming/Source/PixelStreamingEditor/Private/PixelStreamingCommands.cpp b/Engine/Plugins/Media/PixelStreaming/Source/PixelStreamingEditor/Private/PixelStreamingCommands.cpp
--- a/Engine/Plugins/Media/PixelStreaming/Source/PixelStreamingEditor/Private/PixelStreamingCommands.cpp
+++ b/Engine/Plugins/Media/PixelStreaming/Source/PixelStreamingEditor/Private/PixelStreamingCommands.cpp
@@ -7,11 +7,15 @@ namespace UE::EditorPixelStreaming
 #define LOCTEXT_NAMESPACE "PixelStreamingToolBar"
     void FPixelStreamingCommands::RegisterCommands()
     {
-        UI_COMMAND(StartStreaming, "Start Streaming", "Start Streaming", EUserInterfaceActionType::Button, FInputChord());
+        UI_COMMAND(ExternalSignalling, "Use Remote Signalling Server", "Check this option if you wish to use a remote Signalling Server", EUserInterfaceActionType::RadioButton, FInputChord());
         UI_COMMAND(StopStreaming, "Stop Streaming", "Stop Streaming", EUserInterfaceActionType::Button, FInputChord());
-        UI_COMMAND(VP8, "VP8", "VP8", EUserInterfaceActionType::Check, FInputChord());
-        UI_COMMAND(VP9, "VP9", "VP9", EUserInterfaceActionType::Check, FInputChord());
-        UI_COMMAND(H264, "H264", "H264", EUserInterfaceActionType::Check, FInputChord());
+        UI_COMMAND(StartSignalling, "Launch Signalling Server", "Launch a Signalling Server that will listen for connections on the ports specified above", EUserInterfaceActionType::Button, FInputChord());
+        UI_COMMAND(StopSignalling, "Stop Signalling Server", "Stop Signalling Server", EUserInterfaceActionType::Button, FInputChord());
+        UI_COMMAND(StreamLevelEditor, "Stream Level Editor", "Stream the Level Editor viewport", EUserInterfaceActionType::Button, FInputChord());
+        UI_COMMAND(StreamEditor, "Stream Full Editor", "Stream the Full Editor", EUserInterfaceActionType::Button, FInputChord());
+        UI_COMMAND(VP8, "VP8", "VP8", EUserInterfaceActionType::RadioButton, FInputChord());
+        UI_COMMAND(VP9, "VP9", "VP9", EUserInterfaceActionType::RadioButton, FInputChord());
+        UI_COMMAND(H264, "H264", "H264", EUserInterfaceActionType::RadioButton, FInputChord());
     }
 #undef LOCTEXT_NAMESPACE
 }
\ No newline at end of file
diff --git a/Engine/Plugins/Media/PixelStreaming/Source/PixelStreamingEditor/Private/PixelStreamingCommands.h b/Engine/Plugins/Media/PixelStreaming/Source/PixelStreamingEditor/Private/PixelStreamingCommands.h
--- a/Engine/Plugins/Media/PixelStreaming/Source/PixelStreamingEditor/Private/PixelStreamingCommands.h
+++ b/Engine/Plugins/Media/PixelStreaming/Source/PixelStreamingEditor/Private/PixelStreamingCommands.h
@@ -21,10 +21,14 @@ namespace UE::EditorPixelStreaming
 
         virtual void RegisterCommands() override;
 
-        TSharedPtr<FUICommandInfo> StartStreaming;
+        TSharedPtr<FUICommandInfo> ExternalSignalling;
         TSharedPtr<FUICommandInfo> StopStreaming;
         TSharedPtr<FUICommandInfo> VP8;
         TSharedPtr<FUICommandInfo> VP9;
         TSharedPtr<FUICommandInfo> H264;
+        TSharedPtr<FUICommandInfo> StartSignalling;
+        TSharedPtr<FUICommandInfo> StopSignalling;
+        TSharedPtr<FUICommandInfo> StreamLevelEditor;
+        TSharedPtr<FUICommandInfo> StreamEditor;
     };
 }
\ No newline at end of file
diff --git a/Engine/Plugins/Media/PixelStreaming/Source/PixelStreamingEditor/Private/PixelStreamingEditorModule.cpp b/Engine/Plugins/Media/PixelStreaming/Source/PixelStreamingEditor/Private/PixelStreamingEditorModule.cpp
--- a/Engine/Plugins/Media/PixelStreaming/Source/PixelStreamingEditor/Private/PixelStreamingEditorModule.cpp
+++ b/Engine/Plugins/Media/PixelStreaming/Source/PixelStreamingEditor/Private/PixelStreamingEditorModule.cpp
@@ -7,11 +7,19 @@
 #include "PixelStreamingStyle.h"
 #include "IPixelStreamingStreamer.h"
 #include "PixelStreamingVideoInputBackBufferComposited.h"
+#include "PixelStreamingVideoInputViewport.h"
 #include "Settings.h"
 #include "UnrealEngine.h"
 #include "Interfaces/IMainFrameModule.h"
 #include "PixelStreamingProtocol.h"
 #include "Utils.h"
+#include "GenericPlatform/GenericPlatformHttp.h"
+#include "LevelEditorViewport.h"
+#include "Slate/SceneViewport.h"
+#include "SLevelViewport.h"
+
+DECLARE_LOG_CATEGORY_EXTERN(LogPixelStreamingEditor, Log, All);
+DEFINE_LOG_CATEGORY(LogPixelStreamingEditor);
 
 FPixelStreamingEditorModule* FPixelStreamingEditorModule::PixelStreamingEditorModule = nullptr;
 
@@ -94,7 +102,27 @@ void FPixelStreamingEditorModule::InitEditorStreaming(IPixelStreamingModule& Mod
 	// Give the editor streamer the default url if the user hasn't specified one when launching the editor
 	if (Streamer->GetSignallingServerURL().IsEmpty())
 	{
-		Streamer->SetSignallingServerURL(Module.GetDefaultSignallingURL());
+		// No URL was passed on the command line, initialize defaults
+		StreamerPort = 8888;
+		SignallingDomain = TEXT("ws://127.0.0.1");
+
+		Streamer->SetSignallingServerURL(FString::Printf(TEXT("%s:%d"), *SignallingDomain, StreamerPort));
+	} 
+	else
+	{
+		FString SpecifiedSignallingURL = Streamer->GetSignallingServerURL();
+		TOptional<uint16> ExtractedStreamerPort = FGenericPlatformHttp::GetUrlPort(SpecifiedSignallingURL);
+		StreamerPort = (int32) ExtractedStreamerPort.Get(8888);
+
+		FString ExtractedSignallingDomain = FGenericPlatformHttp::GetUrlDomain(SpecifiedSignallingURL);
+		if(FGenericPlatformHttp::IsSecureProtocol(SpecifiedSignallingURL).Get(false))
+		{
+			SignallingDomain = FString::Printf(TEXT("wss://%s"), *ExtractedSignallingDomain);
+		}
+		else
+		{
+			SignallingDomain = FString::Printf(TEXT("ws://%s"), *ExtractedSignallingDomain);
+		}
 	}
 
 	IMainFrameModule::Get().OnMainFrameCreationFinished().AddLambda([&](TSharedPtr<SWindow> RootWindow, bool bIsRunningStartupDialog)
@@ -107,12 +135,12 @@ void FPixelStreamingEditorModule::InitEditorStreaming(IPixelStreamingModule& Mod
 
 		if(UE::EditorPixelStreaming::Settings::CVarEditorPixelStreamingStartOnLaunch.GetValueOnAnyThread())
 		{
-			StartStreaming();
+			StartStreaming(UE::EditorPixelStreaming::EStreamTypes::Editor);
 		}
 	});
 }
 
-void FPixelStreamingEditorModule::StartStreaming()
+void FPixelStreamingEditorModule::StartStreaming(UE::EditorPixelStreaming::EStreamTypes InStreamType)
 {
 	// Activate our level editor streamer
 	IPixelStreamingModule& Module = IPixelStreamingModule::Get();
@@ -122,13 +150,44 @@ void FPixelStreamingEditorModule::StartStreaming()
 		return;
 	}
 
-	// The streamer doesn't currently have a video input so create one
-	TWeakPtr<FPixelStreamingVideoInput> VideoInput = Streamer->GetVideoInput();
-	if(!VideoInput.Pin())
+	StreamType = InStreamType;
+	switch(InStreamType)
 	{
-		Streamer->SetVideoInput(FPixelStreamingVideoInputBackBufferComposited::Create());
+		case UE::EditorPixelStreaming::EStreamTypes::LevelEditorViewport:
+		{
+			FLevelEditorModule& LevelEditorModule = FModuleManager::GetModuleChecked<FLevelEditorModule>("LevelEditor");
+			TSharedPtr<SLevelViewport> ActiveLevelViewport = LevelEditorModule.GetFirstActiveLevelViewport();
+			if (!ActiveLevelViewport.IsValid())
+			{
+				return;
+			}
+
+			FLevelEditorViewportClient& LevelViewportClient = ActiveLevelViewport->GetLevelViewportClient();
+			FSceneViewport* SceneViewport = static_cast<FSceneViewport*>(LevelViewportClient.Viewport);
+			Streamer->SetTargetViewport(SceneViewport->GetViewportWidget());
+			Streamer->SetTargetWindow(SceneViewport->FindWindow());
+			Streamer->SetVideoInput(FPixelStreamingVideoInputViewport::Create());
+		}	
+		break;
+		case UE::EditorPixelStreaming::EStreamTypes::Editor:
+		{
+			Streamer->SetTargetViewport(nullptr);
+			Streamer->SetTargetWindow(nullptr);
+			Streamer->SetVideoInput(FPixelStreamingVideoInputBackBufferComposited::Create());
+		}
+		break;
+		default:
+		{
+			UE_LOG(LogPixelStreamingEditor, Warning, TEXT("Specified Stream Type doesn't have an associated FPixelStreamingVideoInput"));
+		}
+		// Return here as we don't want to start streaming if we didn't set a viewport
+		return;
 	}
 
+	if(!bUseExternalSignallingServer)
+	{
+		Streamer->SetSignallingServerURL(FString::Printf(TEXT("%s:%d"), *SignallingDomain, StreamerPort));
+	}
 	// Use the Pixel Streaming module's start streaming method to start all streamers
 	Module.StartStreaming();
 }
@@ -147,6 +206,80 @@ void FPixelStreamingEditorModule::StopStreaming()
 	Module.StopStreaming();
 }
 
+UE::EditorPixelStreaming::EStreamTypes FPixelStreamingEditorModule::GetStreamType()
+{
+	return StreamType;
+}
+
+void FPixelStreamingEditorModule::StartSignalling()
+{
+	bool bAlreadyLaunched = SignallingServer.IsValid() && SignallingServer->HasLaunched();
+	if(bAlreadyLaunched)
+	{
+		return;
+	}
+
+	// Download Pixel Streaming servers/frontend if we want to use a browser to view Pixel Streaming output
+	// but only attempt this is we haven't already started a download before.
+	if(!DownloadProcess.IsValid())
+	{
+		DownloadProcess = UE::PixelStreamingServers::DownloadPixelStreamingServers(/*bSkipIfPresent*/ true);
+		if(DownloadProcess.IsValid())
+		{
+			DownloadProcess->OnCompleted().BindLambda([this](int ExitCode){
+				StopSignalling();
+				StartSignalling();
+			});
+
+			return;
+		}
+	}
+
+	// Launch signalling server
+	SignallingServer = UE::PixelStreamingServers::MakeSignallingServer();
+
+	UE::PixelStreamingServers::FLaunchArgs LaunchArgs;
+	LaunchArgs.bPollUntilReady = false;
+	LaunchArgs.ReconnectionTimeoutSeconds = 30.0f;
+	LaunchArgs.ReconnectionIntervalSeconds = 2.0f;
+	LaunchArgs.ProcessArgs = FString::Printf(TEXT("--HttpPort=%d --StreamerPort=%d"), ViewerPort, StreamerPort);
+
+	SignallingServer->Launch(LaunchArgs);
+}
+
+void FPixelStreamingEditorModule::StopSignalling()
+{
+	if(SignallingServer.IsValid())
+	{
+		SignallingServer->Stop();
+		SignallingServer.Reset();
+	}
+}
+
+TSharedPtr<UE::PixelStreamingServers::IServer> FPixelStreamingEditorModule::GetSignallingServer() 
+{
+	if(SignallingServer.IsValid())
+	{
+		return SignallingServer;
+	}
+	return nullptr;
+}
+
+void FPixelStreamingEditorModule::SetSignallingDomain(const FString& InSignallingDomain)
+{
+	SignallingDomain = InSignallingDomain;
+}
+
+void FPixelStreamingEditorModule::SetStreamerPort(int32 InStreamerPort)
+{
+	StreamerPort = InStreamerPort;
+}
+
+void FPixelStreamingEditorModule::SetViewerPort(int32 InViewerPort)
+{
+	ViewerPort = InViewerPort;
+}
+
 FPixelStreamingEditorModule* FPixelStreamingEditorModule::GetModule()
 {
 	if (PixelStreamingEditorModule)
diff --git a/Engine/Plugins/Media/PixelStreaming/Source/PixelStreamingEditor/Private/PixelStreamingStyle.cpp b/Engine/Plugins/Media/PixelStreaming/Source/PixelStreamingEditor/Private/PixelStreamingStyle.cpp
--- a/Engine/Plugins/Media/PixelStreaming/Source/PixelStreamingEditor/Private/PixelStreamingStyle.cpp
+++ b/Engine/Plugins/Media/PixelStreaming/Source/PixelStreamingEditor/Private/PixelStreamingStyle.cpp
@@ -33,22 +33,25 @@ namespace UE::EditorPixelStreaming
     }
 
 #define IMAGE_BRUSH( RelativePath, ... ) FSlateImageBrush( Style->RootToContentDir( RelativePath, TEXT(".png") ), __VA_ARGS__ )
+#define IMAGE_BRUSH_SVG( RelativePath, ... ) FSlateVectorImageBrush(Style->RootToContentDir(RelativePath, TEXT(".svg")), __VA_ARGS__)
 #define BOX_BRUSH( RelativePath, ... ) FSlateBoxBrush( Style->RootToContentDir( RelativePath, TEXT(".png") ), __VA_ARGS__ )
 #define BORDER_BRUSH( RelativePath, ... ) FSlateBorderBrush( Style->RootToContentDir( RelativePath, TEXT(".png") ), __VA_ARGS__ )
 #define TTF_FONT( RelativePath, ... ) FSlateFontInfo( Style->RootToContentDir( RelativePath, TEXT(".ttf") ), __VA_ARGS__ )
 #define OTF_FONT( RelativePath, ... ) FSlateFontInfo( Style->RootToContentDir( RelativePath, TEXT(".otf") ), __VA_ARGS__ )
+
     const FVector2D Icon16x16(16.0f, 16.0f);
     const FVector2D Icon20x20(20.0f, 20.0f);
-    const FVector2D Icon40x40(40.0f, 40.0f);
+    const FVector2D Icon64x64(64.0f, 64.0f);
     
     TSharedRef<FSlateStyleSet> FPixelStreamingStyle::Create()
     {
         TSharedRef< FSlateStyleSet > Style = MakeShareable(new FSlateStyleSet("PixelStreamingStyle"));
-        // Style->SetContentRoot(IPluginManager::Get().FindPlugin("PixelStreaming")->GetBaseDir() / TEXT("Resources"));
-        // Style->Set("PixelStreaming.OpenPluginWindow", new IMAGE_BRUSH(TEXT("Radio"), Icon40x40));
+        Style->SetContentRoot(IPluginManager::Get().FindPlugin("PixelStreaming")->GetBaseDir() / TEXT("Resources"));
+        Style->Set("PixelStreaming.Icon", new IMAGE_BRUSH_SVG("PixelStreaming_16", Icon16x16));
         return Style;
     }
 #undef IMAGE_BRUSH
+#undef IMAGE_BRUSH_SVG
 #undef BOX_BRUSH
 #undef BORDER_BRUSH
 #undef TTF_FONT
diff --git a/Engine/Plugins/Media/PixelStreaming/Source/PixelStreamingEditor/Private/PixelStreamingToolbar.cpp b/Engine/Plugins/Media/PixelStreaming/Source/PixelStreamingEditor/Private/PixelStreamingToolbar.cpp
--- a/Engine/Plugins/Media/PixelStreaming/Source/PixelStreamingEditor/Private/PixelStreamingToolbar.cpp
+++ b/Engine/Plugins/Media/PixelStreaming/Source/PixelStreamingEditor/Private/PixelStreamingToolbar.cpp
@@ -27,234 +27,493 @@
 #include "Types/SlateEnums.h"
 #include "Input/Events.h"
 #include "InputCoreTypes.h"
+#include "PixelStreamingServers.h"
+#include "GenericPlatform/GenericPlatformHttp.h"
+#include <SocketSubsystem.h>
+#include <IPAddress.h>
+#include "PixelStreamingEditorUtils.h"
+#include "SlateFwd.h"
+#include "Widgets/Input/SNumericEntryBox.h"
 
 #define LOCTEXT_NAMESPACE "PixelStreamingEditor"
 
 namespace UE::EditorPixelStreaming
 {
-    FPixelStreamingToolbar::FPixelStreamingToolbar()
-    : PixelStreamingModule(IPixelStreamingModule::Get())
-    {
-        FPixelStreamingCommands::Register();
-
-        PluginCommands = MakeShared<FUICommandList>();
-        
-        PluginCommands->MapAction(
-            FPixelStreamingCommands::Get().StartStreaming,
-            FExecuteAction::CreateLambda([]()
-            {
-                FPixelStreamingEditorModule::GetModule()->StartStreaming();
-            }),
-            FCanExecuteAction::CreateLambda([Module = &PixelStreamingModule] {  
-                if(TSharedPtr<IPixelStreamingStreamer> Streamer = Module->GetStreamer(Module->GetDefaultStreamerID()))    
-                {
-                    return !Streamer->IsStreaming();
-                }                 
-                return false;
-            })
-        );
-       
-        PluginCommands->MapAction(
-            FPixelStreamingCommands::Get().StopStreaming,
-            FExecuteAction::CreateLambda([]()
-            {
-                FPixelStreamingEditorModule::GetModule()->StopStreaming();
-            }),
-            FCanExecuteAction::CreateLambda([Module = &PixelStreamingModule] { 
-                if(TSharedPtr<IPixelStreamingStreamer> Streamer = Module->GetStreamer(Module->GetDefaultStreamerID()))    
-                {
-                    return Streamer->IsStreaming();
-                }
-                return true;
-            })
-        );
-
-        PluginCommands->MapAction(
-            FPixelStreamingCommands::Get().VP8,
-            FExecuteAction::CreateLambda([Module = &PixelStreamingModule]()
-            {
-                Module->SetCodec(EPixelStreamingCodec::VP8);
-            }),
-            FCanExecuteAction::CreateLambda([Module = &PixelStreamingModule] { 
-                if(TSharedPtr<IPixelStreamingStreamer> Streamer = Module->GetStreamer(Module->GetDefaultStreamerID()))    
-                {
-                    return !Streamer->IsStreaming();
-                }
-                return false;
-            }),
-            FIsActionChecked::CreateLambda([Module = &PixelStreamingModule]()
-            { 
-                return Module->GetCodec() == EPixelStreamingCodec::VP8;
-            })
-        );
-
-        PluginCommands->MapAction(
-            FPixelStreamingCommands::Get().VP9,
-            FExecuteAction::CreateLambda([Module = &PixelStreamingModule]()
-            {
-                Module->SetCodec(EPixelStreamingCodec::VP9);
-            }),
-            FCanExecuteAction::CreateLambda([Module = &PixelStreamingModule] { 
-                if(TSharedPtr<IPixelStreamingStreamer> Streamer = Module->GetStreamer(Module->GetDefaultStreamerID()))    
-                {
-                    return !Streamer->IsStreaming();
-                }
-                return false;
-            }),
-            FIsActionChecked::CreateLambda([Module = &PixelStreamingModule]()
-            { 
-                return Module->GetCodec() == EPixelStreamingCodec::VP9;
-            })
-        );
-
-        PluginCommands->MapAction(
-            FPixelStreamingCommands::Get().H264,
-            FExecuteAction::CreateLambda([Module = &PixelStreamingModule]()
-            {
-                Module->SetCodec(EPixelStreamingCodec::H264);
-            }),
-            FCanExecuteAction::CreateLambda([Module = &PixelStreamingModule] { 
-                if(TSharedPtr<IPixelStreamingStreamer> Streamer = Module->GetStreamer(Module->GetDefaultStreamerID()))    
-                {
-                    return !Streamer->IsStreaming();
-                }
-                return false;
-            }),
-            FIsActionChecked::CreateLambda([Module = &PixelStreamingModule]()
-            { 
-                return Module->GetCodec() == EPixelStreamingCodec::H264;
-            })
-        );
-
-        UToolMenus::RegisterStartupCallback(FSimpleMulticastDelegate::FDelegate::CreateRaw(this, &FPixelStreamingToolbar::RegisterMenus));
-    }
-    
-    FPixelStreamingToolbar::~FPixelStreamingToolbar()
-    {
-        FPixelStreamingCommands::Unregister();
-    }
-    
-    void FPixelStreamingToolbar::RegisterMenus()
-    {
-        FToolMenuOwnerScoped OwnerScoped(this);
-        {
-            UToolMenu* CustomToolBar = UToolMenus::Get()->ExtendMenu("LevelEditor.LevelEditorToolBar.PlayToolBar");
-            {
-                FToolMenuSection& Section = CustomToolBar->AddSection("PixelStreaming");
-                Section.AddSeparator("PixelStreamingStart");
-                {
-                    TSharedRef<SWidget> TitleWidget = SNew(SHorizontalBox)
-                    + SHorizontalBox::Slot()
-                    .VAlign(VAlign_Center)
-                    [
-                        SNew(STextBlock)
-                        .Text(LOCTEXT("Pixel Streaming", "Pixel Streaming"))
-                    ];
-
-                    // Play Button
-                    FToolMenuEntry TitleEntry = FToolMenuEntry::InitWidget(
-                        FName("Title"),
-                        TitleWidget,
-                        FText()
-                    );
-
-                    Section.AddEntry(TitleEntry);  
-                }
-
-                {
-                    // Play Button
-                    FToolMenuEntry PlayEntry = FToolMenuEntry::InitToolBarButton(
-                        FPixelStreamingCommands::Get().StartStreaming,
-                        TAttribute<FText>(),
-                        TAttribute<FText>(),
-                        FSlateIcon(FAppStyle::GetAppStyleSetName(), "PlayWorld.PlayInViewport"),
-                        FName("StartStreaming")
-                    );
-
-                    PlayEntry.StyleNameOverride = FName("Toolbar.BackplateLeftPlay");   
-                    PlayEntry.SetCommandList(PluginCommands);
-                    Section.AddEntry(PlayEntry);  
-                }
-                           
-                {
-                    // Stop Button
-                    FToolMenuEntry StopEntry = FToolMenuEntry::InitToolBarButton(
-                        FPixelStreamingCommands::Get().StopStreaming,
-                        TAttribute<FText>(),
-                        TAttribute<FText>(),
-                        FSlateIcon(FAppStyle::GetAppStyleSetName(), "PlayWorld.StopPlaySession"),
-                        FName("StopStreaming")
-                    );
-
-                    StopEntry.StyleNameOverride = FName("Toolbar.BackplateCenterStop");
-                    StopEntry.SetCommandList(PluginCommands);
-                    Section.AddEntry(StopEntry);
-                }
-
-                {
-                    // Settings dropdown
-                    FToolMenuEntry SettingsEntry = FToolMenuEntry::InitComboButton(
-                        "PixelStreamingSettingsMenu",
-                        FUIAction(),
-                        FOnGetContent::CreateLambda(
-                            [&]()
-                            {
-                                FMenuBuilder MenuBuilder(true, PluginCommands);
-                                
-                                MenuBuilder.BeginSection("Settings", LOCTEXT("PixelStreamingSettings", "Settings"));
-                                {
-                                    TSharedRef<SWidget> URLInput = SNew(SEditableTextBox)
-                                    .Text_Lambda([&, Module = &PixelStreamingModule]()
-                                    {
-                                        return FText::FromString(Module->GetDefaultSignallingURL());
-                                    })
-								    .OnTextCommitted_Lambda([&, Module = &PixelStreamingModule](const FText& InText, ETextCommit::Type InTextCommit)
-                                    {
-                                        if(TSharedPtr<IPixelStreamingStreamer> Streamer = Module->GetStreamer(Module->GetDefaultStreamerID()))    
-                                        {
-                                            Streamer->SetSignallingServerURL(InText.ToString());
-                                        }
-                                    });
-
-                                    MenuBuilder.AddWidget(URLInput,  LOCTEXT("PixelStreamingURL", "Signalling URL"), true);
-                                }
-                                MenuBuilder.EndSection();
-
-                                MenuBuilder.BeginSection("Advanced", LOCTEXT("PixelStreamingAdvancedSettings", "Advanced"));
-                                {
-                                    MenuBuilder.AddSubMenu(
-                                        LOCTEXT("PixelStreamingCodecSettings", "Select Codec"),
-                                        FText(),
-                                        FNewMenuDelegate::CreateLambda([](FMenuBuilder& SubMenuBuilder)
-                                        {
-                                            SubMenuBuilder.AddMenuEntry(FPixelStreamingCommands::Get().H264);
-                                            SubMenuBuilder.AddMenuEntry(FPixelStreamingCommands::Get().VP8);
-                                            SubMenuBuilder.AddMenuEntry(FPixelStreamingCommands::Get().VP9);
-                                        })
-                                    );
-                                }
-                                MenuBuilder.EndSection();
-
-                                return MenuBuilder.MakeWidget();
-                            }
-                        ),
-                        FText(),
-                        LOCTEXT("LevelEditorToolbarPixelStreamingSettingsTooltip", "Configure PixelStreaming")
-                    );
-                    SettingsEntry.StyleNameOverride = FName("Toolbar.BackplateRightCombo");
-                    SettingsEntry.SetCommandList(PluginCommands);
-                    Section.AddEntry(SettingsEntry);
-                }
-                Section.AddSeparator("PixelStreamingEnd");
-            }
-        }
-    }
-
-    TSharedRef<SWidget> FPixelStreamingToolbar::GeneratePixelStreamingMenuContent(TSharedPtr<FUICommandList> InCommandList)
-    {     
-        FToolMenuContext MenuContext(InCommandList);
-        return UToolMenus::Get()->GenerateWidget("LevelEditor.LevelEditorToolBar.AddQuickMenu", MenuContext);
-    }
+	FPixelStreamingToolbar::FPixelStreamingToolbar()
+	: PixelStreamingModule(IPixelStreamingModule::Get())
+	{
+		FPixelStreamingCommands::Register();
+
+		PluginCommands = MakeShared<FUICommandList>();
+
+		PluginCommands->MapAction(
+			FPixelStreamingCommands::Get().ExternalSignalling,
+			FExecuteAction::CreateLambda([&]()
+			{
+				FPixelStreamingEditorModule::GetModule()->bUseExternalSignallingServer = !FPixelStreamingEditorModule::GetModule()->bUseExternalSignallingServer;
+				FPixelStreamingEditorModule::GetModule()->StopSignalling();
+			}),
+			FCanExecuteAction::CreateLambda([Module = &PixelStreamingModule] {  
+				if(TSharedPtr<IPixelStreamingStreamer> Streamer = Module->GetStreamer(Module->GetDefaultStreamerID()))	
+				{
+					return !Streamer->IsStreaming();
+				}				 
+				return false;
+			}),
+			FIsActionChecked::CreateLambda([&]()
+			{ 
+				return FPixelStreamingEditorModule::GetModule()->bUseExternalSignallingServer;
+			})
+		);
+
+		PluginCommands->MapAction(
+			FPixelStreamingCommands::Get().StreamLevelEditor,
+			FExecuteAction::CreateLambda([]()
+			{
+				FPixelStreamingEditorModule::GetModule()->StartStreaming(EStreamTypes::LevelEditorViewport);
+			}),
+			FCanExecuteAction::CreateLambda([Module = &PixelStreamingModule] {  
+				if(TSharedPtr<IPixelStreamingStreamer> Streamer = Module->GetStreamer(Module->GetDefaultStreamerID()))	
+				{
+					return !Streamer->IsStreaming();
+				}				 
+				return false;
+			})
+		);
+
+		PluginCommands->MapAction(
+			FPixelStreamingCommands::Get().StreamEditor,
+			FExecuteAction::CreateLambda([]()
+			{
+				FPixelStreamingEditorModule::GetModule()->StartStreaming(EStreamTypes::Editor);
+			}),
+			FCanExecuteAction::CreateLambda([Module = &PixelStreamingModule] {  
+				if(TSharedPtr<IPixelStreamingStreamer> Streamer = Module->GetStreamer(Module->GetDefaultStreamerID()))	
+				{
+					return !Streamer->IsStreaming();
+				}				 
+				return false;
+			})
+		);
+	   
+		PluginCommands->MapAction(
+			FPixelStreamingCommands::Get().StopStreaming,
+			FExecuteAction::CreateLambda([]()
+			{
+				FPixelStreamingEditorModule::GetModule()->StopStreaming();
+			}),
+			FCanExecuteAction::CreateLambda([Module = &PixelStreamingModule] { 
+				if(TSharedPtr<IPixelStreamingStreamer> Streamer = Module->GetStreamer(Module->GetDefaultStreamerID()))	
+				{
+					return Streamer->IsStreaming();
+				}
+				return true;
+			})
+		);
+
+		PluginCommands->MapAction(
+			FPixelStreamingCommands::Get().StartSignalling,
+			FExecuteAction::CreateLambda([]()
+			{
+				FPixelStreamingEditorModule::GetModule()->StartSignalling();
+			}),
+			FCanExecuteAction::CreateLambda([] {  
+				TSharedPtr<UE::PixelStreamingServers::IServer> SignallingServer = FPixelStreamingEditorModule::GetModule()->GetSignallingServer();
+				if(!SignallingServer.IsValid() || !SignallingServer->HasLaunched())
+				{
+					return true;
+				}
+				return false;
+			})
+		);
+
+		PluginCommands->MapAction(
+			FPixelStreamingCommands::Get().StopSignalling,
+			FExecuteAction::CreateLambda([]()
+			{
+				FPixelStreamingEditorModule::GetModule()->StopSignalling();
+			}),
+			FCanExecuteAction::CreateLambda([] {  
+				TSharedPtr<UE::PixelStreamingServers::IServer> SignallingServer = FPixelStreamingEditorModule::GetModule()->GetSignallingServer();
+				if(SignallingServer.IsValid() && SignallingServer->HasLaunched())
+				{
+					return true;
+				}
+				return false;
+			})
+		);
+
+		PluginCommands->MapAction(
+			FPixelStreamingCommands::Get().VP8,
+			FExecuteAction::CreateLambda([Module = &PixelStreamingModule]()
+			{
+				Module->SetCodec(EPixelStreamingCodec::VP8);
+			}),
+			FCanExecuteAction::CreateLambda([Module = &PixelStreamingModule] { 
+				if(TSharedPtr<IPixelStreamingStreamer> Streamer = Module->GetStreamer(Module->GetDefaultStreamerID()))	
+				{
+					return !Streamer->IsStreaming();
+				}
+				return false;
+			}),
+			FIsActionChecked::CreateLambda([Module = &PixelStreamingModule]()
+			{ 
+				return Module->GetCodec() == EPixelStreamingCodec::VP8;
+			})
+		);
+
+		PluginCommands->MapAction(
+			FPixelStreamingCommands::Get().VP9,
+			FExecuteAction::CreateLambda([Module = &PixelStreamingModule]()
+			{
+				Module->SetCodec(EPixelStreamingCodec::VP9);
+			}),
+			FCanExecuteAction::CreateLambda([Module = &PixelStreamingModule] { 
+				if(TSharedPtr<IPixelStreamingStreamer> Streamer = Module->GetStreamer(Module->GetDefaultStreamerID()))	
+				{
+					return !Streamer->IsStreaming();
+				}
+				return false;
+			}),
+			FIsActionChecked::CreateLambda([Module = &PixelStreamingModule]()
+			{ 
+				return Module->GetCodec() == EPixelStreamingCodec::VP9;
+			})
+		);
+
+		PluginCommands->MapAction(
+			FPixelStreamingCommands::Get().H264,
+			FExecuteAction::CreateLambda([Module = &PixelStreamingModule]()
+			{
+				Module->SetCodec(EPixelStreamingCodec::H264);
+			}),
+			FCanExecuteAction::CreateLambda([Module = &PixelStreamingModule] { 
+				if(TSharedPtr<IPixelStreamingStreamer> Streamer = Module->GetStreamer(Module->GetDefaultStreamerID()))	
+				{
+					return !Streamer->IsStreaming();
+				}
+				return false;
+			}),
+			FIsActionChecked::CreateLambda([Module = &PixelStreamingModule]()
+			{ 
+				return Module->GetCodec() == EPixelStreamingCodec::H264;
+			})
+		);
+
+		UToolMenus::RegisterStartupCallback(FSimpleMulticastDelegate::FDelegate::CreateRaw(this, &FPixelStreamingToolbar::RegisterMenus));
+	}
+	
+	FPixelStreamingToolbar::~FPixelStreamingToolbar()
+	{
+		FPixelStreamingCommands::Unregister();
+	}
+	
+	void FPixelStreamingToolbar::RegisterMenus()
+	{
+		FToolMenuOwnerScoped OwnerScoped(this);
+		{
+			UToolMenu* CustomToolBar = UToolMenus::Get()->ExtendMenu("LevelEditor.LevelEditorToolBar.PlayToolBar");
+			{
+				FToolMenuSection& Section = CustomToolBar->AddSection("PixelStreaming");
+				Section.AddSeparator("PixelStreamingSeperator");
+				{
+					// Settings dropdown
+					FToolMenuEntry SettingsEntry = FToolMenuEntry::InitComboButton(
+						"PixelStreamingMenus",
+						FUIAction(),
+						FOnGetContent::CreateLambda(
+							[&]()
+							{
+								FMenuBuilder MenuBuilder(true, PluginCommands);  
+
+								// Use external signalling server option
+								MenuBuilder.BeginSection("Signalling Server Location", LOCTEXT("PixelStreamingSettings", "Signalling Server Location"));
+								MenuBuilder.AddMenuEntry(FPixelStreamingCommands::Get().ExternalSignalling);
+								MenuBuilder.EndSection();
+
+								if(!FPixelStreamingEditorModule::GetModule()->bUseExternalSignallingServer)
+								{
+									// Embedded Signalling Server Config (streamer port & http port)
+									RegisterEmbeddedSignallingServerConfig(MenuBuilder);
+								
+									// Signalling Server Viewer URLs
+									TSharedPtr<UE::PixelStreamingServers::IServer> SignallingServer = FPixelStreamingEditorModule::GetModule()->GetSignallingServer();
+									if(SignallingServer.IsValid() && SignallingServer->HasLaunched())
+									{
+										RegisterSignallingServerURLs(MenuBuilder);
+									}
+								}
+								else
+								{
+									// Remote Signalling Server Config (URL)
+									RegisterRemoteSignallingServerConfig(MenuBuilder);
+								}	
+
+								// Pixel Streaming Config (streamable viewports OR currently streamed viewport)
+								TSharedPtr<IPixelStreamingStreamer> Streamer = PixelStreamingModule.GetStreamer(PixelStreamingModule.GetDefaultStreamerID());
+								if(Streamer.IsValid())
+								{
+									RegisterPixelStreamingControls(MenuBuilder);
+								}
+
+								// VCam Config
+								RegisterVCamControls(MenuBuilder);
+								
+								// Codec Config
+								RegisterCodecConfig(MenuBuilder);
+
+								return MenuBuilder.MakeWidget();
+							}
+						),
+						LOCTEXT("PixelStreamingMenu", "Pixel Streaming"),
+						LOCTEXT("PixelStreamingMenuTooltip", "Configure Pixel Streaming"),
+						FSlateIcon(FPixelStreamingStyle::GetStyleSetName(), "PixelStreaming.Icon"),
+						false,
+						"PixelStreamingMenu"
+					);
+					SettingsEntry.StyleNameOverride = "CalloutToolbar";
+					SettingsEntry.SetCommandList(PluginCommands);
+					Section.AddEntry(SettingsEntry);
+				}			  
+			}
+		}
+	}
+
+	void FPixelStreamingToolbar::RegisterEmbeddedSignallingServerConfig(FMenuBuilder& MenuBuilder)
+	{
+		MenuBuilder.BeginSection("Signalling Server Options", LOCTEXT("PixelStreamingSettings", "Signalling Server Options"));
+
+		TSharedPtr<UE::PixelStreamingServers::IServer> SignallingServer = FPixelStreamingEditorModule::GetModule()->GetSignallingServer();
+		if(!SignallingServer.IsValid() || !SignallingServer->HasLaunched())
+		{
+			TSharedRef<SWidget> StreamerPortInputBlock = SNew(SHorizontalBox)
+															+ SHorizontalBox::Slot()
+															.AutoWidth()
+															.VAlign(VAlign_Center)
+															.Padding(FMargin(36.0f, 3.0f, 8.0f, 3.0f))
+															[
+																SNew(STextBlock)
+																	.Text(FText::FromString(TEXT("Streamer Port: ")))
+																	.ColorAndOpacity(FSlateColor(FLinearColor(1.0f, 1.0f, 1.0f)))
+															]
+															+ SHorizontalBox::Slot()
+															.AutoWidth()
+															[
+																SNew(SNumericEntryBox<int32>)
+																	.MinValue(1)
+																	.Value_Lambda([&]()
+																	{
+																		return FPixelStreamingEditorModule::GetModule()->GetStreamerPort();
+																	})
+																	.OnValueChanged_Lambda([&](int32 InStreamerPort)
+																	{
+																		FPixelStreamingEditorModule::GetModule()->SetStreamerPort(InStreamerPort);
+																	})
+																	.OnValueCommitted_Lambda([&](int32 InStreamerPort, ETextCommit::Type InCommitType)
+																	{
+																		FPixelStreamingEditorModule::GetModule()->SetStreamerPort(InStreamerPort);
+																	})	  
+															];
+			MenuBuilder.AddWidget(StreamerPortInputBlock, FText(), true);
+			TSharedRef<SWidget> ViewerPortInputBlock = SNew(SHorizontalBox)
+															+ SHorizontalBox::Slot()
+															.AutoWidth()
+															.VAlign(VAlign_Center)
+															.Padding(FMargin(36.0f, 3.0f, 8.0f, 3.0f))
+															[
+																SNew(STextBlock)
+																	.Text(FText::FromString(TEXT("Viewer Port: ")))
+																	.ColorAndOpacity(FSlateColor(FLinearColor(1.0f, 1.0f, 1.0f)))
+															]
+															+ SHorizontalBox::Slot()
+															.AutoWidth()
+															[
+																SNew(SNumericEntryBox<int32>)
+																	.MinValue(1)
+																	.Value_Lambda([&]()
+																	{
+																		return FPixelStreamingEditorModule::GetModule()->GetViewerPort();
+																	})
+																	.OnValueChanged_Lambda([&, Module = &PixelStreamingModule](int32 InViewerPort)
+																	{
+																		FPixelStreamingEditorModule::GetModule()->SetViewerPort(InViewerPort);
+																	})
+																	.OnValueCommitted_Lambda([&](int32 InViewerPort, ETextCommit::Type InCommitType)
+																	{
+																		FPixelStreamingEditorModule::GetModule()->SetViewerPort(InViewerPort);
+																	})
+															];
+			MenuBuilder.AddWidget(ViewerPortInputBlock, FText(), true);
+			MenuBuilder.AddMenuEntry(FPixelStreamingCommands::Get().StartSignalling);
+		}
+		else
+		{
+			MenuBuilder.AddMenuEntry(FPixelStreamingCommands::Get().StopSignalling);
+		}
+
+		MenuBuilder.EndSection();
+	}
+
+	void FPixelStreamingToolbar::RegisterRemoteSignallingServerConfig(FMenuBuilder& MenuBuilder)
+	{
+		TSharedPtr<UE::PixelStreamingServers::IServer> SignallingServer = FPixelStreamingEditorModule::GetModule()->GetSignallingServer();
+		MenuBuilder.BeginSection("Remote Signalling Server Options", LOCTEXT("PixelStreamingSettings", "Remote Signalling Server Options"));
+		{
+			TSharedRef<SWidget> URLInputBlock = SNew(SHorizontalBox)
+														+ SHorizontalBox::Slot()
+														.AutoWidth()
+														.VAlign(VAlign_Center)
+														.Padding(FMargin(36.0f, 3.0f, 8.0f, 3.0f))
+														[
+															SNew(STextBlock)
+																.Text(FText::FromString(TEXT("Remote Signalling Server URL")))
+																.ColorAndOpacity(FSlateColor(FLinearColor(1.0f, 1.0f, 1.0f)))
+														]
+														+ SHorizontalBox::Slot()
+														.AutoWidth()
+														[
+															SNew(SEditableTextBox)
+																.Text_Lambda([&]()
+																{
+																	TSharedPtr<IPixelStreamingStreamer> Streamer = PixelStreamingModule.GetStreamer(PixelStreamingModule.GetDefaultStreamerID());	
+																	return FText::FromString(Streamer->GetSignallingServerURL());
+																})
+																.OnTextChanged_Lambda([&](const FText& InText)
+																{
+																	TSharedPtr<IPixelStreamingStreamer> Streamer = PixelStreamingModule.GetStreamer(PixelStreamingModule.GetDefaultStreamerID());   
+																	Streamer->SetSignallingServerURL(InText.ToString()); 
+																})
+																.OnTextCommitted_Lambda([&](const FText& InText, ETextCommit::Type InTextCommit)
+																{
+																	TSharedPtr<IPixelStreamingStreamer> Streamer = PixelStreamingModule.GetStreamer(PixelStreamingModule.GetDefaultStreamerID());   
+																	Streamer->SetSignallingServerURL(InText.ToString()); 
+																})
+																.IsEnabled_Lambda([Module = &PixelStreamingModule]()
+																{
+																	if(TSharedPtr<IPixelStreamingStreamer> Streamer = Module->GetStreamer(Module->GetDefaultStreamerID()))	
+																	{
+																		return !Streamer->IsStreaming();
+																	}
+																	return false;
+																})
+														];
+					MenuBuilder.AddWidget(URLInputBlock, FText(), true);							   
+		}
+		MenuBuilder.EndSection();
+	}
+
+	void FPixelStreamingToolbar::RegisterSignallingServerURLs(FMenuBuilder& MenuBuilder)
+	{
+		MenuBuilder.BeginSection("Signalling Server URLs", LOCTEXT("PixelStreamingSignallingURLs", "Signalling Server URLs"));
+		{
+			MenuBuilder.AddWidget(  SNew(SBox)
+										.Padding(FMargin(16.0f, 3.0f))
+										[
+											SNew(STextBlock)
+												.ColorAndOpacity(FSlateColor::UseSubduedForeground())
+												.Text(LOCTEXT("SignallingTip", "The Signalling Server is running and may be accessed via the following URLs (network settings permitting):"))
+												.WrapTextAt(400)
+										], 
+									FText());
+
+
+			MenuBuilder.AddWidget(  SNew(SBox)
+										.Padding(FMargin(32.0f, 3.0f))
+										[
+											SNew(STextBlock)
+												.ColorAndOpacity(FSlateColor::UseSubduedForeground())
+												.Text(FText::FromString(FString::Printf(TEXT("127.0.0.1:%d"), FPixelStreamingEditorModule::GetModule()->GetViewerPort())))
+										], 
+									FText());
+
+
+			TArray<TSharedPtr<FInternetAddr>> AdapterAddresses;
+			if (ISocketSubsystem::Get(PLATFORM_SOCKETSUBSYSTEM)->GetLocalAdapterAddresses(AdapterAddresses))
+			{
+				for(TSharedPtr<FInternetAddr> AdapterAddress : AdapterAddresses)
+				{
+					MenuBuilder.AddWidget(  SNew(SBox)
+												.Padding(FMargin(32.0f, 3.0f))
+												[
+													SNew(STextBlock)
+														.ColorAndOpacity(FSlateColor::UseSubduedForeground())
+														.Text(FText::FromString(FString::Printf(TEXT("%s:%d"), *AdapterAddress->ToString(false), FPixelStreamingEditorModule::GetModule()->GetViewerPort())))
+												], 
+											FText());
+				}
+			}
+		}
+		MenuBuilder.EndSection();
+	}
+
+	void FPixelStreamingToolbar::RegisterPixelStreamingControls(FMenuBuilder& MenuBuilder)
+	{
+		TSharedPtr<IPixelStreamingStreamer> Streamer = PixelStreamingModule.GetStreamer(PixelStreamingModule.GetDefaultStreamerID());
+		MenuBuilder.BeginSection("Pixel Streaming", LOCTEXT("PixelStreamingConfig", "Pixel Streaming"));
+		{
+			if(Streamer->IsStreaming())	
+			{   
+				MenuBuilder.AddWidget(  SNew(SBox)
+											.Padding(FMargin(16.0f, 3.0f))
+											[
+												SNew(STextBlock)
+													.ColorAndOpacity(FSlateColor::UseSubduedForeground())
+													.Text(FText::FromString(FString::Printf(TEXT("Streaming %s"), *ToString(FPixelStreamingEditorModule::GetModule()->GetStreamType()))))
+													.WrapTextAt(400)
+											], 
+										FText());
+				// We currently have to stop VCams streaming from the VCam itself
+				if(FPixelStreamingEditorModule::GetModule()->GetStreamType() != EStreamTypes::VCam)
+				{
+					MenuBuilder.AddMenuEntry(FPixelStreamingCommands::Get().StopStreaming);
+				}
+			}
+			else
+			{
+				if(FPixelStreamingEditorModule::GetModule()->bUseExternalSignallingServer)
+				{
+					
+				}
+				MenuBuilder.AddMenuEntry(FPixelStreamingCommands::Get().StreamLevelEditor);
+				MenuBuilder.AddMenuEntry(FPixelStreamingCommands::Get().StreamEditor);
+			}
+		}
+		MenuBuilder.EndSection();
+	}
+
+	void FPixelStreamingToolbar::RegisterVCamControls(FMenuBuilder& MenuBuilder)
+	{
+		MenuBuilder.BeginSection("Virtual Camera", LOCTEXT("PixelStreamingVirtualCamera", "Virtual Camera"));
+		{
+			MenuBuilder.AddWidget(  SNew(SBox)
+										.Padding(FMargin(16.0f, 3.0f))
+										[
+											SNew(STextBlock)
+												.ColorAndOpacity(FSlateColor::UseSubduedForeground())
+												.Text(LOCTEXT("VirtualCamera", "Virtual Camera streams can be started and stopped via the Camera Actor's VCam component"))
+												.WrapTextAt(400)
+										], 
+									FText());
+		}
+		MenuBuilder.EndSection();
+	}
+
+	void FPixelStreamingToolbar::RegisterCodecConfig(FMenuBuilder& MenuBuilder)
+	{
+		MenuBuilder.BeginSection("Codec", LOCTEXT("PixelStreamingCodecSettings", "Codec"));
+		{
+			MenuBuilder.AddMenuEntry(FPixelStreamingCommands::Get().H264);
+			MenuBuilder.AddMenuEntry(FPixelStreamingCommands::Get().VP8);
+			MenuBuilder.AddMenuEntry(FPixelStreamingCommands::Get().VP9);
+		}
+		MenuBuilder.EndSection();
+	}
+
+	TSharedRef<SWidget> FPixelStreamingToolbar::GeneratePixelStreamingMenuContent(TSharedPtr<FUICommandList> InCommandList)
+	{	 
+		FToolMenuContext MenuContext(InCommandList);
+		return UToolMenus::Get()->GenerateWidget("LevelEditor.LevelEditorToolBar.AddQuickMenu", MenuContext);
+	}
 }
 
 #undef LOCTEXT_NAMESPACE
\ No newline at end of file
diff --git a/Engine/Plugins/Media/PixelStreaming/Source/PixelStreamingEditor/Private/PixelStreamingToolbar.h b/Engine/Plugins/Media/PixelStreaming/Source/PixelStreamingEditor/Private/PixelStreamingToolbar.h
--- a/Engine/Plugins/Media/PixelStreaming/Source/PixelStreamingEditor/Private/PixelStreamingToolbar.h
+++ b/Engine/Plugins/Media/PixelStreaming/Source/PixelStreamingEditor/Private/PixelStreamingToolbar.h
@@ -5,6 +5,7 @@
 #include "Framework/Commands/UIAction.h"
 #include "Framework/Commands/UICommandList.h"
 #include "Framework/Docking/TabManager.h"
+#include "Framework/MultiBox/MultiBoxBuilder.h"
 #include "LevelEditor.h"
 #include "ToolMenus.h"
 #include "IPixelStreamingModule.h"
@@ -24,6 +25,13 @@ namespace UE::EditorPixelStreaming
         
         private:
             void RegisterMenus();
+			void RegisterEmbeddedSignallingServerConfig(FMenuBuilder& MenuBuilder);
+            void RegisterRemoteSignallingServerConfig(FMenuBuilder& MenuBuilder);
+            void RegisterSignallingServerURLs(FMenuBuilder& MenuBuilder);
+            void RegisterPixelStreamingControls(FMenuBuilder& MenuBuilder);
+            void RegisterVCamControls(FMenuBuilder& MenuBuilder);
+            void RegisterCodecConfig(FMenuBuilder& MenuBuilder);
+
             TSharedPtr<class FUICommandList> PluginCommands;
 
             IPixelStreamingModule& PixelStreamingModule;
diff --git a/Engine/Plugins/Media/PixelStreaming/Source/PixelStreamingEditor/Private/PixelStreamingVideoInputBackBufferComposited.cpp b/Engine/Plugins/Media/PixelStreaming/Source/PixelStreamingEditor/Private/PixelStreamingVideoInputBackBufferComposited.cpp
--- a/Engine/Plugins/Media/PixelStreaming/Source/PixelStreamingEditor/Private/PixelStreamingVideoInputBackBufferComposited.cpp
+++ b/Engine/Plugins/Media/PixelStreaming/Source/PixelStreamingEditor/Private/PixelStreamingVideoInputBackBufferComposited.cpp
@@ -93,10 +93,12 @@ void FPixelStreamingVideoInputBackBufferComposited::OnBackBufferReady(SWindow& S
 		for (TSharedRef<SWindow> Window : TopLevelWindows)
 		{
 			FTextureRHIRef Texture = TopLevelWindowTextures.FindRef(&Window.Get());
-			if (Window->GetOpacity() == 0.0f || Window->GetType() == EWindowType::CursorDecorator)
+			if (Window->GetOpacity() == 0.0f || 
+				Window->GetType() == EWindowType::CursorDecorator ||
+				Window->GetSizeInScreen() == FVector2D(0, 0))
 			{
 				continue;
-			}
+			}			
 
 			uint32 CompositedX = CompositedFrame->GetSizeXY().X;
 			uint32 CompositedY = CompositedFrame->GetSizeXY().Y;
@@ -109,8 +111,8 @@ void FPixelStreamingVideoInputBackBufferComposited::OnBackBufferReady(SWindow& S
 			{
 				// If so, create a new compositedframe large enough to accomodate our new texture + its position
 				// and render the contents of the old compositedframe to the new one
-				uint32 SizeX = TotalX > CompositedX ? TotalX : CompositedX;
-				uint32 SizeY = TotalY > CompositedY ? TotalY : CompositedY;
+				uint32 SizeX = (TotalX > CompositedX ? TotalX : CompositedX);
+				uint32 SizeY = (TotalY > CompositedY ? TotalY : CompositedY);
 
 				StagingTexture = UE::PixelStreaming::CreateRHITexture(SizeX, SizeY);
 				FRHICopyTextureInfo CopyInfo;
diff --git a/Engine/Plugins/Media/PixelStreaming/Source/PixelStreamingEditor/Public/PixelStreamingEditorModule.h b/Engine/Plugins/Media/PixelStreaming/Source/PixelStreamingEditor/Public/PixelStreamingEditorModule.h
--- a/Engine/Plugins/Media/PixelStreaming/Source/PixelStreamingEditor/Public/PixelStreamingEditorModule.h
+++ b/Engine/Plugins/Media/PixelStreaming/Source/PixelStreamingEditor/Public/PixelStreamingEditorModule.h
@@ -3,6 +3,8 @@
 #pragma once
 
 #include "IPixelStreamingModule.h"
+#include "PixelStreamingServers.h"
+#include "PixelStreamingEditorUtils.h"
 
 namespace UE::EditorPixelStreaming
 {
@@ -16,15 +18,44 @@ public:
     virtual void StartupModule() override;
     virtual void ShutdownModule() override;
 
-    void StartStreaming();
+    void StartStreaming(UE::EditorPixelStreaming::EStreamTypes InStreamType);
     void StopStreaming();
+    void SetStreamType(UE::EditorPixelStreaming::EStreamTypes InStreamType) { StreamType = InStreamType; };
+    UE::EditorPixelStreaming::EStreamTypes GetStreamType();
+
+    void StartSignalling();
+    void StopSignalling();
+    TSharedPtr<UE::PixelStreamingServers::IServer> GetSignallingServer();
+    
+    void SetSignallingDomain(const FString& InSignallingDomain);
+    FString GetSignallingDomain() { return SignallingDomain; };
+    void SetStreamerPort(int32 InStreamerPort);
+    int32 GetStreamerPort() { return StreamerPort; };
+    void SetViewerPort(int32 InViewerPort);
+    int32 GetViewerPort() { return ViewerPort; };
        
     static FPixelStreamingEditorModule* GetModule();
 private:
     void InitEditorStreaming(IPixelStreamingModule& Module);
 	bool ParseResolution(const TCHAR* InResolution, uint32& OutX, uint32& OutY);
     void MaybeResizeEditor(TSharedPtr<SWindow> RootWindow);
+    FString GetSignallingServerURL();
 
     TSharedPtr<UE::EditorPixelStreaming::FPixelStreamingToolbar> Toolbar;
     static FPixelStreamingEditorModule* PixelStreamingEditorModule;	
+    // Signalling/webserver
+	TSharedPtr<UE::PixelStreamingServers::IServer> SignallingServer;
+    // Download process for PS web frontend files (if we want to view output in the browser)
+	TSharedPtr<FMonitoredProcess> DownloadProcess;
+    //
+    UE::EditorPixelStreaming::EStreamTypes StreamType;
+    // The signalling server host: eg ws://127.0.0.1
+    FString SignallingDomain;
+    // The port the streamer will connect to. eg 8888
+    int32 StreamerPort;
+    // The port the streams can be viewed at on the browser. eg 80
+    int32 ViewerPort = 80;
+
+public:
+    bool bUseExternalSignallingServer = false;
 };
diff --git a/Engine/Plugins/Media/PixelStreaming/Source/PixelStreamingEditor/Public/PixelStreamingEditorUtils.h b/Engine/Plugins/Media/PixelStreaming/Source/PixelStreamingEditor/Public/PixelStreamingEditorUtils.h
new file mode 100644
--- /dev/null
+++ b/Engine/Plugins/Media/PixelStreaming/Source/PixelStreamingEditor/Public/PixelStreamingEditorUtils.h
@@ -0,0 +1,28 @@
+// Copyright Epic Games, Inc. All Rights Reserved.
+
+#pragma once
+
+namespace UE::EditorPixelStreaming
+{
+	enum class PIXELSTREAMINGEDITOR_API EStreamTypes : uint8
+	{
+		LevelEditorViewport = 0,
+		Editor = 1,
+		VCam = 2
+	};
+
+	inline FString ToString(EStreamTypes StreamType)
+	{
+		switch(StreamType)
+		{
+			case EStreamTypes::LevelEditorViewport:
+				return TEXT("the Level Editor");
+			case EStreamTypes::Editor:
+				return TEXT("the Full Editor");
+			case EStreamTypes::VCam:
+				return TEXT("a Virtual Camera");
+			default:
+				return TEXT("Unknown stream type!");
+		}
+	}
+}
\ No newline at end of file
diff --git a/Engine/Plugins/Media/PixelStreaming/Source/PixelStreamingServers/Private/SignallingServer.cpp b/Engine/Plugins/Media/PixelStreaming/Source/PixelStreamingServers/Private/SignallingServer.cpp
--- a/Engine/Plugins/Media/PixelStreaming/Source/PixelStreamingServers/Private/SignallingServer.cpp
+++ b/Engine/Plugins/Media/PixelStreaming/Source/PixelStreamingServers/Private/SignallingServer.cpp
@@ -248,6 +248,7 @@ namespace UE::PixelStreamingServers
 			TSharedRef<FJsonObject> PongJSON = MakeShared<FJsonObject>();
 			PongJSON->SetStringField(FString(TEXT("type")), FString(TEXT("pong")));
 			PongJSON->SetNumberField(FString(TEXT("time")), UnixTime);
+			UE_LOG(LogPixelStreamingServers, Log, TEXT("To Streamer id=%d: %s"), ConnectionId, *Utils::ToString(PongJSON));
 			StreamersWS->Send(ConnectionId, Utils::ToString(PongJSON));
 			return;
 		}
diff --git a/Engine/Plugins/Media/PixelStreamingVCam/Source/PixelStreamingVCam/PixelStreamingVCam.Build.cs b/Engine/Plugins/Media/PixelStreamingVCam/Source/PixelStreamingVCam/PixelStreamingVCam.Build.cs
--- a/Engine/Plugins/Media/PixelStreamingVCam/Source/PixelStreamingVCam/PixelStreamingVCam.Build.cs
+++ b/Engine/Plugins/Media/PixelStreamingVCam/Source/PixelStreamingVCam/PixelStreamingVCam.Build.cs
@@ -31,7 +31,8 @@ namespace UnrealBuildTool.Rules
 				"RHI",
 				"CinematicCamera",
 				"LiveLinkInterface",
-				"PixelCapture"
+				"PixelCapture",
+				"PixelStreamingEditor"
 			});
 
 			// Can't package non-editor targets (e.g. games) with UnrealEd, so this dependency should only be added in editor.
diff --git a/Engine/Plugins/Media/PixelStreamingVCam/Source/PixelStreamingVCam/Private/PixelStreamingMediaOutput.cpp b/Engine/Plugins/Media/PixelStreamingVCam/Source/PixelStreamingVCam/Private/PixelStreamingMediaOutput.cpp
--- a/Engine/Plugins/Media/PixelStreamingVCam/Source/PixelStreamingVCam/Private/PixelStreamingMediaOutput.cpp
+++ b/Engine/Plugins/Media/PixelStreamingVCam/Source/PixelStreamingVCam/Private/PixelStreamingMediaOutput.cpp
@@ -6,6 +6,8 @@
 #include "PixelStreamingPlayerId.h"
 #include "IPixelStreamingStreamer.h"
 #include "CineCameraComponent.h"
+#include "PixelStreamingEditorModule.h"
+#include "PixelStreamingEditorUtils.h"
 
 void UPixelStreamingMediaOutput::BeginDestroy()
 {
@@ -19,7 +21,7 @@ UMediaCapture* UPixelStreamingMediaOutput::CreateMediaCaptureImpl()
 	if (!Streamer.IsValid())
 	{
 		IPixelStreamingModule& Module = FModuleManager::LoadModuleChecked<IPixelStreamingModule>("PixelStreaming");
-		Streamer = Module.CreateStreamer(StreamID);
+		Streamer = Module.GetStreamer(Module.GetDefaultStreamerID());
 	}
 
 	Capture = nullptr;
@@ -62,7 +64,7 @@ void UPixelStreamingMediaOutput::StartStreaming()
 {
 	if (Streamer)
 	{
-		Streamer->SetSignallingServerURL(SignallingServerURL);
+		FPixelStreamingEditorModule::GetModule()->SetStreamType(UE::EditorPixelStreaming::EStreamTypes::VCam);
 		Streamer->SetVideoInput(Capture->GetVideoInput());
 
 		if (!Streamer->IsStreaming())
@@ -78,6 +80,7 @@ void UPixelStreamingMediaOutput::StopStreaming()
 	{
 		Streamer->StopStreaming();
 		Streamer->SetTargetViewport(nullptr);
+		Streamer->SetTargetWindow(nullptr);
 	}
 }
 
diff --git a/Engine/Plugins/Media/PixelStreamingVCam/Source/PixelStreamingVCam/Private/VCamPixelStreamingSession.cpp b/Engine/Plugins/Media/PixelStreamingVCam/Source/PixelStreamingVCam/Private/VCamPixelStreamingSession.cpp
--- a/Engine/Plugins/Media/PixelStreamingVCam/Source/PixelStreamingVCam/Private/VCamPixelStreamingSession.cpp
+++ b/Engine/Plugins/Media/PixelStreamingVCam/Source/PixelStreamingVCam/Private/VCamPixelStreamingSession.cpp
@@ -16,6 +16,7 @@
 #include "PixelStreamingVCamLog.h"
 #include "IPixelStreamingModule.h"
 #include "PixelStreamingProtocol.h"
+#include "PixelStreamingEditorModule.h"
 #include "VCamComponent.h"
 
 namespace VCamPixelStreamingSession
@@ -81,14 +82,24 @@ void UVCamPixelStreamingSession::Activate()
 
 	if (StartSignallingServer)
 	{
-		if (UVCamPixelStreamingSubsystem* PixelStreamingSubsystem = UVCamPixelStreamingSubsystem::Get())
+		if(!FPixelStreamingEditorModule::GetModule()->bUseExternalSignallingServer)
 		{
-			PixelStreamingSubsystem->LaunchSignallingServer(PortNumber, HttpPort);
+			if (UVCamPixelStreamingSubsystem* PixelStreamingSubsystem = UVCamPixelStreamingSubsystem::Get())
+			{
+				PixelStreamingSubsystem->LaunchSignallingServer();
+			}
 		}
+		else
+		{
+			StartSignallingServer = false;
+		}
+		
 	}
 
-	MediaOutput->SetSignallingServerURL(FString::Printf(TEXT("%s:%s"), *IP, *FString::FromInt(PortNumber)));
-	UE_LOG(LogPixelStreamingVCam, Log, TEXT("Activating PixelStreaming VCam Session. Endpoint: %s:%s"), *IP, *FString::FromInt(PortNumber));
+	FString SignallingDomain = FPixelStreamingEditorModule::GetModule()->GetSignallingDomain();
+	int32 StreamerPort = FPixelStreamingEditorModule::GetModule()->GetStreamerPort();
+	MediaOutput->SetSignallingServerURL(FString::Printf(TEXT("%s:%s"), *SignallingDomain, *FString::FromInt(StreamerPort)));
+	UE_LOG(LogPixelStreamingVCam, Log, TEXT("Activating PixelStreaming VCam Session. Endpoint: %s:%s"), *SignallingDomain, *FString::FromInt(StreamerPort));
 
 	MediaCapture = Cast<UPixelStreamingMediaCapture>(MediaOutput->CreateMediaCapture());
 
@@ -178,7 +189,8 @@ void UVCamPixelStreamingSession::Activate()
 
 void UVCamPixelStreamingSession::StopSignallingServer()
 {
-	if (UVCamPixelStreamingSubsystem* PixelStreamingSubsystem = UVCamPixelStreamingSubsystem::Get())
+	// Only stop the signalling server if we've been the ones to start it
+	if (UVCamPixelStreamingSubsystem* PixelStreamingSubsystem = UVCamPixelStreamingSubsystem::Get(); StartSignallingServer)
 	{
 		PixelStreamingSubsystem->StopSignallingServer();
 	}
@@ -233,16 +245,10 @@ void UVCamPixelStreamingSession::PostEditChangeProperty(FPropertyChangedEvent& P
 	FProperty* Property = PropertyChangedEvent.MemberProperty;
 	if (Property && PropertyChangedEvent.ChangeType != EPropertyChangeType::Interactive)
 	{
-		static FName NAME_IP = GET_MEMBER_NAME_CHECKED(UVCamPixelStreamingSession, IP);
-		static FName NAME_PortNumber = GET_MEMBER_NAME_CHECKED(UVCamPixelStreamingSession, PortNumber);
 		static FName NAME_FromComposureOutputProviderIndex = GET_MEMBER_NAME_CHECKED(UVCamPixelStreamingSession, FromComposureOutputProviderIndex);
 		static FName NAME_StartSignallingServer = GET_MEMBER_NAME_CHECKED(UVCamPixelStreamingSession, StartSignallingServer);
 		
-		// If any of these properties change then deactivate streaming so user is forced to restart streaming
-		if ((Property->GetFName() == NAME_IP) || 
-			(Property->GetFName() == NAME_PortNumber) || 
-			(Property->GetFName() == NAME_FromComposureOutputProviderIndex) || 
-			(Property->GetFName() == NAME_StartSignallingServer))
+		if (Property->GetFName() == NAME_FromComposureOutputProviderIndex || Property->GetFName() == NAME_StartSignallingServer)
 		{
 			SetActive(false);
 		}
diff --git a/Engine/Plugins/Media/PixelStreamingVCam/Source/PixelStreamingVCam/Private/VCamPixelStreamingSubsystem.cpp b/Engine/Plugins/Media/PixelStreamingVCam/Source/PixelStreamingVCam/Private/VCamPixelStreamingSubsystem.cpp
--- a/Engine/Plugins/Media/PixelStreamingVCam/Source/PixelStreamingVCam/Private/VCamPixelStreamingSubsystem.cpp
+++ b/Engine/Plugins/Media/PixelStreamingVCam/Source/PixelStreamingVCam/Private/VCamPixelStreamingSubsystem.cpp
@@ -5,6 +5,7 @@
 #include "VCamPixelStreamingLiveLink.h"
 #include "VCamPixelStreamingSession.h"
 #include "PixelStreamingVCamLog.h"
+#include "PixelStreamingEditorModule.h"
 
 void UVCamPixelStreamingSubsystem::Initialize(FSubsystemCollectionBase& Collection)
 {
@@ -86,43 +87,12 @@ TSharedPtr<FPixelStreamingLiveLinkSource> UVCamPixelStreamingSubsystem::TryGetLi
 	return LiveLinkSource;
 }
 
-TSharedPtr<UE::PixelStreamingServers::IServer> UVCamPixelStreamingSubsystem::LaunchSignallingServer(int StreamerPort, int PlayerPort)
+void UVCamPixelStreamingSubsystem::LaunchSignallingServer()
 {
-	bool bAlreadyLaunched = Server.IsValid() && Server->HasLaunched();
-	if(bAlreadyLaunched)
-	{
-		return Server;
-	}
-
-	// Download Pixel Streaming servers/frontend if we want to use a browser to view Pixel Streaming output
-	// but only attempt this is we haven't already started a download before.
-	if(!DownloadProcess.IsValid())
-	{
-		DownloadProcess = UE::PixelStreamingServers::DownloadPixelStreamingServers(/*bSkipIfPresent*/ true);
-		if(DownloadProcess.IsValid())
-		{
-			DownloadProcess->OnCompleted().BindLambda([this, StreamerPort, PlayerPort](int ExitCode){
-				UE_LOG(LogPixelStreamingVCam, Log, TEXT("Pixel Streaming frontend downloaded, restarting signalling server to use it."));
-				StopSignallingServer();
-				LaunchSignallingServer(StreamerPort, PlayerPort);
-			});
-		}
-	}
-
-	// Launch signalling server
-	Server = UE::PixelStreamingServers::MakeSignallingServer();
-	UE::PixelStreamingServers::FLaunchArgs LaunchArgs;
-	LaunchArgs.bPollUntilReady = false;
-	LaunchArgs.ProcessArgs = FString::Printf(TEXT("--StreamerPort=%d --HttpPort=%d"), StreamerPort, PlayerPort);
-	Server->Launch(LaunchArgs);
-	return Server;
+	FPixelStreamingEditorModule::GetModule()->StartSignalling();
 }
 
 void UVCamPixelStreamingSubsystem::StopSignallingServer()
 {
-	if(Server.IsValid())
-	{
-		Server->Stop();
-		Server.Reset();
-	}
+	FPixelStreamingEditorModule::GetModule()->StopSignalling();
 }
\ No newline at end of file
diff --git a/Engine/Plugins/Media/PixelStreamingVCam/Source/PixelStreamingVCam/Public/VCamPixelStreamingSession.h b/Engine/Plugins/Media/PixelStreamingVCam/Source/PixelStreamingVCam/Public/VCamPixelStreamingSession.h
--- a/Engine/Plugins/Media/PixelStreamingVCam/Source/PixelStreamingVCam/Public/VCamPixelStreamingSession.h
+++ b/Engine/Plugins/Media/PixelStreamingVCam/Source/PixelStreamingVCam/Public/VCamPixelStreamingSession.h
@@ -25,18 +25,7 @@ public:
 	virtual void Tick(const float DeltaTime) override;
 #if WITH_EDITOR
 	virtual void PostEditChangeProperty(FPropertyChangedEvent& PropertyChangedEvent) override;
-#endif
-	UPROPERTY(EditAnywhere, BlueprintReadWrite, Category = "Output")
-	FString IP = TEXT("ws://127.0.0.1");
-
-	// Network port number - change this only if connecting multiple RemoteSession devices to the same PC
-	UPROPERTY(EditAnywhere, BlueprintReadWrite, Category = "Output")
-	int32 PortNumber = 8888;
-	
-	// Http webserver port number - e.g. Go to http://localhost:YourPort to access the streamed VCam output. Warning Ports below 1024 require sudo under Linux so it is recommended to use a higher port on that platform.
-	UPROPERTY(EditAnywhere, BlueprintReadWrite, Category = "Output")
-	int32 HttpPort = 80;
-	
+#endif	
 	// If using the output from a Composure Output Provider, specify it here
 	UPROPERTY(EditAnywhere, BlueprintReadWrite, Category = "Output")
 	int32 FromComposureOutputProviderIndex = INDEX_NONE;
diff --git a/Engine/Plugins/Media/PixelStreamingVCam/Source/PixelStreamingVCam/Public/VCamPixelStreamingSubsystem.h b/Engine/Plugins/Media/PixelStreamingVCam/Source/PixelStreamingVCam/Public/VCamPixelStreamingSubsystem.h
--- a/Engine/Plugins/Media/PixelStreamingVCam/Source/PixelStreamingVCam/Public/VCamPixelStreamingSubsystem.h
+++ b/Engine/Plugins/Media/PixelStreamingVCam/Source/PixelStreamingVCam/Public/VCamPixelStreamingSubsystem.h
@@ -28,7 +28,7 @@ public:
 	// Additionally registers the OutputProvider to the Source if a new source was created
 	TSharedPtr<FPixelStreamingLiveLinkSource> TryGetLiveLinkSource(UVCamPixelStreamingSession* OutputProvider);
 
-	TSharedPtr<UE::PixelStreamingServers::IServer> LaunchSignallingServer(int StreamerPort, int PlayerPort);
+	void LaunchSignallingServer();
 	void StopSignallingServer();
 	
 private:
@@ -39,9 +39,6 @@ private:
 	// Signalling/webserver
 	TSharedPtr<UE::PixelStreamingServers::IServer> Server;
 
-	// Download process for PS web frontend files (if we want to view output in the browser)
-	TSharedPtr<FMonitoredProcess> DownloadProcess;
-
 	// An associated Live Link Source shared by all output providers
 	TSharedPtr<FPixelStreamingLiveLinkSource> LiveLinkSource;
 };
