Fixing the starting bitrate of streams

diff --git a/Engine/Plugins/Media/PixelStreaming/Source/PixelStreaming/Private/PixelStreamingPeerConnection.cpp b/Engine/Plugins/Media/PixelStreaming/Source/PixelStreaming/Private/PixelStreamingPeerConnection.cpp
--- a/Engine/Plugins/Media/PixelStreaming/Source/PixelStreaming/Private/PixelStreamingPeerConnection.cpp
+++ b/Engine/Plugins/Media/PixelStreaming/Source/PixelStreaming/Private/PixelStreamingPeerConnection.cpp
@@ -246,6 +246,13 @@ TUniquePtr<FPixelStreamingPeerConnection> FPixelStreamingPeerConnection::Create(
 		NewPeerConnection->PeerConnection = Result.MoveValue();
 	#endif
 
+	// Setup suggested bitrate settings on the Peer Connection based on our CVars
+	webrtc::BitrateSettings BitrateSettings;
+	BitrateSettings.min_bitrate_bps = Settings::CVarPixelStreamingWebRTCMinBitrate.GetValueOnAnyThread();
+	BitrateSettings.max_bitrate_bps = Settings::CVarPixelStreamingWebRTCMaxBitrate.GetValueOnAnyThread();
+	BitrateSettings.start_bitrate_bps = Settings::CVarPixelStreamingWebRTCStartBitrate.GetValueOnAnyThread();
+	NewPeerConnection->PeerConnection->SetBitrate(BitrateSettings);
+
 	return NewPeerConnection;
 }
 
