Fix duplicate LocText in PixelStreamingToolbar.cpp - Issue 218234: Localization warnings in PixelStreamingToolbar.cpp

diff --git a/Engine/Plugins/Media/PixelStreaming/Source/PixelStreamingEditor/Private/PixelStreamingToolbar.cpp b/Engine/Plugins/Media/PixelStreaming/Source/PixelStreamingEditor/Private/PixelStreamingToolbar.cpp
--- a/Engine/Plugins/Media/PixelStreaming/Source/PixelStreamingEditor/Private/PixelStreamingToolbar.cpp
+++ b/Engine/Plugins/Media/PixelStreaming/Source/PixelStreamingEditor/Private/PixelStreamingToolbar.cpp
@@ -227,7 +227,7 @@ namespace UE::EditorPixelStreaming
 								FMenuBuilder MenuBuilder(true, PluginCommands);  
 
 								// Use external signalling server option
-								MenuBuilder.BeginSection("Signalling Server Location", LOCTEXT("PixelStreamingSettings", "Signalling Server Location"));
+								MenuBuilder.BeginSection("Signalling Server Location", LOCTEXT("PixelStreamingSSLocation", "Signalling Server Location"));
 								MenuBuilder.AddMenuEntry(FPixelStreamingCommands::Get().ExternalSignalling);
 								MenuBuilder.EndSection();
 
@@ -281,7 +281,7 @@ namespace UE::EditorPixelStreaming
 
 	void FPixelStreamingToolbar::RegisterEmbeddedSignallingServerConfig(FMenuBuilder& MenuBuilder)
 	{
-		MenuBuilder.BeginSection("Signalling Server Options", LOCTEXT("PixelStreamingSettings", "Signalling Server Options"));
+		MenuBuilder.BeginSection("Signalling Server Options", LOCTEXT("PixelStreamingEmbeddedSSOptions", "Embedded Signalling Server Options"));
 
 		TSharedPtr<UE::PixelStreamingServers::IServer> SignallingServer = FPixelStreamingEditorModule::GetModule()->GetSignallingServer();
 		if(!SignallingServer.IsValid() || !SignallingServer->HasLaunched())
@@ -357,7 +357,7 @@ namespace UE::EditorPixelStreaming
 	void FPixelStreamingToolbar::RegisterRemoteSignallingServerConfig(FMenuBuilder& MenuBuilder)
 	{
 		TSharedPtr<UE::PixelStreamingServers::IServer> SignallingServer = FPixelStreamingEditorModule::GetModule()->GetSignallingServer();
-		MenuBuilder.BeginSection("Remote Signalling Server Options", LOCTEXT("PixelStreamingSettings", "Remote Signalling Server Options"));
+		MenuBuilder.BeginSection("Remote Signalling Server Options", LOCTEXT("PixelStreamingRemoteSSOptions", "Remote Signalling Server Options"));
 		{
 			TSharedRef<SWidget> URLInputBlock = SNew(SHorizontalBox)
 														+ SHorizontalBox::Slot()
