Expose LeftTriggerAnalog and RightTriggerAnalog as acceptable button press options in Pixel Streaming. Although they are not buttons if users want to use them as buttons why not?

diff --git a/Engine/Plugins/Media/PixelStreaming/Source/PixelStreaming/Private/PixelStreamingInputHandler.cpp b/Engine/Plugins/Media/PixelStreaming/Source/PixelStreaming/Private/PixelStreamingInputHandler.cpp
--- a/Engine/Plugins/Media/PixelStreaming/Source/PixelStreaming/Private/PixelStreamingInputHandler.cpp
+++ b/Engine/Plugins/Media/PixelStreaming/Source/PixelStreaming/Private/PixelStreamingInputHandler.cpp
@@ -713,7 +713,16 @@ namespace UE::PixelStreaming
 				return FGamepadKeyNames::RightShoulder;
 			}
 			break;
-			// Buttons 6 and 7 are mapped as analog axis as they are the triggers
+			case 6:
+			{
+				return FGamepadKeyNames::LeftTriggerAnalog;
+			}
+			break;
+			case 7:
+			{
+				return FGamepadKeyNames::RightTriggerAnalog;
+			}
+			break;
 			case 8:
 			{
 				return FGamepadKeyNames::SpecialLeft;
