Updated trigger buttons to use the proper Left/RightTriggerThreshold values.

diff --git a/Engine/Plugins/Media/PixelStreaming/Source/PixelStreaming/Private/PixelStreamingInputHandler.cpp b/Engine/Plugins/Media/PixelStreaming/Source/PixelStreaming/Private/PixelStreamingInputHandler.cpp
--- a/Engine/Plugins/Media/PixelStreaming/Source/PixelStreaming/Private/PixelStreamingInputHandler.cpp
+++ b/Engine/Plugins/Media/PixelStreaming/Source/PixelStreaming/Private/PixelStreamingInputHandler.cpp
@@ -715,12 +715,12 @@ namespace UE::PixelStreaming
 			break;
 			case 6:
 			{
-				return FGamepadKeyNames::LeftTriggerAnalog;
+				return FGamepadKeyNames::LeftTriggerThreshold;
 			}
 			break;
 			case 7:
 			{
-				return FGamepadKeyNames::RightTriggerAnalog;
+				return FGamepadKeyNames::RightTriggerThreshold;
 			}
 			break;
 			case 8:
