Fix mouse up crashing when clicking into the stream, fix was to catch the case where no mouse enter or mouse move had yet been performed.

diff --git a/Engine/Plugins/Media/PixelStreaming/Source/PixelStreaming/Private/PixelStreamingInputHandler.cpp b/Engine/Plugins/Media/PixelStreaming/Source/PixelStreaming/Private/PixelStreamingInputHandler.cpp
--- a/Engine/Plugins/Media/PixelStreaming/Source/PixelStreaming/Private/PixelStreamingInputHandler.cpp
+++ b/Engine/Plugins/Media/PixelStreaming/Source/PixelStreaming/Private/PixelStreamingInputHandler.cpp
@@ -391,6 +391,12 @@ namespace UE::PixelStreaming
 
     void FPixelStreamingInputHandler::HandleOnMouseUp(FMemoryReader Ar)
     {
+		// Ensure we have wrapped the slate application at this point
+		if(!bIsMouseActive)
+		{
+			HandleOnMouseEnter(Ar);
+		}
+
         TPayloadThreeParam<uint8, uint16, uint16> Payload(Ar);
 
         EMouseButtons::Type Button = static_cast<EMouseButtons::Type>(Payload.Param1);
@@ -403,11 +409,21 @@ namespace UE::PixelStreaming
 
     void FPixelStreamingInputHandler::HandleOnMouseDown(FMemoryReader Ar)
     {
+		// Ensure we have wrapped the slate application at this point
+		if(!bIsMouseActive)
+		{
+			HandleOnMouseEnter(Ar);
+		}
+
         TPayloadThreeParam<uint8, uint16, uint16> Payload(Ar);
         //                                                                           convert range from 0,65536 -> 0,1
         FVector2D ScreenLocation = ConvertFromNormalizedScreenLocation(FVector2D(Payload.Param2 / uint16_MAX, Payload.Param3 / uint16_MAX));
         EMouseButtons::Type Button = static_cast<EMouseButtons::Type>(Payload.Param1);
 
+		// Set cursor pos on mouse down - we may not have moved if this is the very first click
+		FSlateApplication::Get().OnCursorSet();
+		PixelStreamerApplicationWrapper->Cursor->SetPosition(ScreenLocation.X, ScreenLocation.Y);
+
 		UE_LOG(LogPixelStreamingInputHandler, Verbose, TEXT("MOUSE_DOWN: Button = %d; Pos = (%.4f, %.4f)"), Button, ScreenLocation.X, ScreenLocation.Y);
 		// Force window focus
 		FSlateApplication::Get().ProcessApplicationActivationEvent(true);
