Explicitly enable alr_probing for video in Pixel Streaming - it "should" have been on by default anyway, but being explicit won't hurt against future WebRTC API changes.

diff --git a/Engine/Plugins/Media/PixelStreaming/Source/PixelStreaming/Private/Streamer.cpp b/Engine/Plugins/Media/PixelStreaming/Source/PixelStreaming/Private/Streamer.cpp
--- a/Engine/Plugins/Media/PixelStreaming/Source/PixelStreaming/Private/Streamer.cpp
+++ b/Engine/Plugins/Media/PixelStreaming/Source/PixelStreaming/Private/Streamer.cpp
@@ -437,6 +437,10 @@ namespace UE::PixelStreaming
 #if WEBRTC_VERSION == 84
 		PeerConnectionConfig.enable_simulcast_stats = true;
 #endif
+
+		// We want periodic bandwidth probing so ramping happens quickly
+		PeerConnectionConfig.media_config.video.periodic_alr_bandwidth_probing = true;
+
 	}
 
 	void FStreamer::OnSignallingSessionDescription(FPixelStreamingPlayerId PlayerId, webrtc::SdpType Type, const FString& Sdp)
