Add launch arg `-PixelStreamingNegotiateCodecs` to negotiate codecs. Default case: negotiate codecs is disabled and UE will only send a single selected codec.

diff --git a/Engine/Plugins/Media/PixelStreaming/Source/PixelStreaming/Private/Settings.h b/Engine/Plugins/Media/PixelStreaming/Source/PixelStreaming/Private/Settings.h
--- a/Engine/Plugins/Media/PixelStreaming/Source/PixelStreaming/Private/Settings.h
+++ b/Engine/Plugins/Media/PixelStreaming/Source/PixelStreaming/Private/Settings.h
@@ -111,6 +111,11 @@ namespace UE::PixelStreaming::Settings
 		return FParse::Param(FCommandLine::Get(), TEXT("PixelStreamingHideCursor"));
 	}
 
+	inline bool ShouldNegotiateCodecs()
+	{
+		return FParse::Param(FCommandLine::Get(), TEXT("PixelStreamingNegotiateCodecs"));
+	}
+
 	inline bool IsUsingSafeTextureCopy()
 	{
 		return FParse::Param(FCommandLine::Get(), TEXT("PixelCaptureUseFence"));
diff --git a/Engine/Plugins/Media/PixelStreaming/Source/PixelStreaming/Private/VideoEncoderFactorySingleLayer.cpp b/Engine/Plugins/Media/PixelStreaming/Source/PixelStreaming/Private/VideoEncoderFactorySingleLayer.cpp
--- a/Engine/Plugins/Media/PixelStreaming/Source/PixelStreaming/Private/VideoEncoderFactorySingleLayer.cpp
+++ b/Engine/Plugins/Media/PixelStreaming/Source/PixelStreaming/Private/VideoEncoderFactorySingleLayer.cpp
@@ -46,6 +46,7 @@ namespace UE::PixelStreaming
 	// which leads to very low bitrate streams.
 	std::vector<webrtc::SdpVideoFormat> FVideoEncoderFactorySingleLayer::GetSupportedFormats() const
 	{
+
 		// static so we dont create the list every time this is called since the list will not change during runtime.
 		static TMap<EPixelStreamingCodec, std::vector<webrtc::SdpVideoFormat>> CodecMap = CreateSupportedFormatMap();
 
@@ -53,6 +54,25 @@ namespace UE::PixelStreaming
 		// but that actually works in our favor since we're describing more about the plugin state than the actual
 		// instance of this factory.
 		static std::vector<webrtc::SdpVideoFormat> SupportedFormats;
+
+		// If we are not negotiating codecs simply return just the one codec that is selected in UE
+		if(!Settings::ShouldNegotiateCodecs())
+		{
+			const EPixelStreamingCodec SelectedCodec = UE::PixelStreaming::Settings::GetSelectedCodec();
+			if (CodecMap.Contains(SelectedCodec))
+			{
+				for (auto& Format : CodecMap[SelectedCodec])
+				{
+					SupportedFormats.push_back(Format);
+				}
+			}
+			else
+			{
+				UE_LOG(LogPixelStreaming, Error, TEXT("Selected codec was not a supported codec."));
+			}
+			return SupportedFormats;
+		}
+
 		static EPixelStreamingCodec LastSelectedCodec = EPixelStreamingCodec::Invalid;
 
 		const EPixelStreamingCodec SelectedCodec = UE::PixelStreaming::Settings::GetSelectedCodec();
