Fix: Stream FPS couldn't be configured from the browser

diff --git a/Engine/Plugins/Media/PixelStreaming/Source/PixelStreaming/Private/PixelStreamingPeerConnection.cpp b/Engine/Plugins/Media/PixelStreaming/Source/PixelStreaming/Private/PixelStreamingPeerConnection.cpp
--- a/Engine/Plugins/Media/PixelStreaming/Source/PixelStreaming/Private/PixelStreamingPeerConnection.cpp
+++ b/Engine/Plugins/Media/PixelStreaming/Source/PixelStreaming/Private/PixelStreamingPeerConnection.cpp
@@ -51,7 +51,7 @@ namespace
 				// Try re-enabling this when we upgrade WebRTC versions.
 				// LayerEncoding.network_priority = webrtc::Priority::kHigh;
 
-				LayerEncoding.max_framerate = FMath::Max(60, Settings::CVarPixelStreamingWebRTCFps.GetValueOnAnyThread());
+				LayerEncoding.max_framerate = Settings::CVarPixelStreamingWebRTCFps.GetValueOnAnyThread();
 				EncodingParams.push_back(LayerEncoding);
 			}
 		}
@@ -61,7 +61,7 @@ namespace
 			Encoding.rid = "base";
 			Encoding.max_bitrate_bps = Settings::CVarPixelStreamingWebRTCMaxBitrate.GetValueOnAnyThread();
 			Encoding.min_bitrate_bps = Settings::CVarPixelStreamingWebRTCMinBitrate.GetValueOnAnyThread();
-			Encoding.max_framerate = FMath::Max(60, Settings::CVarPixelStreamingWebRTCFps.GetValueOnAnyThread());
+			Encoding.max_framerate = Settings::CVarPixelStreamingWebRTCFps.GetValueOnAnyThread();
 			Encoding.scale_resolution_down_by.reset();
 			Encoding.network_priority = webrtc::Priority::kHigh;
 			EncodingParams.push_back(Encoding);
