Fix: Viewport streaming could interact with windows which weren't targeted

diff --git a/Engine/Plugins/Media/PixelStreaming/Source/PixelStreaming/Private/PixelStreamingInputHandler.cpp b/Engine/Plugins/Media/PixelStreaming/Source/PixelStreaming/Private/PixelStreamingInputHandler.cpp
--- a/Engine/Plugins/Media/PixelStreaming/Source/PixelStreaming/Private/PixelStreamingInputHandler.cpp
+++ b/Engine/Plugins/Media/PixelStreaming/Source/PixelStreaming/Private/PixelStreamingInputHandler.cpp
@@ -125,6 +125,7 @@ namespace UE::PixelStreaming
 	void FPixelStreamingInputHandler::SetTargetWindow(TWeakPtr<SWindow> InWindow)
 	{
 		TargetWindow = InWindow;
+		PixelStreamerApplicationWrapper->SetTargetWindow(InWindow);
 	}
 
 	TWeakPtr<SWindow> FPixelStreamingInputHandler::GetTargetWindow()
