Add Pixel Streaming toolbar to new toolbar group which should enforce a separator that was missing in some cases causing user confusion.

diff --git a/Engine/Plugins/Media/PixelStreaming/Source/PixelStreamingEditor/Private/PixelStreamingToolbar.cpp b/Engine/Plugins/Media/PixelStreaming/Source/PixelStreamingEditor/Private/PixelStreamingToolbar.cpp
--- a/Engine/Plugins/Media/PixelStreaming/Source/PixelStreamingEditor/Private/PixelStreamingToolbar.cpp
+++ b/Engine/Plugins/Media/PixelStreaming/Source/PixelStreamingEditor/Private/PixelStreamingToolbar.cpp
@@ -212,7 +212,7 @@ namespace UE::EditorPixelStreaming
 	{
 		FToolMenuOwnerScoped OwnerScoped(this);
 		{
-			UToolMenu* CustomToolBar = UToolMenus::Get()->ExtendMenu("LevelEditor.LevelEditorToolBar.PlayToolBar");
+			UToolMenu* CustomToolBar = UToolMenus::Get()->ExtendMenu("LevelEditor.LevelEditorToolBar.User");
 			{
 				FToolMenuSection& Section = CustomToolBar->AddSection("PixelStreaming");
 				Section.AddSeparator("PixelStreamingSeperator");
