Fix: Touch input not scaling and interacting with windows which weren't targeted when using Virtual Cameras. This has been achieved by enabling the PixelStreamingInputHandler to route input directly to a widget if requested. Default behaviour is input is routed to windows. Controlling this behaviour is done on the streamer with the new enum EPixelStreamingInputType.

diff --git a/Engine/Plugins/Experimental/VirtualCamera/Source/PixelStreamingVCam/PixelStreamingVCam.Build.cs b/Engine/Plugins/Experimental/VirtualCamera/Source/PixelStreamingVCam/PixelStreamingVCam.Build.cs
--- a/Engine/Plugins/Experimental/VirtualCamera/Source/PixelStreamingVCam/PixelStreamingVCam.Build.cs
+++ b/Engine/Plugins/Experimental/VirtualCamera/Source/PixelStreamingVCam/PixelStreamingVCam.Build.cs
@@ -32,7 +32,11 @@ namespace UnrealBuildTool.Rules
 				"CinematicCamera",
 				"LiveLinkInterface",
 				"PixelCapture",
-				"PixelStreamingEditor"
+				"PixelStreamingEditor",
+				"SlateCore",
+				"Slate",
+				"VPUtilities",
+				"InputCore"
 			});
 
 			// Can't package non-editor targets (e.g. games) with UnrealEd, so this dependency should only be added in editor.
diff --git a/Engine/Plugins/Experimental/VirtualCamera/Source/PixelStreamingVCam/Private/PixelStreamingMediaOutput.cpp b/Engine/Plugins/Experimental/VirtualCamera/Source/PixelStreamingVCam/Private/PixelStreamingMediaOutput.cpp
--- a/Engine/Plugins/Experimental/VirtualCamera/Source/PixelStreamingVCam/Private/PixelStreamingMediaOutput.cpp
+++ b/Engine/Plugins/Experimental/VirtualCamera/Source/PixelStreamingVCam/Private/PixelStreamingMediaOutput.cpp
@@ -30,7 +30,6 @@ UMediaCapture* UPixelStreamingMediaOutput::CreateMediaCaptureImpl()
 	{
 		Capture = NewObject<UPixelStreamingMediaCapture>();
 		Capture->SetMediaOutput(this);
-		Capture->OnCaptureViewportInitialized.AddUObject(this, &UPixelStreamingMediaOutput::OnCaptureViewportInitialized);
 	}
 
 	if (!VideoInput)
@@ -43,15 +42,6 @@ UMediaCapture* UPixelStreamingMediaOutput::CreateMediaCaptureImpl()
 	return Capture;
 }
 
-void UPixelStreamingMediaOutput::OnCaptureViewportInitialized()
-{
-	if(Streamer)
-	{
-		Streamer->SetTargetViewport(Capture->GetViewport()->GetViewportWidget());
-		Streamer->SetTargetWindow(Capture->GetViewport()->FindWindow());
-	}
-}
-
 void UPixelStreamingMediaOutput::StartStreaming()
 {
 	if (Streamer)
@@ -80,7 +70,6 @@ void UPixelStreamingMediaOutput::StopStreaming()
 	if (Streamer)
 	{
 		Streamer->StopStreaming();
-		Streamer->SetTargetViewport(nullptr);
 		Streamer->SetTargetWindow(nullptr);
 	}
 }
diff --git a/Engine/Plugins/Experimental/VirtualCamera/Source/PixelStreamingVCam/Private/VCamPixelStreamingSession.cpp b/Engine/Plugins/Experimental/VirtualCamera/Source/PixelStreamingVCam/Private/VCamPixelStreamingSession.cpp
--- a/Engine/Plugins/Experimental/VirtualCamera/Source/PixelStreamingVCam/Private/VCamPixelStreamingSession.cpp
+++ b/Engine/Plugins/Experimental/VirtualCamera/Source/PixelStreamingVCam/Private/VCamPixelStreamingSession.cpp
@@ -18,6 +18,11 @@
 #include "PixelStreamingProtocol.h"
 #include "PixelStreamingEditorModule.h"
 #include "VCamComponent.h"
+#include "Input/HittestGrid.h"
+#include "Widgets/SVirtualWindow.h"
+#include "VPFullScreenUserWidget.h"
+#include "InputCoreTypes.h"
+#include "IPixelStreamingInputHandler.h"
 
 namespace VCamPixelStreamingSession
 {
@@ -86,6 +91,9 @@ void UVCamPixelStreamingSession::Activate()
 	// This sets up media capture and streamer
 	SetupCapture();
 
+	// Super::Activate() creates our UMG which we need before setting up our custom input handling
+	Super::Activate();
+
 	// We setup custom handling of ARKit transforms coming from iOS devices here
 	SetupCustomInputHandling();
 
@@ -98,7 +106,6 @@ void UVCamPixelStreamingSession::Activate()
 	MediaOutput->SetSignallingServerURL(FString::Printf(TEXT("%s:%s"), *SignallingDomain, *FString::FromInt(StreamerPort)));
 	UE_LOG(LogPixelStreamingVCam, Log, TEXT("Activating PixelStreaming VCam Session. Endpoint: %s:%s"), *SignallingDomain, *FString::FromInt(StreamerPort));
 	
-	Super::Activate();
 }
 
 void UVCamPixelStreamingSession::SetupCapture()
@@ -152,14 +159,51 @@ void UVCamPixelStreamingSession::OnCaptureStateChanged()
 
 void UVCamPixelStreamingSession::SetupCustomInputHandling()
 {
+	// If we have a UMG, then use it
+	if (UMGWidget)
+	{
+		TSharedPtr<SVirtualWindow> InputWindow;
+		// If we are rendering from a ComposureOutputProvider, we need to get the InputWindow from that UMG, not the one in the PixelStreamingOutputProvider
+		if (UVCamOutputComposure* ComposureProvider = Cast<UVCamOutputComposure>(GetOtherOutputProviderByIndex(FromComposureOutputProviderIndex)))
+		{
+			if (UVPFullScreenUserWidget* ComposureUMGWidget = ComposureProvider->GetUMGWidget())
+			{
+				InputWindow = ComposureUMGWidget->PostProcessDisplayType.GetSlateWindow();
+				UE_LOG(LogPixelStreamingVCam, Log, TEXT("InputChannel callback - Routing input to active viewport with Composure UMG"));
+			}
+			else
+			{
+				UE_LOG(LogPixelStreamingVCam, Warning, TEXT("InputChannel callback - Composure usage was requested, but the specified ComposureOutputProvider has no UMG set"));
+			}
+		}
+		else
+		{
+			InputWindow = UMGWidget->PostProcessDisplayType.GetSlateWindow();
+			UE_LOG(LogPixelStreamingVCam, Log, TEXT("InputChannel callback - Routing input to active viewport with UMG"));
+		}
+		//
+		MediaOutput->GetStreamer()->SetTargetWindow(InputWindow);
+		MediaOutput->GetStreamer()->SetInputHandlerType(EPixelStreamingInputType::RouteToWidget);
+	}
+	else
+	{
+		MediaOutput->GetStreamer()->SetTargetWindow(GetTargetInputWindow());
+		MediaOutput->GetStreamer()->SetInputHandlerType(EPixelStreamingInputType::RouteToWidget);
+		UE_LOG(LogPixelStreamingVCam, Log, TEXT("InputChannel callback - Routing input to active viewport"));
+	}
+
 	if(MediaOutput)
 	{
 		IPixelStreamingModule& PixelStreamingModule = IPixelStreamingModule::Get();
-		
+		typedef Protocol::EPixelStreamingMessageTypes EType;		
 		Protocol::EPixelStreamingMessageDirection MessageDirection = Protocol::EPixelStreamingMessageDirection::ToStreamer;
-
-		typedef Protocol::EPixelStreamingMessageTypes EType;
-		Protocol::FPixelStreamingInputMessage Message = Protocol::FPixelStreamingInputMessage(100, 72, {
+		TMap<FString, Protocol::FPixelStreamingInputMessage> Protocol = PixelStreamingModule.GetProtocol().ToStreamerProtocol;
+		/*
+		* ====================
+		* ARKit Transform
+		* ====================
+		*/
+		Protocol::FPixelStreamingInputMessage ARKitMessage = Protocol::FPixelStreamingInputMessage(100, 72, {
 			// 4x4 Transform
 			EType::Float, EType::Float, EType::Float, EType::Float,
 			EType::Float, EType::Float, EType::Float, EType::Float,
@@ -168,8 +212,8 @@ void UVCamPixelStreamingSession::SetupCustomInputHandling()
 			// Timestamp
 			EType::Double
 		});
-
-		const TFunction<void(FMemoryReader)>& Handler = [this](FMemoryReader Ar) { 
+		
+		const TFunction<void(FMemoryReader)>& ARKitHandler = [this](FMemoryReader Ar) { 
 
 			if(!EnableARKitTracking)
 			{
@@ -197,21 +241,8 @@ void UVCamPixelStreamingSession::SetupCustomInputHandling()
 			{
 				LiveLinkSource->PushTransformForSubject(GetFName(), FTransform(ARKitMatrix), Timestamp);
 			}
-
-			/**
-			 * Code to control the level editor viewport
-			 */
-			// for (FLevelEditorViewportClient* LevelVC : GEditor->GetLevelViewportClients())
-			// {
-			// 	if (LevelVC && LevelVC->IsPerspective())
-			// 	{
-			// 		LevelVC->SetViewLocation(Translation);
-			// 		LevelVC->SetViewRotation(ModifiedRotation);
-			// 	}
-			// }
 		};
-
-		PixelStreamingModule.RegisterMessage(MessageDirection, "ARKitTransform", Message, Handler);
+		PixelStreamingModule.RegisterMessage(MessageDirection, "ARKitTransform", ARKitMessage, ARKitHandler);
 	}
 	else
 	{
@@ -261,7 +292,8 @@ void UVCamPixelStreamingSession::StartCapture()
 void UVCamPixelStreamingSession::SetupSignallingServer()
 {
 	// Only start the signalling server if we aren't using an external signalling server
-	if (UVCamPixelStreamingSubsystem* PixelStreamingSubsystem = UVCamPixelStreamingSubsystem::Get(); !FPixelStreamingEditorModule::GetModule()->bUseExternalSignallingServer)
+	UVCamPixelStreamingSubsystem* PixelStreamingSubsystem = UVCamPixelStreamingSubsystem::Get();
+	if (PixelStreamingSubsystem && !FPixelStreamingEditorModule::GetModule()->bUseExternalSignallingServer)
 	{
 		PixelStreamingSubsystem->LaunchSignallingServer();
 	}
@@ -270,7 +302,8 @@ void UVCamPixelStreamingSession::SetupSignallingServer()
 void UVCamPixelStreamingSession::StopSignallingServer()
 {
 	// Only stop the signalling server if we've been the ones to start it
-	if (UVCamPixelStreamingSubsystem* PixelStreamingSubsystem = UVCamPixelStreamingSubsystem::Get(); !FPixelStreamingEditorModule::GetModule()->bUseExternalSignallingServer)
+	UVCamPixelStreamingSubsystem* PixelStreamingSubsystem = UVCamPixelStreamingSubsystem::Get();
+	if (PixelStreamingSubsystem && !FPixelStreamingEditorModule::GetModule()->bUseExternalSignallingServer)
 	{
 		PixelStreamingSubsystem->StopSignallingServer();
 	}
@@ -346,4 +379,14 @@ void UVCamPixelStreamingSession::PostEditChangeProperty(FPropertyChangedEvent& P
 }
 #endif
 
+TWeakPtr<SWindow> UVCamPixelStreamingSession::GetTargetInputWindow() const
+{
+	TWeakPtr<SWindow> InputWindow;
+	if (UVCamComponent* OuterComponent = GetTypedOuter<UVCamComponent>())
+	{
+		InputWindow = OuterComponent->GetTargetInputWindow();
+	}
+	return InputWindow;
+}
+
 IMPLEMENT_MODULE(FDefaultModuleImpl, PixelStreamingVCam)
diff --git a/Engine/Plugins/Experimental/VirtualCamera/Source/PixelStreamingVCam/Public/VCamPixelStreamingSession.h b/Engine/Plugins/Experimental/VirtualCamera/Source/PixelStreamingVCam/Public/VCamPixelStreamingSession.h
--- a/Engine/Plugins/Experimental/VirtualCamera/Source/PixelStreamingVCam/Public/VCamPixelStreamingSession.h
+++ b/Engine/Plugins/Experimental/VirtualCamera/Source/PixelStreamingVCam/Public/VCamPixelStreamingSession.h
@@ -57,9 +57,10 @@ private:
 	void SetupCustomInputHandling();
 	void OnCaptureStateChanged();
 	void OnARKitTransformReceived(FPixelStreamingPlayerId PlayerId, uint8 Type, TArray<uint8> Data);
+	TWeakPtr<SWindow> GetTargetInputWindow() const;
 
 private:
-	FHitResult LastViewportTouchResult;
-	bool bUsingDummyUMG = false;
-	bool bOldThrottleCPUWhenNotForeground;
+	FHitResult 	LastViewportTouchResult;
+	bool 		bUsingDummyUMG = false;
+	bool 		bOldThrottleCPUWhenNotForeground;	
 };
diff --git a/Engine/Plugins/Media/PixelStreaming/Source/PixelStreaming/Private/PixelStreamingInputHandler.cpp b/Engine/Plugins/Media/PixelStreaming/Source/PixelStreaming/Private/PixelStreamingInputHandler.cpp
--- a/Engine/Plugins/Media/PixelStreaming/Source/PixelStreaming/Private/PixelStreamingInputHandler.cpp
+++ b/Engine/Plugins/Media/PixelStreaming/Source/PixelStreaming/Private/PixelStreamingInputHandler.cpp
@@ -18,6 +18,7 @@
 #include "Framework/Application/SlateUser.h"
 #include "Widgets/Input/SEditableTextBox.h"
 #include "Misc/CoreMiscDefines.h"
+#include "Input/HittestGrid.h"
 
 DECLARE_LOG_CATEGORY_EXTERN(LogPixelStreamingInputHandler, Log, VeryVerbose);
 DEFINE_LOG_CATEGORY(LogPixelStreamingInputHandler);
@@ -230,30 +231,47 @@ namespace UE::PixelStreaming
 			// If Touch is valid
 			if (Touch.Param5 != 0)
 			{
-				if (NumActiveTouches == 0 && !bIsMouseActive)
-				{
-					FSlateApplication::Get().OnCursorSet();
-					// Make sure the application is active.
-					FSlateApplication::Get().ProcessApplicationActivationEvent(true);
+				//                                                                           convert range from 0,65536 -> 0,1
+				FVector2D TouchLocation = ConvertFromNormalizedScreenLocation(FVector2D(Touch.Param1 / uint16_MAX, Touch.Param2 / uint16_MAX));
+				const int32 TouchIndex = Touch.Param3;
+				const float TouchForce = Touch.Param4 / 255.0f;
+				UE_LOG(LogPixelStreamingInputHandler, Verbose, TEXT("TOUCH_START: TouchIndex = %d; Pos = (%d, %d); CursorPos = (%d, %d); Force = %.3f"), TouchIndex, Touch.Param1, Touch.Param2, static_cast<int>(TouchLocation.X), static_cast<int>(TouchLocation.Y), TouchForce);
 
-					FVector2D OldCursorLocation = PixelStreamerApplicationWrapper->WrappedApplication->Cursor->GetPosition();
-					PixelStreamerApplicationWrapper->Cursor->SetPosition(OldCursorLocation.X, OldCursorLocation.Y);
-					FSlateApplication::Get().OverridePlatformApplication(PixelStreamerApplicationWrapper);
+				if(InputType == EPixelStreamingInputType::RouteToWidget)
+				{
+					// TouchLocation = TouchLocation - TargetViewport.Pin()->GetCachedGeometry().GetAbsolutePosition();
+					FWidgetPath WidgetPath = FindRoutingMessageWidget(TouchLocation);
+					
+					if (WidgetPath.IsValid())
+					{
+						FScopedSwitchWorldHack SwitchWorld(WidgetPath);
+						FPointerEvent PointerEvent(0, TouchIndex, TouchLocation, TouchLocation, TouchForce, true);
+						FSlateApplication::Get().RoutePointerDownEvent(WidgetPath, PointerEvent);
+					}
 				}
+				else if(InputType == EPixelStreamingInputType::RouteToWindow)
+				{
+					if (NumActiveTouches == 0 && !bIsMouseActive)
+					{
+						FSlateApplication::Get().OnCursorSet();
+						// Make sure the application is active.
+						FSlateApplication::Get().ProcessApplicationActivationEvent(true);
 
-				//                                                                           convert range from 0,65536 -> 0,1
-				FVector2D TouchLocation = ConvertFromNormalizedScreenLocation(FVector2D(Touch.Param1 / uint16_MAX, Touch.Param2 / uint16_MAX));
+						FVector2D OldCursorLocation = PixelStreamerApplicationWrapper->WrappedApplication->Cursor->GetPosition();
+						PixelStreamerApplicationWrapper->Cursor->SetPosition(OldCursorLocation.X, OldCursorLocation.Y);
+						FSlateApplication::Get().OverridePlatformApplication(PixelStreamerApplicationWrapper);
+					}
 
-				// We must update the user cursor position explicitly before updating the application cursor position
-				// as if there's a delta between them, when the touch event is started it will trigger a move
-				// resulting in a large 'drag' across the screen
-				TSharedPtr<FSlateUser> User = FSlateApplication::Get().GetCursorUser();
-				User->SetCursorPosition(TouchLocation);
-				PixelStreamerApplicationWrapper->Cursor->SetPosition(TouchLocation.X, TouchLocation.Y);
-				PixelStreamerApplicationWrapper->WrappedApplication->Cursor->SetPosition(TouchLocation.X, TouchLocation.Y);
+					// We must update the user cursor position explicitly before updating the application cursor position
+					// as if there's a delta between them, when the touch event is started it will trigger a move
+					// resulting in a large 'drag' across the screen
+					TSharedPtr<FSlateUser> User = FSlateApplication::Get().GetCursorUser();
+					User->SetCursorPosition(TouchLocation);
+					PixelStreamerApplicationWrapper->Cursor->SetPosition(TouchLocation.X, TouchLocation.Y);
+					PixelStreamerApplicationWrapper->WrappedApplication->Cursor->SetPosition(TouchLocation.X, TouchLocation.Y);
 
-				UE_LOG(LogPixelStreamingInputHandler, Verbose, TEXT("TOUCH_START: TouchIndex = %d; Pos = (%d, %d); CursorPos = (%d, %d); Force = %.3f"), Touch.Param3, Touch.Param1, Touch.Param2, static_cast<int>(TouchLocation.X), static_cast<int>(TouchLocation.Y), Touch.Param4 / 255.0f);
-				MessageHandler->OnTouchStarted(PixelStreamerApplicationWrapper->GetWindowUnderCursor(), TouchLocation, Touch.Param4 / 255.0f, Touch.Param3, 0); // TODO: ControllerId?
+					MessageHandler->OnTouchStarted(PixelStreamerApplicationWrapper->GetWindowUnderCursor(), TouchLocation, TouchForce, TouchIndex, 0); // TODO: ControllerId?
+				}
 
 				NumActiveTouches++;
 			}
@@ -277,13 +295,33 @@ namespace UE::PixelStreaming
 				//                                                                           convert range from 0,65536 -> 0,1
 				FVector2D TouchLocation = ConvertFromNormalizedScreenLocation(FVector2D(Touch.Param1 / uint16_MAX, Touch.Param2 / uint16_MAX));
 				const int32 TouchIndex = Touch.Param3;
-                UE_LOG(LogPixelStreamingInputHandler, Verbose, TEXT("TOUCH_MOVE: TouchIndex = %d; Pos = (%d, %d); CursorPos = (%d, %d); Force = %.3f"), Touch.Param3, Touch.Param1, Touch.Param2, static_cast<int>(TouchLocation.X), static_cast<int>(TouchLocation.Y),  Touch.Param4 / 255.0f);
+				const float TouchForce = Touch.Param4 / 255.0f;
+                UE_LOG(LogPixelStreamingInputHandler, Verbose, TEXT("TOUCH_MOVE: TouchIndex = %d; Pos = (%d, %d); CursorPos = (%d, %d); Force = %.3f"), TouchIndex, Touch.Param1, Touch.Param2, static_cast<int>(TouchLocation.X), static_cast<int>(TouchLocation.Y), TouchForce);
 
 				FCachedTouchEvent& TouchEvent = CachedTouchEvents.FindOrAdd(TouchIndex);
-				TouchEvent.Location = TouchLocation;
-				TouchEvent.Force = Touch.Param4 / 255.0f;
+				TouchEvent.Force = TouchForce;
 				TouchEvent.ControllerIndex = 0;
-				MessageHandler->OnTouchMoved(TouchEvent.Location, TouchEvent.Force, TouchIndex, TouchEvent.ControllerIndex); // TODO: ControllerId?
+					
+				if(InputType == EPixelStreamingInputType::RouteToWidget)
+				{
+					// TouchLocation = TouchLocation - TargetViewport.Pin()->GetCachedGeometry().GetAbsolutePosition();
+					TouchEvent.Location = TouchLocation;
+					FWidgetPath WidgetPath = FindRoutingMessageWidget(TouchLocation);
+					
+					if (WidgetPath.IsValid())
+					{
+						FScopedSwitchWorldHack SwitchWorld(WidgetPath);
+						FPointerEvent PointerEvent(0, TouchIndex, TouchLocation, LastTouchLocation, TouchForce, true);
+						FSlateApplication::Get().RoutePointerMoveEvent(WidgetPath, PointerEvent, false);
+					}
+
+					LastTouchLocation = TouchLocation;
+				}
+				else if(InputType == EPixelStreamingInputType::RouteToWindow)
+				{
+					TouchEvent.Location = TouchLocation;
+					MessageHandler->OnTouchMoved(TouchEvent.Location, TouchEvent.Force, TouchIndex, TouchEvent.ControllerIndex); // TODO: ControllerId?	
+				}
 				TouchIndicesProcessedThisFrame.Add(TouchIndex);
             }
         }
@@ -305,7 +343,25 @@ namespace UE::PixelStreaming
 				FVector2D TouchLocation = ConvertFromNormalizedScreenLocation(FVector2D(Touch.Param1 / uint16_MAX, Touch.Param2 / uint16_MAX));
 				const int32 TouchIndex = Touch.Param3;
 				UE_LOG(LogPixelStreamingInputHandler, Verbose, TEXT("TOUCH_END: TouchIndex = %d; Pos = (%d, %d); CursorPos = (%d, %d)"), Touch.Param3, Touch.Param1, Touch.Param2, static_cast<int>(TouchLocation.X), static_cast<int>(TouchLocation.Y));
-				MessageHandler->OnTouchEnded(TouchLocation, TouchIndex, 0); // TODO: ControllerId?
+
+				if(InputType == EPixelStreamingInputType::RouteToWidget)
+				{
+					// TouchLocation = TouchLocation - TargetViewport.Pin()->GetCachedGeometry().GetAbsolutePosition();
+					FWidgetPath WidgetPath = FindRoutingMessageWidget(TouchLocation);
+					
+					if (WidgetPath.IsValid())
+					{
+						FScopedSwitchWorldHack SwitchWorld(WidgetPath);
+						float TouchForce = 0.0f;
+						FPointerEvent PointerEvent(0, TouchIndex, TouchLocation, TouchLocation, TouchForce, true);
+						FSlateApplication::Get().RoutePointerUpEvent(WidgetPath, PointerEvent);
+					}
+				}
+				else if(InputType == EPixelStreamingInputType::RouteToWindow)
+				{
+					MessageHandler->OnTouchEnded(TouchLocation, TouchIndex, 0); // TODO: ControllerId?
+				}
+
             	CachedTouchEvents.Remove(TouchIndex);
 				NumActiveTouches--;
 			}
@@ -315,7 +371,7 @@ namespace UE::PixelStreaming
 		// then set the platform application back to its default. We need to set it back to default
 		// so that people using the editor (if editor streaming) can click on buttons outside the target window
 		// and also have the correct cursor (pixel streaming forces default cursor)
-		if (NumActiveTouches == 0 && !bIsMouseActive)
+		if (NumActiveTouches == 0 && !bIsMouseActive && InputType == EPixelStreamingInputType::RouteToWindow)
 		{
 			FVector2D OldCursorLocation = PixelStreamerApplicationWrapper->Cursor->GetPosition();
 			PixelStreamerApplicationWrapper->WrappedApplication->Cursor->SetPosition(OldCursorLocation.X, OldCursorLocation.Y);
@@ -402,10 +458,39 @@ namespace UE::PixelStreaming
 
         EMouseButtons::Type Button = static_cast<EMouseButtons::Type>(Payload.Param1);
         UE_LOG(LogPixelStreamingInputHandler, Verbose, TEXT("MOUSE_UP: Button = %d"), Button);
-        if(Button != EMouseButtons::Type::Invalid)
-        {
-            MessageHandler->OnMouseUp(Button);
-        }
+
+		if(InputType == EPixelStreamingInputType::RouteToWidget)
+		{
+			FSlateApplication& SlateApplication = FSlateApplication::Get();
+			FWidgetPath WidgetPath = FindRoutingMessageWidget(SlateApplication.GetCursorPos());
+			
+			if (WidgetPath.IsValid())
+			{
+				FScopedSwitchWorldHack SwitchWorld(WidgetPath);
+				FKey Key = TranslateMouseButtonToKey(Button);
+
+				FPointerEvent MouseEvent(
+					SlateApplication.GetUserIndexForMouse(),
+					FSlateApplicationBase::CursorPointerIndex,
+					SlateApplication.GetCursorPos(),
+					SlateApplication.GetLastCursorPos(),
+					SlateApplication.GetPressedMouseButtons(),
+					Key,
+					0,
+					SlateApplication.GetPlatformApplication()->GetModifierKeys()
+				);
+
+				SlateApplication.RoutePointerUpEvent(WidgetPath, MouseEvent);
+			}
+		}
+		else if(InputType == EPixelStreamingInputType::RouteToWindow)
+		{
+			if(Button != EMouseButtons::Type::Invalid)
+			{
+				MessageHandler->OnMouseUp(Button);
+			}
+		}
+        
     }
 
     void FPixelStreamingInputHandler::HandleOnMouseDown(FMemoryReader Ar)
@@ -420,15 +505,41 @@ namespace UE::PixelStreaming
         //                                                                           convert range from 0,65536 -> 0,1
         FVector2D ScreenLocation = ConvertFromNormalizedScreenLocation(FVector2D(Payload.Param2 / uint16_MAX, Payload.Param3 / uint16_MAX));
         EMouseButtons::Type Button = static_cast<EMouseButtons::Type>(Payload.Param1);
-
-		// Set cursor pos on mouse down - we may not have moved if this is the very first click
-		FSlateApplication::Get().OnCursorSet();
-		PixelStreamerApplicationWrapper->Cursor->SetPosition(ScreenLocation.X, ScreenLocation.Y);
-
 		UE_LOG(LogPixelStreamingInputHandler, Verbose, TEXT("MOUSE_DOWN: Button = %d; Pos = (%.4f, %.4f)"), Button, ScreenLocation.X, ScreenLocation.Y);
+		// Set cursor pos on mouse down - we may not have moved if this is the very first click
+		FSlateApplication& SlateApplication = FSlateApplication::Get();
+		SlateApplication.OnCursorSet();
 		// Force window focus
-		FSlateApplication::Get().ProcessApplicationActivationEvent(true);
-		MessageHandler->OnMouseDown(PixelStreamerApplicationWrapper->GetWindowUnderCursor(), Button, ScreenLocation);
+		SlateApplication.ProcessApplicationActivationEvent(true);
+
+		if(InputType == EPixelStreamingInputType::RouteToWidget)
+		{
+			FWidgetPath WidgetPath = FindRoutingMessageWidget(ScreenLocation);
+			
+			if (WidgetPath.IsValid())
+			{
+				FScopedSwitchWorldHack SwitchWorld(WidgetPath);
+
+				FKey Key = TranslateMouseButtonToKey(Button);
+
+				FPointerEvent MouseEvent(
+					SlateApplication.GetUserIndexForMouse(),
+					FSlateApplicationBase::CursorPointerIndex,
+					ScreenLocation,
+					SlateApplication.GetLastCursorPos(),
+					SlateApplication.GetPressedMouseButtons(),
+					Key,
+					0,
+					SlateApplication.GetPlatformApplication()->GetModifierKeys()
+				);
+
+				SlateApplication.RoutePointerDownEvent(WidgetPath, MouseEvent);
+			}
+		}
+		else if(InputType == EPixelStreamingInputType::RouteToWindow)
+		{
+			MessageHandler->OnMouseDown(PixelStreamerApplicationWrapper->GetWindowUnderCursor(), Button, ScreenLocation);
+		}		
 	}
 
     void FPixelStreamingInputHandler::HandleOnMouseMove(FMemoryReader Ar)
@@ -438,11 +549,36 @@ namespace UE::PixelStreaming
         FIntPoint ScreenLocation = ConvertFromNormalizedScreenLocation(FVector2D(Payload.Param1 / uint16_MAX, Payload.Param2 / uint16_MAX));
         //                                                                 convert range from -32,768 to 32,767 -> -1,1
         FIntPoint Delta = ConvertFromNormalizedScreenLocation(FVector2D(Payload.Param3 / int16_MAX, Payload.Param4 / int16_MAX), false);
-
-		FSlateApplication::Get().OnCursorSet();
 		UE_LOG(LogPixelStreamingInputHandler, Verbose, TEXT("MOUSE_MOVE: Pos = (%d, %d); Delta = (%d, %d)"), ScreenLocation.X, ScreenLocation.Y, Delta.X, Delta.Y);
+		FSlateApplication& SlateApplication = FSlateApplication::Get();
+		SlateApplication.OnCursorSet();
 		PixelStreamerApplicationWrapper->Cursor->SetPosition(ScreenLocation.X, ScreenLocation.Y);
-		MessageHandler->OnRawMouseMove(Delta.X, Delta.Y);
+
+		if(InputType == EPixelStreamingInputType::RouteToWidget)
+		{
+			FWidgetPath WidgetPath = FindRoutingMessageWidget(ScreenLocation);
+			
+			if (WidgetPath.IsValid())
+			{
+				FScopedSwitchWorldHack SwitchWorld(WidgetPath);
+
+				FPointerEvent MouseEvent(
+					SlateApplication.GetUserIndexForMouse(),
+					FSlateApplicationBase::CursorPointerIndex,
+					SlateApplication.GetCursorPos(),
+					SlateApplication.GetLastCursorPos(),
+					FVector2D(Delta.X, Delta.Y), 
+					SlateApplication.GetPressedMouseButtons(),
+					SlateApplication.GetPlatformApplication()->GetModifierKeys()
+				);
+
+				SlateApplication.RoutePointerMoveEvent(WidgetPath, MouseEvent, false);
+			}
+		}
+		else if(InputType == EPixelStreamingInputType::RouteToWindow)
+		{
+			MessageHandler->OnRawMouseMove(Delta.X, Delta.Y);
+		}
 	}
 
     void FPixelStreamingInputHandler::HandleOnMouseWheel(FMemoryReader Ar)
@@ -451,8 +587,36 @@ namespace UE::PixelStreaming
          //                                                                           convert range from 0,65536 -> 0,1
         FIntPoint ScreenLocation = ConvertFromNormalizedScreenLocation(FVector2D(Payload.Param2 / uint16_MAX, Payload.Param3 / uint16_MAX));
         const float SpinFactor = 1 / 120.0f;
-        MessageHandler->OnMouseWheel(Payload.Param1 * SpinFactor, ScreenLocation);
-        UE_LOG(LogPixelStreamingInputHandler, Verbose, TEXT("MOUSE_WHEEL: Delta = %d; Pos = (%d, %d)"), Payload.Param1, ScreenLocation.X, ScreenLocation.Y);
+		UE_LOG(LogPixelStreamingInputHandler, Verbose, TEXT("MOUSE_WHEEL: Delta = %d; Pos = (%d, %d)"), Payload.Param1, ScreenLocation.X, ScreenLocation.Y);
+
+		if(InputType == EPixelStreamingInputType::RouteToWidget)
+		{
+			FWidgetPath WidgetPath = FindRoutingMessageWidget(ScreenLocation);
+			
+			if (WidgetPath.IsValid())
+			{
+				FScopedSwitchWorldHack SwitchWorld(WidgetPath);
+
+				FSlateApplication& SlateApplication = FSlateApplication::Get();
+
+				FPointerEvent MouseEvent(
+					SlateApplication.GetUserIndexForMouse(),
+					FSlateApplicationBase::CursorPointerIndex,
+					SlateApplication.GetCursorPos(),
+					SlateApplication.GetCursorPos(),
+					SlateApplication.GetPressedMouseButtons(),
+					EKeys::Invalid,
+					Payload.Param1 * SpinFactor,
+					SlateApplication.GetPlatformApplication()->GetModifierKeys()
+				);
+
+				SlateApplication.RouteMouseWheelOrGestureEvent(WidgetPath, MouseEvent, nullptr);
+			}
+		}
+		else if(InputType == EPixelStreamingInputType::RouteToWindow)
+		{
+			MessageHandler->OnMouseWheel(Payload.Param1 * SpinFactor, ScreenLocation);
+		}
     }
 
 	void FPixelStreamingInputHandler::HandleOnMouseDoubleClick(FMemoryReader Ar)
@@ -461,11 +625,38 @@ namespace UE::PixelStreaming
         //                                                                           convert range from 0,65536 -> 0,1
         FVector2D ScreenLocation = ConvertFromNormalizedScreenLocation(FVector2D(Payload.Param2 / uint16_MAX, Payload.Param3 / uint16_MAX));
         EMouseButtons::Type Button = static_cast<EMouseButtons::Type>(Payload.Param1);
-
 		UE_LOG(LogPixelStreamingInputHandler, Verbose, TEXT("MOUSE_DOWN: Button = %d; Pos = (%.4f, %.4f)"), Button, ScreenLocation.X, ScreenLocation.Y);
 		// Force window focus
-		FSlateApplication::Get().ProcessApplicationActivationEvent(true);
-		MessageHandler->OnMouseDoubleClick(PixelStreamerApplicationWrapper->GetWindowUnderCursor(), Button, ScreenLocation);
+		FSlateApplication& SlateApplication = FSlateApplication::Get();
+		SlateApplication.ProcessApplicationActivationEvent(true);
+
+		if(InputType == EPixelStreamingInputType::RouteToWidget)
+		{
+			FWidgetPath WidgetPath = FindRoutingMessageWidget(ScreenLocation);
+			
+			if (WidgetPath.IsValid())
+			{
+				FScopedSwitchWorldHack SwitchWorld(WidgetPath);
+				FKey Key = TranslateMouseButtonToKey( Button );
+
+				FPointerEvent MouseEvent(
+					SlateApplication.GetUserIndexForMouse(),
+					FSlateApplicationBase::CursorPointerIndex,
+					SlateApplication.GetCursorPos(),
+					SlateApplication.GetLastCursorPos(),
+					SlateApplication.GetPressedMouseButtons(),
+					Key,
+					0,
+					SlateApplication.GetPlatformApplication()->GetModifierKeys()
+				);
+
+				SlateApplication.RoutePointerDoubleClickEvent(WidgetPath, MouseEvent);
+			}
+		}
+		else if(InputType == EPixelStreamingInputType::RouteToWindow)
+		{
+			MessageHandler->OnMouseDoubleClick(PixelStreamerApplicationWrapper->GetWindowUnderCursor(), Button, ScreenLocation);
+		}
 	}
 
     /**
@@ -598,34 +789,28 @@ namespace UE::PixelStreaming
 	{
 		FIntPoint OutVector((int32)ScreenLocation.X, (int32)ScreenLocation.Y);
 
-		TSharedPtr<SWindow> ApplicationWindow = TargetWindow.Pin();
-		if (ApplicationWindow.IsValid())
+		if (TSharedPtr<SWindow> ApplicationWindow = TargetWindow.Pin())
 		{
 			FVector2D WindowOrigin = ApplicationWindow->GetPositionInScreen();
-			if (TargetViewport.IsValid())
+			if (TSharedPtr<SViewport> ViewportWidget = TargetViewport.Pin())
 			{
-				TSharedPtr<SViewport> ViewportWidget = TargetViewport.Pin();
-
-				if (ViewportWidget.IsValid())
-				{
-					FGeometry InnerWindowGeometry = ApplicationWindow->GetWindowGeometryInWindow();
+				FGeometry InnerWindowGeometry = ApplicationWindow->GetWindowGeometryInWindow();
 
-					// Find the widget path relative to the window
-					FArrangedChildren JustWindow(EVisibility::Visible);
-					JustWindow.AddWidget(FArrangedWidget(ApplicationWindow.ToSharedRef(), InnerWindowGeometry));
+				// Find the widget path relative to the window
+				FArrangedChildren JustWindow(EVisibility::Visible);
+				JustWindow.AddWidget(FArrangedWidget(ApplicationWindow.ToSharedRef(), InnerWindowGeometry));
 
-					FWidgetPath PathToWidget(ApplicationWindow.ToSharedRef(), JustWindow);
-					if (PathToWidget.ExtendPathTo(FWidgetMatcher(ViewportWidget.ToSharedRef()), EVisibility::Visible))
-					{
-						FArrangedWidget ArrangedWidget = PathToWidget.FindArrangedWidget(ViewportWidget.ToSharedRef()).Get(FArrangedWidget::GetNullWidget());
+				FWidgetPath PathToWidget(ApplicationWindow.ToSharedRef(), JustWindow);
+				if (PathToWidget.ExtendPathTo(FWidgetMatcher(ViewportWidget.ToSharedRef()), EVisibility::Visible))
+				{
+					FArrangedWidget ArrangedWidget = PathToWidget.FindArrangedWidget(ViewportWidget.ToSharedRef()).Get(FArrangedWidget::GetNullWidget());
 
-						FVector2D WindowClientOffset = ArrangedWidget.Geometry.GetAbsolutePosition();
-						FVector2D WindowClientSize = ArrangedWidget.Geometry.GetAbsoluteSize();
+					FVector2D WindowClientOffset = ArrangedWidget.Geometry.GetAbsolutePosition();
+					FVector2D WindowClientSize = ArrangedWidget.Geometry.GetAbsoluteSize();
 
-						FVector2D OutTemp = bIncludeOffset ? WindowOrigin + WindowClientOffset + (ScreenLocation * WindowClientSize) : (ScreenLocation * WindowClientSize);
-						UE_LOG(LogPixelStreamingInputHandler, Verbose, TEXT("%.4f, %.4f"), ScreenLocation.X, ScreenLocation.Y);
-						OutVector = FIntPoint((int32)OutTemp.X, (int32)OutTemp.Y);
-					}
+					FVector2D OutTemp = bIncludeOffset ? WindowOrigin + WindowClientOffset + (ScreenLocation * WindowClientSize) : (ScreenLocation * WindowClientSize);
+					UE_LOG(LogPixelStreamingInputHandler, Verbose, TEXT("%.4f, %.4f"), ScreenLocation.X, ScreenLocation.Y);
+					OutVector = FIntPoint((int32)OutTemp.X, (int32)OutTemp.Y);
 				}
 			}
 			else
@@ -635,9 +820,9 @@ namespace UE::PixelStreaming
 				OutVector = FIntPoint((int32)OutTemp.X, (int32)OutTemp.Y);
 			}
 		}
-		else if(TargetScreenSize.IsValid())
+		else if(TSharedPtr<FIntPoint> ScreenSize = TargetScreenSize.Pin())
 		{
-			FIntPoint SizeInScreen = *(TargetScreenSize.Pin());
+			FIntPoint SizeInScreen = *ScreenSize;
 			FVector2D OutTemp = FVector2D(SizeInScreen) * ScreenLocation;
 			OutVector = FIntPoint((int32)OutTemp.X, (int32)OutTemp.Y);
 		}
@@ -801,11 +986,51 @@ namespace UE::PixelStreaming
     		// Only broadcast events that haven't already been fired this frame
     		if (!TouchIndicesProcessedThisFrame.Contains(TouchIndex))
     		{
-    			MessageHandler->OnTouchMoved(TouchEvent.Location, TouchEvent.Force, TouchIndex, TouchEvent.ControllerIndex);
+				if(InputType == EPixelStreamingInputType::RouteToWidget)
+				{
+					FWidgetPath WidgetPath = FindRoutingMessageWidget(TouchEvent.Location);
+					
+					if (WidgetPath.IsValid())
+					{
+						FScopedSwitchWorldHack SwitchWorld(WidgetPath);
+						FPointerEvent PointerEvent(0, TouchIndex, TouchEvent.Location, LastTouchLocation, TouchEvent.Force, true);
+						FSlateApplication::Get().RoutePointerMoveEvent(WidgetPath, PointerEvent, false);
+					}
+				}
+				else if(InputType == EPixelStreamingInputType::RouteToWindow)
+				{
+					MessageHandler->OnTouchMoved(TouchEvent.Location, TouchEvent.Force, TouchIndex, TouchEvent.ControllerIndex);
+				}
     		}
     	}
 	}
 
+	FKey FPixelStreamingInputHandler::TranslateMouseButtonToKey(const EMouseButtons::Type Button)
+	{
+		FKey Key = EKeys::Invalid;
+
+		switch( Button )
+		{
+		case EMouseButtons::Left:
+			Key = EKeys::LeftMouseButton;
+			break;
+		case EMouseButtons::Middle:
+			Key = EKeys::MiddleMouseButton;
+			break;
+		case EMouseButtons::Right:
+			Key = EKeys::RightMouseButton;
+			break;
+		case EMouseButtons::Thumb01:
+			Key = EKeys::ThumbMouseButton;
+			break;
+		case EMouseButtons::Thumb02:
+			Key = EKeys::ThumbMouseButton2;
+			break;
+		}
+
+		return Key;
+	}
+
 	void FPixelStreamingInputHandler::FindFocusedWidget()
 	{
 		FSlateApplication::Get().ForEachUser([this](FSlateUser& User) {
@@ -881,4 +1106,18 @@ namespace UE::PixelStreaming
 			}
 		});
     }
+
+	FWidgetPath FPixelStreamingInputHandler::FindRoutingMessageWidget(const FVector2D& Location) const
+	{
+		if (TSharedPtr<SWindow> PlaybackWindowPinned = TargetWindow.Pin())
+		{
+			if (PlaybackWindowPinned->AcceptsInput())
+			{
+				bool bIgnoreEnabledStatus = false;
+				TArray<FWidgetAndPointer> WidgetsAndCursors = PlaybackWindowPinned->GetHittestGrid().GetBubblePath(Location, FSlateApplication::Get().GetCursorRadius(), bIgnoreEnabledStatus);
+				return FWidgetPath(MoveTemp(WidgetsAndCursors));
+			}
+		}
+		return FWidgetPath();
+	}
 } // namespace UE::PixelStreaming
diff --git a/Engine/Plugins/Media/PixelStreaming/Source/PixelStreaming/Private/PixelStreamingInputHandler.h b/Engine/Plugins/Media/PixelStreaming/Source/PixelStreaming/Private/PixelStreamingInputHandler.h
--- a/Engine/Plugins/Media/PixelStreaming/Source/PixelStreaming/Private/PixelStreamingInputHandler.h
+++ b/Engine/Plugins/Media/PixelStreaming/Source/PixelStreaming/Private/PixelStreamingInputHandler.h
@@ -37,21 +37,23 @@ namespace UE::PixelStreaming
 
 		virtual void SetChannelValues(int32 ControllerId, const FForceFeedbackValues& values) override;
 
-		void OnMessage(const webrtc::DataBuffer& Buffer);
+		virtual void OnMessage(const webrtc::DataBuffer& Buffer) override;
 
-		void SetTargetWindow(TWeakPtr<SWindow> InWindow);
-		TWeakPtr<SWindow> GetTargetWindow();
+		virtual void SetTargetWindow(TWeakPtr<SWindow> InWindow) override;
+		virtual TWeakPtr<SWindow> GetTargetWindow() override;
 
-		void SetTargetViewport(TWeakPtr<SViewport> InViewport);
-		TWeakPtr<SViewport> GetTargetViewport();
+		virtual void SetTargetViewport(TWeakPtr<SViewport> InViewport) override;
+		virtual TWeakPtr<SViewport> GetTargetViewport() override;
 
-		void SetTargetScreenSize(TWeakPtr<FIntPoint> InScreenSize);
-		TWeakPtr<FIntPoint> GetTargetScreenSize();
+		virtual void SetTargetScreenSize(TWeakPtr<FIntPoint> InScreenSize) override;
+		virtual TWeakPtr<FIntPoint> GetTargetScreenSize() override;
 
-		bool IsFakingTouchEvents() const { return bFakingTouchEvents; }
+		virtual bool IsFakingTouchEvents() const override { return bFakingTouchEvents; }
 
-        void RegisterMessageHandler(const FString& MessageType, const TFunction<void(FMemoryReader)>& Handler);
-		TFunction<void(FMemoryReader)> FindMessageHandler(const FString& MessageType);
+        virtual void RegisterMessageHandler(const FString& MessageType, const TFunction<void(FMemoryReader)>& Handler) override;
+		virtual TFunction<void(FMemoryReader)> FindMessageHandler(const FString& MessageType) override;
+
+		virtual void SetInputType(EPixelStreamingInputType InInputType) override { InputType = InInputType; };
 
     protected:
         /**
@@ -92,9 +94,11 @@ namespace UE::PixelStreaming
         virtual void HandleUIInteraction(FMemoryReader Ar);
 
 		FIntPoint ConvertFromNormalizedScreenLocation(const FVector2D& ScreenLocation, bool bIncludeOffset = true);
+		FWidgetPath FindRoutingMessageWidget(const FVector2D& Location) const;
 
 		FGamepadKeyNames::Type ConvertAxisIndexToGamepadAxis(uint8 AnalogAxis);
 		FGamepadKeyNames::Type ConvertButtonIndexToGamepadButton(uint8 ButtonIndex);
+		FKey TranslateMouseButtonToKey(const EMouseButtons::Type Button);
 
 		struct FCachedTouchEvent
 		{
@@ -121,13 +125,15 @@ namespace UE::PixelStreaming
             TArray<uint8> Data;
         };
 
-		TWeakPtr<SWindow> TargetWindow;
-		TWeakPtr<SViewport> TargetViewport;
-		TWeakPtr<FIntPoint> TargetScreenSize; // Manual size override used when we don't have a single window/viewport target
-		uint8 NumActiveTouches;
-		bool bIsMouseActive;
-		TMap<uint8,  TFunction<void(FMemoryReader)>> DispatchTable;
-		TQueue<FMessage> Messages;
+		TWeakPtr<SWindow> 			TargetWindow;
+		TWeakPtr<SViewport> 		TargetViewport;
+		TWeakPtr<FIntPoint> 		TargetScreenSize; // Manual size override used when we don't have a single window/viewport target
+		uint8 						NumActiveTouches;
+		bool 						bIsMouseActive;
+		TQueue<FMessage> 			Messages;
+		EPixelStreamingInputType	InputType = EPixelStreamingInputType::RouteToWindow;
+		FVector2D					LastTouchLocation = FVector2D(EForceInit::ForceInitToZero);
+		TMap<uint8, TFunction<void(FMemoryReader)>> DispatchTable;
 
 		/** Reference to the message handler which events should be passed to. */
 		TSharedPtr<FGenericApplicationMessageHandler> MessageHandler;
diff --git a/Engine/Plugins/Media/PixelStreaming/Source/PixelStreaming/Private/Streamer.cpp b/Engine/Plugins/Media/PixelStreaming/Source/PixelStreaming/Private/Streamer.cpp
--- a/Engine/Plugins/Media/PixelStreaming/Source/PixelStreaming/Private/Streamer.cpp
+++ b/Engine/Plugins/Media/PixelStreaming/Source/PixelStreaming/Private/Streamer.cpp
@@ -314,6 +314,11 @@ namespace UE::PixelStreaming
 		// TODO Delete player session?
 	}
 
+	void FStreamer::SetInputHandlerType(EPixelStreamingInputType InputType)
+	{
+		InputHandler->SetInputType(InputType);
+	}
+
 	IPixelStreamingAudioSink* FStreamer::GetPeerAudioSink(FPixelStreamingPlayerId PlayerId)
 	{
 		if (FPlayerContext* PlayerContext = Players.Find(PlayerId))
diff --git a/Engine/Plugins/Media/PixelStreaming/Source/PixelStreaming/Private/Streamer.h b/Engine/Plugins/Media/PixelStreaming/Source/PixelStreaming/Private/Streamer.h
--- a/Engine/Plugins/Media/PixelStreaming/Source/PixelStreaming/Private/Streamer.h
+++ b/Engine/Plugins/Media/PixelStreaming/Source/PixelStreaming/Private/Streamer.h
@@ -58,6 +58,7 @@ namespace UE::PixelStreaming
 
 		virtual void SetInputHandler(TSharedPtr<IPixelStreamingInputHandler> InInputHandler) override { InputHandler = InInputHandler; }
 		virtual TWeakPtr<IPixelStreamingInputHandler> GetInputHandler() override { return InputHandler; }
+		virtual void SetInputHandlerType(EPixelStreamingInputType InputType) override;
 
 		virtual IPixelStreamingAudioSink* GetPeerAudioSink(FPixelStreamingPlayerId PlayerId) override;
 		virtual IPixelStreamingAudioSink* GetUnlistenedAudioSink() override;
diff --git a/Engine/Plugins/Media/PixelStreaming/Source/PixelStreaming/Public/IPixelStreamingInputHandler.h b/Engine/Plugins/Media/PixelStreaming/Source/PixelStreaming/Public/IPixelStreamingInputHandler.h
--- a/Engine/Plugins/Media/PixelStreaming/Source/PixelStreaming/Public/IPixelStreamingInputHandler.h
+++ b/Engine/Plugins/Media/PixelStreaming/Source/PixelStreaming/Public/IPixelStreamingInputHandler.h
@@ -10,6 +10,12 @@
 #include "PixelStreamingWebRTCIncludes.h"
 #include "Serialization/MemoryReader.h"
 
+enum class PIXELSTREAMING_API EPixelStreamingInputType : uint8
+{
+	RouteToWindow = 0,
+	RouteToWidget = 1
+};
+
 /**
  * IPixelStreamingInputHandler extends the IInputDevice interface. Setting the target viewport allows for
  * scaling of input from browser to application, and setting the target window ensure that if windows are tiled (eg editor)
@@ -76,7 +82,6 @@ public:
 	 */
 	virtual void RegisterMessageHandler(const FString& MessageType, const TFunction<void(FMemoryReader)>& Handler) = 0;
 
-
 	/**
 	 * @brief Find the function to be called whenever the specified message type is received.
 	 * 
@@ -84,4 +89,9 @@ public:
 	 * @return TFunction<void(FMemoryReader)> The function called when this message type is received.
 	 */
 	virtual TFunction<void(FMemoryReader)> FindMessageHandler(const FString& MessageType) = 0;
+
+	/**
+	 * @brief Set the input handlers type. This controls whether input is routed to widgets or windows
+	 */
+	virtual void SetInputType(EPixelStreamingInputType InputType) = 0;
 };
\ No newline at end of file
diff --git a/Engine/Plugins/Media/PixelStreaming/Source/PixelStreaming/Public/IPixelStreamingStreamer.h b/Engine/Plugins/Media/PixelStreaming/Source/PixelStreaming/Public/IPixelStreamingStreamer.h
--- a/Engine/Plugins/Media/PixelStreaming/Source/PixelStreaming/Public/IPixelStreamingStreamer.h
+++ b/Engine/Plugins/Media/PixelStreaming/Source/PixelStreaming/Public/IPixelStreamingStreamer.h
@@ -209,6 +209,11 @@ public:
 	 */
 	virtual TWeakPtr<IPixelStreamingInputHandler> GetInputHandler() = 0;
 
+	/**
+	 * @brief Set the type for this streamers input handler. This controls whether input is routed to widgets or windows
+	 */
+	virtual void SetInputHandlerType(EPixelStreamingInputType InputType) = 0;
+
 	/**
 	 * @brief Get the audio sink associated with a specific peer/player.
 	 */
@@ -220,14 +225,14 @@ public:
 	virtual IPixelStreamingAudioSink* GetUnlistenedAudioSink() = 0;
 
 	/**
-	 * Create an audio input that can be used to push audio into the Pixel Streaming mix.
+	 * @brief Create an audio input that can be used to push audio into the Pixel Streaming mix.
 	 * Note: Requires launching with -PixelStreamingExperimentalAudioInput.
 	 * @return The audio input to push to.
 	 */
 	virtual TSharedPtr<IPixelStreamingAudioInput> CreateAudioInput() = 0;
 
 	/**
-	 * Remove the audio input from the Pixel Streaming mix.
+	 * @brief Remove the audio input from the Pixel Streaming mix.
 	 */
 	virtual void RemoveAudioInput(TSharedPtr<IPixelStreamingAudioInput> AudioInput) = 0;
 };
diff --git a/Engine/Plugins/Media/PixelStreaming/Source/PixelStreamingEditor/Private/PixelStreamingEditorModule.cpp b/Engine/Plugins/Media/PixelStreaming/Source/PixelStreamingEditor/Private/PixelStreamingEditorModule.cpp
--- a/Engine/Plugins/Media/PixelStreaming/Source/PixelStreamingEditor/Private/PixelStreamingEditorModule.cpp
+++ b/Engine/Plugins/Media/PixelStreaming/Source/PixelStreamingEditor/Private/PixelStreamingEditorModule.cpp
@@ -166,6 +166,7 @@ void FPixelStreamingEditorModule::StartStreaming(UE::EditorPixelStreaming::EStre
 			FSceneViewport* SceneViewport = static_cast<FSceneViewport*>(LevelViewportClient.Viewport);
 			Streamer->SetTargetViewport(SceneViewport->GetViewportWidget());
 			Streamer->SetTargetWindow(SceneViewport->FindWindow());
+			Streamer->SetInputHandlerType(EPixelStreamingInputType::RouteToWindow);
 			Streamer->SetVideoInput(FPixelStreamingVideoInputViewport::Create());
 		}	
 		break;
@@ -173,6 +174,7 @@ void FPixelStreamingEditorModule::StartStreaming(UE::EditorPixelStreaming::EStre
 		{
 			Streamer->SetTargetViewport(nullptr);
 			Streamer->SetTargetWindow(nullptr);
+			Streamer->SetInputHandlerType(EPixelStreamingInputType::RouteToWindow);
 			Streamer->SetVideoInput(FPixelStreamingVideoInputBackBufferComposited::Create());
 		}
 		break;
