Updating some pixel streaming tests.
- Codec negotiation was disabled leading to local description setting failing
- Forced the codec to a supported recv side codec on some tests so webrtc handshake will complete.

diff --git a/Engine/Plugins/Media/PixelStreaming/Source/PixelStreaming/Private/Tests/DataChannelTest.cpp b/Engine/Plugins/Media/PixelStreaming/Source/PixelStreaming/Private/Tests/DataChannelTest.cpp
--- a/Engine/Plugins/Media/PixelStreaming/Source/PixelStreaming/Private/Tests/DataChannelTest.cpp
+++ b/Engine/Plugins/Media/PixelStreaming/Source/PixelStreaming/Private/Tests/DataChannelTest.cpp
@@ -15,6 +15,9 @@ namespace UE::PixelStreaming
     IMPLEMENT_SIMPLE_AUTOMATION_TEST(FDataChannelEchoTest, "System.Plugins.PixelStreaming.FDataChannelEchoTest", EAutomationTestFlags::EditorContext | EAutomationTestFlags::ClientContext | EAutomationTestFlags::ProductFilter)
     bool FDataChannelEchoTest::RunTest(const FString& Parameters)
     {
+        // need to be able to accept codec to handshake otherwise setting local description fails when generating an answer
+        SetCodec(EPixelStreamingCodec::VP8, false /*bUseComputeShader*/);
+
         int32 StreamerPort = 7567;
 		int32 PlayerPort = 4581;
 
diff --git a/Engine/Plugins/Media/PixelStreaming/Source/PixelStreaming/Private/Tests/HandshakeTest.cpp b/Engine/Plugins/Media/PixelStreaming/Source/PixelStreaming/Private/Tests/HandshakeTest.cpp
--- a/Engine/Plugins/Media/PixelStreaming/Source/PixelStreaming/Private/Tests/HandshakeTest.cpp
+++ b/Engine/Plugins/Media/PixelStreaming/Source/PixelStreaming/Private/Tests/HandshakeTest.cpp
@@ -11,8 +11,12 @@ namespace UE::PixelStreaming
 {
 	void DoHandShake(FAutomationTestBase* ParentTest, FMockPlayer::EMode PlayerOfferMode, bool bUseLegacySignallingServer)
 	{
+		// need to be able to accept codec to handshake otherwise setting local description fails when generating an answer
+		SetCodec(EPixelStreamingCodec::VP8, false /*bUseComputeShader*/);
+
 		int32 StreamerPort = 8866;
 		int32 PlayerPort = 6688;
+
 		TSharedPtr<UE::PixelStreamingServers::IServer> OutSignallingServer = bUseLegacySignallingServer ? CreateLegacySignallingServer(StreamerPort, PlayerPort) : CreateSignallingServer(StreamerPort, PlayerPort);
 		TSharedPtr<IPixelStreamingStreamer> OutStreamer = CreateStreamer(StreamerPort);
 		TSharedPtr<FMockPlayer> OutPlayer = CreatePlayer(PlayerOfferMode);
