Fix Pixel Streaming AudioInputMixer crashing on shutdown when using -PixelStreamingExperimentalAudioInput.

robomerge[FN_EngineMerge] -UE5-RES

diff --git a/Engine/Plugins/Media/PixelStreaming/Source/PixelStreaming/Private/AudioInputMixer.cpp b/Engine/Plugins/Media/PixelStreaming/Source/PixelStreaming/Private/AudioInputMixer.cpp
--- a/Engine/Plugins/Media/PixelStreaming/Source/PixelStreaming/Private/AudioInputMixer.cpp
+++ b/Engine/Plugins/Media/PixelStreaming/Source/PixelStreaming/Private/AudioInputMixer.cpp
@@ -57,11 +57,15 @@ namespace UE::PixelStreaming
 
 	void FAudioInputMixer::StopMixing()
 	{
-		if (MixerThread.IsValid())
+		if (MixerThread.IsValid() && bIsMixing)
 		{
 			MixerThread->Kill();
 			bIsMixing = false;
 		}
+		if (Runnable.IsValid())
+		{
+			Runnable->Stop();
+		}
 	}
 
 	/* ---------------------------- FMixerRunnable ---------------------------- */
