Update get_ps_servers scripts to grab the 5.1 branch by default
Additionally, fix bug with shell version of script which caused it not to copy files/folders beginning with "."

diff --git a/Engine/Plugins/Media/PixelStreaming/Resources/WebServers/get_ps_servers.bat b/Engine/Plugins/Media/PixelStreaming/Resources/WebServers/get_ps_servers.bat
--- a/Engine/Plugins/Media/PixelStreaming/Resources/WebServers/get_ps_servers.bat
+++ b/Engine/Plugins/Media/PixelStreaming/Resources/WebServers/get_ps_servers.bat
@@ -51,7 +51,7 @@ if DEFINED UEVersion (
 
 @Rem If no arguments select a specific version, fetch the appropriate default
 if NOT DEFINED PSInfraTagOrBranch (
-    SET PSInfraTagOrBranch=master
+    SET PSInfraTagOrBranch=UE5.1
     SET IsTag=0
 )
 
diff --git a/Engine/Plugins/Media/PixelStreaming/Resources/WebServers/get_ps_servers.sh b/Engine/Plugins/Media/PixelStreaming/Resources/WebServers/get_ps_servers.sh
--- a/Engine/Plugins/Media/PixelStreaming/Resources/WebServers/get_ps_servers.sh
+++ b/Engine/Plugins/Media/PixelStreaming/Resources/WebServers/get_ps_servers.sh
@@ -59,7 +59,7 @@ fi
 # If no arguments select a specific version, fetch the appropriate default
 if [ -z "$PSInfraTagOrBranch" ]
 then
-  PSInfraTagOrBranch=master
+  PSInfraTagOrBranch=UE5.1
   IsTag=0
 fi
 
@@ -80,12 +80,14 @@ else
 
   # Download ps-infra and follow redirects.
   curl -L https://github.com/$PSInfraOrg/$PSInfraRepo/archive/refs/$RefType/$PSInfraTagOrBranch.tar.gz > ps-infra.tar.gz
-  
+
   # Unarchive the .tar
   tar -xmf ps-infra.tar.gz || $(echo "bad archive, contents:" && head --lines=20 ps-infra.tar.gz && exit 0)
 
   # Move the server folders into the current directory (WebServers) and delete the original directory
   mv PixelStreamingInfrastructure-*/* .
+  # Copy any files and folders beginning with dot (ignored by * glob) and discard errors regarding to not being able to move "." and ".."
+  mv PixelStreamingInfrastructure-*/.* . 2>/dev/null || :
   rm -rf PixelStreamingInfrastructure-*
 
   # Delete the downloaded tar
