Add binary message support to Pixel Streaming streamer.

Due to the recent introduction of binary message handling (CL 2372768) in the underlying WebSockets wrapper in UE
this change was now required in order to use the embedded C++ signalling server (which only supports sending binary message types).
This fixes broken editor streaming, viewport streaming, and Pixel Streaming VCam - which were all using the embedded C++ signalling server.

diff --git a/Engine/Plugins/Media/PixelStreaming/Source/PixelStreaming/Private/PixelStreamingSignallingConnection.cpp b/Engine/Plugins/Media/PixelStreaming/Source/PixelStreaming/Private/PixelStreamingSignallingConnection.cpp
--- a/Engine/Plugins/Media/PixelStreaming/Source/PixelStreaming/Private/PixelStreamingSignallingConnection.cpp
+++ b/Engine/Plugins/Media/PixelStreaming/Source/PixelStreaming/Private/PixelStreamingSignallingConnection.cpp
@@ -26,8 +26,8 @@ FPixelStreamingSignallingConnection::FPixelStreamingSignallingConnection(const F
 	RegisterHandler("offer", [this](FJsonObjectPtr JsonMsg) { OnSessionDescription(JsonMsg); });
 	RegisterHandler("answer", [this](FJsonObjectPtr JsonMsg) { OnSessionDescription(JsonMsg); });
 	RegisterHandler("iceCandidate", [this](FJsonObjectPtr JsonMsg) { OnIceCandidate(JsonMsg); });
-	RegisterHandler("ping", [this](FJsonObjectPtr JsonMsg) { /* nothing */ });
-	RegisterHandler("pong", [this](FJsonObjectPtr JsonMsg) { /* nothing */ });
+	RegisterHandler("ping", [this](FJsonObjectPtr JsonMsg) { OnPing(JsonMsg); });
+	RegisterHandler("pong", [this](FJsonObjectPtr JsonMsg) { OnPong(JsonMsg); });
 	RegisterHandler("playerCount", [this](FJsonObjectPtr JsonMsg) { OnPlayerCount(JsonMsg); });
 	RegisterHandler("playerConnected", [this](FJsonObjectPtr JsonMsg) { OnPlayerConnected(JsonMsg); });
 	RegisterHandler("playerDisconnected", [this](FJsonObjectPtr JsonMsg) { OnPlayerDisconnected(JsonMsg); });
@@ -65,6 +65,7 @@ void FPixelStreamingSignallingConnection::Connect(FString InUrl, bool bIsReconne
 	OnConnectionErrorHandle = WebSocket->OnConnectionError().AddLambda([this](const FString& Error) { OnConnectionError(Error); });
 	OnClosedHandle = WebSocket->OnClosed().AddLambda([this](int32 StatusCode, const FString& Reason, bool bWasClean) { OnClosed(StatusCode, Reason, bWasClean); });
 	OnMessageHandle = WebSocket->OnMessage().AddLambda([this](const FString& Msg) { OnMessage(Msg); });
+	OnBinaryMessageHandle = WebSocket->OnBinaryMessage().AddLambda([this](const void* Data, int32 Count, bool bIsLastFragment) { OnBinaryMessage((const uint8*)Data, Count, bIsLastFragment); });
 
 	if (bIsReconnect)
 	{
@@ -100,6 +101,7 @@ void FPixelStreamingSignallingConnection::Disconnect()
 	WebSocket->OnConnectionError().Remove(OnConnectionErrorHandle);
 	WebSocket->OnClosed().Remove(OnClosedHandle);
 	WebSocket->OnMessage().Remove(OnMessageHandle);
+	WebSocket->OnBinaryMessage().Remove(OnBinaryMessageHandle);
 
 	WebSocket->Close();
 	WebSocket = nullptr;
@@ -300,6 +302,16 @@ void FPixelStreamingSignallingConnection::OnConnected()
 	}
 }
 
+void FPixelStreamingSignallingConnection::OnPing(const FJsonObjectPtr& Json)
+{
+	UE_LOG(LogPixelStreamingSS, Verbose, TEXT("Streamer got pinged."));
+}
+
+void FPixelStreamingSignallingConnection::OnPong(const FJsonObjectPtr& Json)
+{
+	UE_LOG(LogPixelStreamingSS, Verbose, TEXT("Streamer got ponged."));
+}
+
 void FPixelStreamingSignallingConnection::OnConnectionError(const FString& Error)
 {
 	UE_LOG(LogPixelStreamingSS, Log, TEXT("Failed to connect to SS %s - signalling server may not be up yet. Message: \"%s\""), *Url, *Error);
@@ -369,6 +381,14 @@ void FPixelStreamingSignallingConnection::OnMessage(const FString& Msg)
 	}
 }
 
+void FPixelStreamingSignallingConnection::OnBinaryMessage(const uint8* Data, int32 Length, bool bIsLastFragment)
+{
+	FUTF8ToTCHAR Convert((const ANSICHAR*)Data, Length);
+	const TCHAR* PayloadChars = Convert.Get();
+	FString Msg = FString(Convert.Length(), PayloadChars);
+	OnMessage(Msg);
+}
+
 void FPixelStreamingSignallingConnection::RegisterHandler(const FString& messageType, const TFunction<void(FJsonObjectPtr)>& handler)
 {
 	MessageHandlers.Add(messageType, handler);
diff --git a/Engine/Plugins/Media/PixelStreaming/Source/PixelStreaming/Public/PixelStreamingSignallingConnection.h b/Engine/Plugins/Media/PixelStreaming/Source/PixelStreaming/Public/PixelStreamingSignallingConnection.h
--- a/Engine/Plugins/Media/PixelStreaming/Source/PixelStreaming/Public/PixelStreamingSignallingConnection.h
+++ b/Engine/Plugins/Media/PixelStreaming/Source/PixelStreaming/Public/PixelStreamingSignallingConnection.h
@@ -105,12 +105,15 @@ private:
 	void OnConnectionError(const FString& Error);
 	void OnClosed(int32 StatusCode, const FString& Reason, bool bWasClean);
 	void OnMessage(const FString& Msg);
+	void OnBinaryMessage(const uint8* Data, int32 Length, bool bIsLastFragment);
 
 	using FJsonObjectPtr = TSharedPtr<FJsonObject>;
 	void RegisterHandler(const FString& messageType, const TFunction<void(FJsonObjectPtr)>& handler);
 
 	void OnIdRequested();
 	void OnConfig(const FJsonObjectPtr& Json);
+	void OnPing(const FJsonObjectPtr& Json);
+	void OnPong(const FJsonObjectPtr& Json);
 	void OnSessionDescription(const FJsonObjectPtr& Json);
 	void OnIceCandidate(const FJsonObjectPtr& Json);
 	void OnPlayerCount(const FJsonObjectPtr& Json);
@@ -149,6 +152,7 @@ private:
 	FDelegateHandle OnConnectionErrorHandle;
 	FDelegateHandle OnClosedHandle;
 	FDelegateHandle OnMessageHandle;
+	FDelegateHandle OnBinaryMessageHandle;
 
 	bool bAutoReconnectEnabled = true;
 	bool bKeepAliveEnabled = true;
diff --git a/Engine/Plugins/Media/PixelStreaming/Source/PixelStreamingServers/Private/WebSocketServerWrapper.cpp b/Engine/Plugins/Media/PixelStreaming/Source/PixelStreamingServers/Private/WebSocketServerWrapper.cpp
--- a/Engine/Plugins/Media/PixelStreaming/Source/PixelStreamingServers/Private/WebSocketServerWrapper.cpp
+++ b/Engine/Plugins/Media/PixelStreaming/Source/PixelStreamingServers/Private/WebSocketServerWrapper.cpp
@@ -38,13 +38,14 @@ namespace UE::PixelStreamingServers
 	{
 		// Convert FString into uint8 array.
 		FTCHARToUTF8 UTF8String(*Message);
-		int32 MessageSize = UTF8String.Length();
-		TArray<uint8> MessageArr;
-		MessageArr.SetNum(MessageSize);
-		FMemory::Memcpy(MessageArr.GetData(), UTF8String.Get(), MessageSize);
+		//int32 MessageSize = UTF8String.Length();
+		//TArray<uint8> MessageArr;
+		//MessageArr.SetNum(MessageSize);
+		//FMemory::Memcpy(MessageArr.GetData(), UTF8String.Get(), MessageSize);
 		
 		// Send the uint8 buffer
-		return SocketConnection->Send(MessageArr.GetData(), MessageSize, false);
+		// Note: Due to how this socket connection is implemented, only binary messages are supported
+		return SocketConnection->Send((const uint8*) UTF8String.Get(), UTF8String.Length(), false);
 	}
 
 	void FWebSocketConnection::SetCallbacks()
diff --git a/Engine/Plugins/Media/PixelStreaming/Source/PixelStreamingServers/Private/WebSocketServerWrapper.h b/Engine/Plugins/Media/PixelStreaming/Source/PixelStreamingServers/Private/WebSocketServerWrapper.h
--- a/Engine/Plugins/Media/PixelStreaming/Source/PixelStreamingServers/Private/WebSocketServerWrapper.h
+++ b/Engine/Plugins/Media/PixelStreaming/Source/PixelStreamingServers/Private/WebSocketServerWrapper.h
@@ -46,6 +46,12 @@ namespace UE::PixelStreamingServers
 		bool IsLaunched() const;
 		bool HasConnections() const;
 		bool Close(uint16 ConnectionId);
+
+		/**
+		 * Send message to a particular websocket connection.
+		 * Note: This message is sent as binary using the websocket protocol because of the WS implementation that is used.
+		 * @return True if the message able to be sent.
+		 */
 		bool Send(uint16 ConnectionId, FString Message) const;
 		bool GetFirstConnection(uint16& OutConnectionId) const;
 
