Add support to set custom "command" handlers for datachannel messages that are of type "Command".

- Moved all existing command handlers to this new system
- Exposed public method IPixelStreamingInputHandler::SetCommandHandler

diff --git a/Engine/Plugins/Media/PixelStreaming/Source/PixelStreaming/Private/PixelStreamingInputHandler.cpp b/Engine/Plugins/Media/PixelStreaming/Source/PixelStreaming/Private/PixelStreamingInputHandler.cpp
--- a/Engine/Plugins/Media/PixelStreaming/Source/PixelStreaming/Private/PixelStreamingInputHandler.cpp
+++ b/Engine/Plugins/Media/PixelStreaming/Source/PixelStreaming/Private/PixelStreamingInputHandler.cpp
@@ -60,6 +60,8 @@ namespace UE::PixelStreaming
 
         RegisterMessageHandler("Command", [this](FMemoryReader Ar) { HandleCommand(Ar); });
         RegisterMessageHandler("UIInteraction", [this](FMemoryReader Ar) { HandleUIInteraction(Ar); });
+
+		PopulateDefaultCommandHandlers();
     }
 
     FPixelStreamingInputHandler::~FPixelStreamingInputHandler()
@@ -659,39 +661,28 @@ namespace UE::PixelStreaming
 		}
 	}
 
-    /**
-     * Command handling
-     */
-    void FPixelStreamingInputHandler::HandleCommand(FMemoryReader Ar)
-    {
-        FString Res;
-        Res.GetCharArray().SetNumUninitialized(Ar.TotalSize() / 2 + 1);
-        Ar.Serialize(Res.GetCharArray().GetData(), Ar.TotalSize());
-
-		FString Descriptor = Res.Mid(MessageHeaderOffset);
-		UE_LOG(LogPixelStreamingInputHandler, Verbose, TEXT("Command: %s"), *Descriptor);
+	void FPixelStreamingInputHandler::SetCommandHandler(const FString& CommandName, const TFunction<void(FString, FString)>& Handler)
+	{
+		CommandHandlers.Add(CommandName, Handler);
+	}
 
-		bool bSuccess = false;
+	void FPixelStreamingInputHandler::PopulateDefaultCommandHandlers() {
 
-		FString ConsoleCommand;
-		UE::PixelStreaming::ExtractJsonFromDescriptor(Descriptor, TEXT("ConsoleCommand"), ConsoleCommand, bSuccess);
-		if (bSuccess && UE::PixelStreaming::Settings::CVarPixelStreamingAllowConsoleCommands.GetValueOnAnyThread())
-		{
+		// Execute console commands if passed "ConsoleCommand" and -PixelStreamingAllowConsoleCommands is on.
+		CommandHandlers.Add(TEXT("ConsoleCommand"), [](FString Descriptor, FString ConsoleCommand){
+			if(!UE::PixelStreaming::Settings::CVarPixelStreamingAllowConsoleCommands.GetValueOnAnyThread())
+			{
+				return;
+			}
 			GEngine->Exec(GEngine->GetWorld(), *ConsoleCommand);
-			return;
-		}
+		});
 
-		/**
-		 * Allowed Console commands
-		 */
-		FString WidthString;
-		FString HeightString;
-		UE::PixelStreaming::ExtractJsonFromDescriptor(Descriptor, TEXT("Resolution.Width"), WidthString, bSuccess);
-		if (bSuccess)
-		{
+		// Change width/height if sent { "Resolution.Width": 1920, "Resolution.Height": 1080 }
+		CommandHandlers.Add(TEXT("Resolution.Width"), [](FString Descriptor, FString WidthString){
+			bool bSuccess = false;
+			FString HeightString;
 			UE::PixelStreaming::ExtractJsonFromDescriptor(Descriptor, TEXT("Resolution.Height"), HeightString, bSuccess);
-
-			if (bSuccess)
+			if(bSuccess) 
 			{
 				int Width = FCString::Atoi(*WidthString);
 				int Height = FCString::Atoi(*HeightString);
@@ -702,68 +693,71 @@ namespace UE::PixelStreaming
 
 				FString ChangeResCommand = FString::Printf(TEXT("r.SetRes %dx%d"), Width, Height);
 				GEngine->Exec(GEngine->GetWorld(), *ChangeResCommand);
-				return;
 			}
-		}
+		});
 
-		FString StatFPSString;
-		UE::PixelStreaming::ExtractJsonFromDescriptor(Descriptor, TEXT("Stat.FPS"), StatFPSString, bSuccess);
-		if (bSuccess)
-		{
+		// Response to "Stat.FPS" by calling "stat fps"
+		CommandHandlers.Add(TEXT("Stat.FPS"), [](FString Descriptor, FString FPSCommand){
 			FString StatFPSCommand = FString::Printf(TEXT("stat fps"));
 			GEngine->Exec(GEngine->GetWorld(), *StatFPSCommand);
-			return;
-		}
+		});
 
-		/**
-		 * Encoder Settings
-		 */
-		FString MinQPString;
-		UE::PixelStreaming::ExtractJsonFromDescriptor(Descriptor, TEXT("Encoder.MinQP"), MinQPString, bSuccess);
-		if (bSuccess)
-		{
+		// Set Encoder.MinQP CVar
+		CommandHandlers.Add(TEXT("Encoder.MinQP"), [](FString Descriptor, FString MinQPString){
 			int MinQP = FCString::Atoi(*MinQPString);
 			UE::PixelStreaming::Settings::CVarPixelStreamingEncoderMinQP->Set(MinQP, ECVF_SetByCommandline);
-			return;
-		}
+		});
 
-		FString MaxQPString;
-		UE::PixelStreaming::ExtractJsonFromDescriptor(Descriptor, TEXT("Encoder.MaxQP"), MaxQPString, bSuccess);
-		if (bSuccess)
-		{
+		// Set Encoder.MaxQP CVar
+		CommandHandlers.Add(TEXT("Encoder.MaxQP"), [](FString Descriptor, FString MaxQPString){
 			int MaxQP = FCString::Atoi(*MaxQPString);
 			UE::PixelStreaming::Settings::CVarPixelStreamingEncoderMaxQP->Set(MaxQP, ECVF_SetByCommandline);
-			return;
-		}
+		});
 
-		/**
-		 * WebRTC Settings
-		 */
-		FString FPSString;
-		UE::PixelStreaming::ExtractJsonFromDescriptor(Descriptor, TEXT("WebRTC.Fps"), FPSString, bSuccess);
-		if (bSuccess)
-		{
+		// Set WebRTC max FPS
+		CommandHandlers.Add(TEXT("WebRTC.Fps"), [](FString Descriptor, FString FPSString){
 			int FPS = FCString::Atoi(*FPSString);
 			UE::PixelStreaming::Settings::CVarPixelStreamingWebRTCFps->Set(FPS, ECVF_SetByCommandline);
-			return;
-		}
+		});
 
-		FString MinBitrateString;
-		UE::PixelStreaming::ExtractJsonFromDescriptor(Descriptor, TEXT("WebRTC.MinBitrate"), MinBitrateString, bSuccess);
-		if (bSuccess)
-		{
+		// Set MinBitrate
+		CommandHandlers.Add(TEXT("WebRTC.MinBitrate"), [](FString Descriptor, FString MinBitrateString){
 			int MinBitrate = FCString::Atoi(*MinBitrateString);
 			UE::PixelStreaming::Settings::CVarPixelStreamingWebRTCMinBitrate->Set(MinBitrate, ECVF_SetByCommandline);
-			return;
-		}
+		});
 
-		FString MaxBitrateString;
-		UE::PixelStreaming::ExtractJsonFromDescriptor(Descriptor, TEXT("WebRTC.MaxBitrate"), MaxBitrateString, bSuccess);
-		if (bSuccess)
-		{
+		// Set MaxBitrate
+		CommandHandlers.Add(TEXT("WebRTC.MaxBitrate"), [](FString Descriptor, FString MaxBitrateString){
 			int MaxBitrate = FCString::Atoi(*MaxBitrateString);
 			UE::PixelStreaming::Settings::CVarPixelStreamingWebRTCMaxBitrate->Set(MaxBitrate, ECVF_SetByCommandline);
-			return;
+		});
+
+	}
+
+    /**
+     * Command handling
+     */
+    void FPixelStreamingInputHandler::HandleCommand(FMemoryReader Ar)
+    {
+        FString Res;
+        Res.GetCharArray().SetNumUninitialized(Ar.TotalSize() / 2 + 1);
+        Ar.Serialize(Res.GetCharArray().GetData(), Ar.TotalSize());
+
+		FString Descriptor = Res.Mid(MessageHeaderOffset);
+		UE_LOG(LogPixelStreamingInputHandler, Verbose, TEXT("Command: %s"), *Descriptor);
+
+		// Iterate each command handler and see if the command we got matches any of the bound command names.
+		for(auto& CommandHandlersPair : CommandHandlers) 
+		{
+			FString CommandValue;
+			bool bSuccess = false;
+			UE::PixelStreaming::ExtractJsonFromDescriptor(Descriptor, CommandHandlersPair.Key, CommandValue, bSuccess);
+			if(bSuccess) 
+			{
+				// Execute bound command handler with descriptor and parsed command value
+				CommandHandlersPair.Value(Descriptor, CommandValue);
+				return;
+			}
 		}
 	}
 
diff --git a/Engine/Plugins/Media/PixelStreaming/Source/PixelStreaming/Private/PixelStreamingInputHandler.h b/Engine/Plugins/Media/PixelStreaming/Source/PixelStreaming/Private/PixelStreamingInputHandler.h
--- a/Engine/Plugins/Media/PixelStreaming/Source/PixelStreaming/Private/PixelStreamingInputHandler.h
+++ b/Engine/Plugins/Media/PixelStreaming/Source/PixelStreaming/Private/PixelStreamingInputHandler.h
@@ -28,6 +28,9 @@ namespace UE::PixelStreaming
 		/** Set which MessageHandler will route input  */
 		virtual void SetMessageHandler(const TSharedRef<FGenericApplicationMessageHandler>& InTargetHandler) override;
 
+		/** Register a custom function to execute when command JSON is received. */
+		virtual void SetCommandHandler(const FString& CommandName, const TFunction<void(FString, FString)>& Handler) override;
+
 		/** Exec handler to allow console commands to be passed through for debugging */
 		virtual bool Exec(UWorld* InWorld, const TCHAR* Cmd, FOutputDevice& Ar) override;
 		/**
@@ -93,6 +96,11 @@ namespace UE::PixelStreaming
          */
         virtual void HandleUIInteraction(FMemoryReader Ar);
 
+		/**
+		 * Populate default command handlers for data channel messages sent with "{ type: "Command" }".
+		 */
+		void PopulateDefaultCommandHandlers();
+
 		FIntPoint ConvertFromNormalizedScreenLocation(const FVector2D& ScreenLocation, bool bIncludeOffset = true);
 		FWidgetPath FindRoutingMessageWidget(const FVector2D& Location) const;
 
@@ -169,6 +177,15 @@ namespace UE::PixelStreaming
 		 * 2 characters for the length which are skipped
 		 */
 		const size_t MessageHeaderOffset = 1;
+
+		/**
+		 * A map of named commands we respond to when we receive a datachannel message of type "command".
+		 * Key = command name (e.g "Encoder.MaxQP")
+		 * Value = The command handler lambda function whose parameters are as follows:
+		 * 	FString - the descriptor (e.g. the full json payload of the command message)
+		 * 	FString - the parsed value of the command, e.g. if key was "Encoder.MaxQP" and descriptor was { type: "Command", "Encoder.MaxQP": 51 }, then parsed value is "51".
+		 */
+		TMap<FString, TFunction<void(FString, FString)>> CommandHandlers;
 	
 	private:
 		float uint16_MAX = (float) UINT16_MAX;
diff --git a/Engine/Plugins/Media/PixelStreaming/Source/PixelStreaming/Public/IPixelStreamingInputHandler.h b/Engine/Plugins/Media/PixelStreaming/Source/PixelStreaming/Public/IPixelStreamingInputHandler.h
--- a/Engine/Plugins/Media/PixelStreaming/Source/PixelStreaming/Public/IPixelStreamingInputHandler.h
+++ b/Engine/Plugins/Media/PixelStreaming/Source/PixelStreaming/Public/IPixelStreamingInputHandler.h
@@ -82,6 +82,12 @@ public:
 	 */
 	virtual void RegisterMessageHandler(const FString& MessageType, const TFunction<void(FMemoryReader)>& Handler) = 0;
 
+	/**
+	 * @brief Register a custom function to execute when command JSON is received over the data channel: "{ type: "Command", YourCommand: YourCommandValue }".
+	 * Note: You can also override the default Pixel Streaming command handlers by setting handlers with the same name as those already used, e.g. "Stat.FPS".
+	 */
+	virtual void SetCommandHandler(const FString& CommandName, const TFunction<void(FString, FString)>& Handler) = 0;
+
 	/**
 	 * @brief Find the function to be called whenever the specified message type is received.
 	 * 
diff --git a/Engine/Plugins/Media/PixelStreaming/Source/PixelStreamingEditor/Private/PixelStreamingEditorModule.cpp b/Engine/Plugins/Media/PixelStreaming/Source/PixelStreamingEditor/Private/PixelStreamingEditorModule.cpp
--- a/Engine/Plugins/Media/PixelStreaming/Source/PixelStreamingEditor/Private/PixelStreamingEditorModule.cpp
+++ b/Engine/Plugins/Media/PixelStreaming/Source/PixelStreamingEditor/Private/PixelStreamingEditorModule.cpp
@@ -59,46 +59,27 @@ void FPixelStreamingEditorModule::InitEditorStreaming(IPixelStreamingModule& Mod
 		return;
 	}
 
-	// The current handler is the function that will currently be executed if a message with type "Command" is received
-	TFunction<void(FMemoryReader)> CurrentHandler = Module.FindMessageHandler("Command");
-	TFunction<void(FMemoryReader)> ExtendedHandler = [CurrentHandler](FMemoryReader Ar)
+	// Add custom handle for { type: "Command", Resolution.Width: "1920", Resolution.Height: "1080" } when doing Editor streaming
+	// because we cannot resize the game viewport, but instead want to resize the parent window.
+	TSharedPtr<IPixelStreamingInputHandler> InputHandler = Streamer->GetInputHandler().Pin();
+	if(InputHandler) 
 	{
-		// We then create our new handler which will execute the "base" handler
-		CurrentHandler(Ar);
-		// and then perform out extended functionality after.
-		// equivalent to the super::DoSomeFunc pattern
-		FString Res;
-		Res.GetCharArray().SetNumUninitialized(Ar.TotalSize() / 2 + 1);
-		Ar.Serialize(Res.GetCharArray().GetData(), Ar.TotalSize());
-		FString Descriptor = Res.Mid(1);
-		FString WidthString;
-		FString HeightString;
-		bool bSuccess;
-		UE::PixelStreaming::ExtractJsonFromDescriptor(Descriptor, TEXT("Resolution.Width"), WidthString, bSuccess);
-		if (bSuccess)
-		{
+		InputHandler->SetCommandHandler("Resolution.Width", [](FString Descriptor, FString WidthString){
+			bool bSuccess;
+			FString HeightString;
 			UE::PixelStreaming::ExtractJsonFromDescriptor(Descriptor, TEXT("Resolution.Height"), HeightString, bSuccess);
-			if (bSuccess)
+			int Width = FCString::Atoi(*WidthString);
+			int Height = FCString::Atoi(*HeightString);
+			if (Width < 1 || Height < 1)
 			{
-				int Width = FCString::Atoi(*WidthString);
-				int Height = FCString::Atoi(*HeightString);
-				if (Width < 1 || Height < 1)
-				{
-					return;
-				}
-
-				TSharedPtr<SWindow> ParentWindow = IMainFrameModule::Get().GetParentWindow();
-				ParentWindow->Resize(FVector2D(Width, Height));
-				FSlateApplication::Get().OnSizeChanged(ParentWindow->GetNativeWindow().ToSharedRef(), Width, Height);
+				return;
 			}
-		}
-	};
-	Module.RegisterMessage( Protocol::EPixelStreamingMessageDirection::ToStreamer, 
-							"Command", 
-							Protocol::FPixelStreamingInputMessage(51, 0, {}),
-							[ExtendedHandler](FMemoryReader Ar) {
-								ExtendedHandler(Ar);
-							});
+
+			TSharedPtr<SWindow> ParentWindow = IMainFrameModule::Get().GetParentWindow();
+			ParentWindow->Resize(FVector2D(Width, Height));
+			FSlateApplication::Get().OnSizeChanged(ParentWindow->GetNativeWindow().ToSharedRef(), Width, Height);
+		});
+	}
 
 	// Give the editor streamer the default url if the user hasn't specified one when launching the editor
 	if (Streamer->GetSignallingServerURL().IsEmpty())
