Expose Pixel Streaming JSON parsing to public API.

diff --git a/Engine/Plugins/Media/PixelStreaming/Source/PixelStreaming/Private/PixelStreamingInputComponent.cpp b/Engine/Plugins/Media/PixelStreaming/Source/PixelStreaming/Private/PixelStreamingInputComponent.cpp
--- a/Engine/Plugins/Media/PixelStreaming/Source/PixelStreaming/Private/PixelStreamingInputComponent.cpp
+++ b/Engine/Plugins/Media/PixelStreaming/Source/PixelStreaming/Private/PixelStreamingInputComponent.cpp
@@ -13,7 +13,7 @@
 #include "PixelStreamingProtocol.h"
 #include "IPixelStreamingStreamer.h"
 #include "Settings.h"
-#include "Utils.h"
+#include "PixelStreamingUtils.h"
 
 UPixelStreamingInput::UPixelStreamingInput(const FObjectInitializer& ObjectInitializer)
 	: Super(ObjectInitializer)
diff --git a/Engine/Plugins/Media/PixelStreaming/Source/PixelStreaming/Private/PixelStreamingInputHandler.cpp b/Engine/Plugins/Media/PixelStreaming/Source/PixelStreaming/Private/PixelStreamingInputHandler.cpp
--- a/Engine/Plugins/Media/PixelStreaming/Source/PixelStreaming/Private/PixelStreamingInputHandler.cpp
+++ b/Engine/Plugins/Media/PixelStreaming/Source/PixelStreaming/Private/PixelStreamingInputHandler.cpp
@@ -12,7 +12,7 @@
 #include "Widgets/SViewport.h"
 #include "Policies/CondensedJsonPrintPolicy.h"
 #include "Serialization/JsonSerializer.h"
-#include "Utils.h"
+#include "PixelStreamingUtils.h"
 #include "PixelStreamingInputComponent.h"
 #include "Engine/Engine.h"
 #include "Framework/Application/SlateUser.h"
diff --git a/Engine/Plugins/Media/PixelStreaming/Source/PixelStreaming/Private/PixelStreamingUtils.cpp b/Engine/Plugins/Media/PixelStreaming/Source/PixelStreaming/Private/PixelStreamingUtils.cpp
new file mode 100644
--- /dev/null
+++ b/Engine/Plugins/Media/PixelStreaming/Source/PixelStreaming/Private/PixelStreamingUtils.cpp
@@ -0,0 +1,50 @@
+// Copyright Epic Games, Inc. All Rights Reserved.
+
+#include "PixelStreamingUtils.h"
+#include "Dom/JsonObject.h"
+#include "Serialization/JsonReader.h"
+#include "Serialization/JsonWriter.h"
+#include "Serialization/JsonSerializer.h"
+#include "Policies/CondensedJsonPrintPolicy.h"
+
+namespace UE::PixelStreaming
+{
+	void ExtendJsonWithField(const FString& Descriptor, FString FieldName, FString StringValue, FString& OutNewDescriptor, bool& OutSuccess)
+	{
+		TSharedPtr<FJsonObject> JsonObject = MakeShareable(new FJsonObject);
+
+		if (!Descriptor.IsEmpty())
+		{
+			TSharedRef<TJsonReader<>> JsonReader = TJsonReaderFactory<>::Create(Descriptor);
+			if (!FJsonSerializer::Deserialize(JsonReader, JsonObject) || !JsonObject.IsValid())
+			{
+				OutSuccess = false;
+				return;
+			}
+		}
+
+		TSharedRef<FJsonValueString> JsonValueObject = MakeShareable(new FJsonValueString(StringValue));
+		JsonObject->SetField(FieldName, JsonValueObject);
+
+		TSharedRef<TJsonWriter<TCHAR, TCondensedJsonPrintPolicy<TCHAR>>> JsonWriter = TJsonWriterFactory<TCHAR, TCondensedJsonPrintPolicy<TCHAR>>::Create(&OutNewDescriptor);
+		OutSuccess = FJsonSerializer::Serialize(JsonObject.ToSharedRef(), JsonWriter);
+	}
+
+	void ExtractJsonFromDescriptor(FString Descriptor, FString FieldName, FString& OutStringValue, bool& OutSuccess)
+	{
+		TSharedPtr<FJsonObject> JsonObject = MakeShareable(new FJsonObject);
+
+		TSharedRef<TJsonReader<>> JsonReader = TJsonReaderFactory<>::Create(Descriptor);
+		if (FJsonSerializer::Deserialize(JsonReader, JsonObject) && JsonObject.IsValid())
+		{
+			const TSharedPtr<FJsonObject>* JsonObjectPtr = &JsonObject;
+			OutSuccess = (*JsonObjectPtr)->TryGetStringField(FieldName, OutStringValue);
+		}
+		else
+		{
+			OutSuccess = false;
+		}
+	}
+
+} // UE::PixelStreaming
+
diff --git a/Engine/Plugins/Media/PixelStreaming/Source/PixelStreaming/Private/Utils.h b/Engine/Plugins/Media/PixelStreaming/Source/PixelStreaming/Private/Utils.h
--- a/Engine/Plugins/Media/PixelStreaming/Source/PixelStreaming/Private/Utils.h
+++ b/Engine/Plugins/Media/PixelStreaming/Source/PixelStreaming/Private/Utils.h
@@ -4,11 +4,6 @@
 
 #include "WebRTCIncludes.h"
 #include "Async/Async.h"
-#include "Dom/JsonObject.h"
-#include "Serialization/JsonReader.h"
-#include "Serialization/JsonWriter.h"
-#include "Serialization/JsonSerializer.h"
-#include "Policies/CondensedJsonPrintPolicy.h"
 #include "CoreMinimal.h"
 #include "HAL/IConsoleManager.h"
 #include "Misc/CommandLine.h"
@@ -131,40 +126,4 @@ namespace UE::PixelStreaming
 		return Pos + DataSize;
 	}
 
-	inline void ExtendJsonWithField(const FString& Descriptor, FString FieldName, FString StringValue, FString& NewDescriptor, bool& Success)
-	{
-		TSharedPtr<FJsonObject> JsonObject = MakeShareable(new FJsonObject);
-
-		if (!Descriptor.IsEmpty())
-		{
-			TSharedRef<TJsonReader<>> JsonReader = TJsonReaderFactory<>::Create(Descriptor);
-			if (!FJsonSerializer::Deserialize(JsonReader, JsonObject) || !JsonObject.IsValid())
-			{
-				Success = false;
-				return;
-			}
-		}
-
-		TSharedRef<FJsonValueString> JsonValueObject = MakeShareable(new FJsonValueString(StringValue));
-		JsonObject->SetField(FieldName, JsonValueObject);
-
-		TSharedRef<TJsonWriter<TCHAR, TCondensedJsonPrintPolicy<TCHAR>>> JsonWriter = TJsonWriterFactory<TCHAR, TCondensedJsonPrintPolicy<TCHAR>>::Create(&NewDescriptor);
-		Success = FJsonSerializer::Serialize(JsonObject.ToSharedRef(), JsonWriter);
-	}
-
-	inline void ExtractJsonFromDescriptor(FString Descriptor, FString FieldName, FString& StringValue, bool& Success)
-	{
-		TSharedPtr<FJsonObject> JsonObject = MakeShareable(new FJsonObject);
-
-		TSharedRef<TJsonReader<>> JsonReader = TJsonReaderFactory<>::Create(Descriptor);
-		if (FJsonSerializer::Deserialize(JsonReader, JsonObject) && JsonObject.IsValid())
-		{
-			const TSharedPtr<FJsonObject>* JsonObjectPtr = &JsonObject;
-			Success = (*JsonObjectPtr)->TryGetStringField(FieldName, StringValue);
-		}
-		else
-		{
-			Success = false;
-		}
-	}
 } // namespace UE::PixelStreaming
diff --git a/Engine/Plugins/Media/PixelStreaming/Source/PixelStreaming/Public/PixelStreamingUtils.h b/Engine/Plugins/Media/PixelStreaming/Source/PixelStreaming/Public/PixelStreamingUtils.h
new file mode 100644
--- /dev/null
+++ b/Engine/Plugins/Media/PixelStreaming/Source/PixelStreaming/Public/PixelStreamingUtils.h
@@ -0,0 +1,27 @@
+// Copyright Epic Games, Inc. All Rights Reserved.
+
+#pragma once
+
+#include "Containers/UnrealString.h"
+
+namespace UE::PixelStreaming
+{
+	/**
+	 * Extends an existing JSON string with a new name: value property.
+	 * @param Descriptor The JSON string.
+	 * @param FieldName The field name to add the json string.
+	 * @param StringValue The value to set in the new field.
+	 * @param OutNewDescriptor The new JSON string.
+	 * @param OutSuccess True if we were able to extend the JSON string.
+	 */
+	PIXELSTREAMING_API void ExtendJsonWithField(const FString& Descriptor, FString FieldName, FString StringValue, FString& OutNewDescriptor, bool& OutSuccess);
+
+	/**
+	 * Extract a field's value from a JSON string.
+	 * @param Descriptor The JSON string.
+	 * @param FieldName The field name to extract a value from.
+	 * @param OutStringValue The value from the extracted field.
+	 * @param OutSuccess True if we were able to extract the value.
+	 */
+	PIXELSTREAMING_API void ExtractJsonFromDescriptor(FString Descriptor, FString FieldName, FString& OutStringValue, bool& OutSuccess);
+}
\ No newline at end of file
diff --git a/Engine/Plugins/Media/PixelStreaming/Source/PixelStreamingEditor/Private/PixelStreamingEditorModule.cpp b/Engine/Plugins/Media/PixelStreaming/Source/PixelStreamingEditor/Private/PixelStreamingEditorModule.cpp
--- a/Engine/Plugins/Media/PixelStreaming/Source/PixelStreamingEditor/Private/PixelStreamingEditorModule.cpp
+++ b/Engine/Plugins/Media/PixelStreaming/Source/PixelStreamingEditor/Private/PixelStreamingEditorModule.cpp
@@ -13,7 +13,7 @@
 #include "UnrealEngine.h"
 #include "Interfaces/IMainFrameModule.h"
 #include "PixelStreamingProtocol.h"
-#include "Utils.h"
+#include "PixelStreamingUtils.h"
 #include "GenericPlatform/GenericPlatformHttp.h"
 #include "LevelEditorViewport.h"
 #include "Slate/SceneViewport.h"
