Extend JSON handling in Pixel Streaming to handle nested JSON objects.

diff --git a/Engine/Plugins/Media/PixelStreaming/Source/PixelStreaming/Private/PixelStreamingUtils.cpp b/Engine/Plugins/Media/PixelStreaming/Source/PixelStreaming/Private/PixelStreamingUtils.cpp
--- a/Engine/Plugins/Media/PixelStreaming/Source/PixelStreaming/Private/PixelStreamingUtils.cpp
+++ b/Engine/Plugins/Media/PixelStreaming/Source/PixelStreaming/Private/PixelStreamingUtils.cpp
@@ -32,12 +32,39 @@ namespace UE::PixelStreaming
 
 	void ExtractJsonFromDescriptor(FString Descriptor, FString FieldName, FString& OutStringValue, bool& OutSuccess)
 	{
+		/*
+		 * ExtractJsonFromDescriptor supports parsing nested objects, as well as keys that contain a period.
+		 * eg
+		 * { "Encoder.MinQP": val }
+		 * vs
+		 * { "Encoder": { "MinQP": val }}
+		 */
 		TSharedPtr<FJsonObject> JsonObject = MakeShareable(new FJsonObject);
 
 		TSharedRef<TJsonReader<>> JsonReader = TJsonReaderFactory<>::Create(Descriptor);
 		if (FJsonSerializer::Deserialize(JsonReader, JsonObject) && JsonObject.IsValid())
 		{
 			const TSharedPtr<FJsonObject>* JsonObjectPtr = &JsonObject;
+
+			// Create a copy of our field name
+			FString OriginalFieldName = FieldName;
+			if (FieldName.Contains(TEXT(".")))
+			{
+				TArray<FString> FieldComponents;
+				FieldName.ParseIntoArray(FieldComponents, TEXT("."));
+				FieldName = FieldComponents.Pop();
+
+				for (const FString& FieldComponent : FieldComponents)
+				{
+					if (!(*JsonObjectPtr)->TryGetObjectField(FieldComponent, JsonObjectPtr))
+					{
+						// If we can't get the nested object, try parse the key as a whole just incase the key contains a period
+						OutSuccess = (*JsonObjectPtr)->TryGetStringField(OriginalFieldName, OutStringValue);
+						return;
+					}
+				}
+			}
+
 			OutSuccess = (*JsonObjectPtr)->TryGetStringField(FieldName, OutStringValue);
 		}
 		else
