Integration of AVCodecs into PixelStreaming and support for new pixel formats

diff --git a/Engine/Plugins/Media/PixelCapture/Source/PixelCapture/Public/PixelCaptureOutputFrameRHI.h b/Engine/Plugins/Media/PixelCapture/Source/PixelCapture/Public/PixelCaptureOutputFrameRHI.h
--- a/Engine/Plugins/Media/PixelCapture/Source/PixelCapture/Public/PixelCaptureOutputFrameRHI.h
+++ b/Engine/Plugins/Media/PixelCapture/Source/PixelCapture/Public/PixelCaptureOutputFrameRHI.h
@@ -11,17 +11,54 @@
 class PIXELCAPTURE_API FPixelCaptureOutputFrameRHI : public IPixelCaptureOutputFrame
 {
 public:
+	/**
+	 * Sometimes we need to pack data in a pixel format not fully supported by the
+	 * RHI, so we may instead store it as a different format for transport and read it as the 
+	 * original format sometime later, this struct is meant to hold information
+	 * about the wrapper format if it exists
+	 */
+	struct WrapperFormatData
+	{
+		int32 Width;
+		int32 Height;
+		EPixelFormat Format;
+	};
+
 	FPixelCaptureOutputFrameRHI(FTexture2DRHIRef InFrameTexture)
-		: FrameTexture(InFrameTexture)
+		: FrameTexture(InFrameTexture), bHasWrapperFormatData(false), WrapperData()
 	{
 	}
+
+	FPixelCaptureOutputFrameRHI(FTexture2DRHIRef InFrameTexture, WrapperFormatData InWrapperData)
+		: FrameTexture(InFrameTexture), bHasWrapperFormatData(true), WrapperData(InWrapperData)
+	{
+	}
+
 	virtual ~FPixelCaptureOutputFrameRHI() = default;
 
-	virtual int32 GetWidth() const override { return FrameTexture->GetDesc().Extent.X; }
-	virtual int32 GetHeight() const override { return FrameTexture->GetDesc().Extent.Y; }
+	virtual int32 GetWidth() const override 
+	{
+		return bHasWrapperFormatData ? WrapperData.Width : FrameTexture->GetDesc().Extent.X;
+	}
+	virtual int32 GetHeight() const override
+	{
+		return bHasWrapperFormatData ? WrapperData.Height : FrameTexture->GetDesc().Extent.Y;
+	}
+
+	virtual bool GetWrapperFormatData(WrapperFormatData& OutData) const { 
+		if(bHasWrapperFormatData)
+		{
+			OutData = WrapperData;
+			return true;
+		}
+
+		return false;
+	}
 
 	FTexture2DRHIRef GetFrameTexture() const { return FrameTexture; }
 
 private:
 	FTexture2DRHIRef FrameTexture;
+	bool bHasWrapperFormatData;
+	WrapperFormatData WrapperData;
 };
diff --git a/Engine/Plugins/Media/PixelStreaming/PixelStreaming.uplugin b/Engine/Plugins/Media/PixelStreaming/PixelStreaming.uplugin
--- a/Engine/Plugins/Media/PixelStreaming/PixelStreaming.uplugin
+++ b/Engine/Plugins/Media/PixelStreaming/PixelStreaming.uplugin
@@ -49,6 +49,20 @@
 			"LoadingPhase": "PostEngineInit",
 			"PlatformAllowList": [ "Win64", "Linux" ],
 			"TargetDenyList": [ "Server" ]
+		},
+		{
+			"Name": "PixelStreamingHMD",
+			"Type": "Runtime",
+			"LoadingPhase": "PreDefault",
+			"PlatformAllowList": [ "Win64", "Linux" ],
+			"TargetDenyList": [ "Server" ]
+		},
+		{
+			"Name": "PixelStreamingInput",
+			"Type": "Runtime",
+			"LoadingPhase": "PreDefault",
+			"PlatformAllowList": [ "Win64", "Linux" ],
+			"TargetDenyList": [ "Server" ]
 		}
 	],
 	"Plugins": [
@@ -62,9 +76,18 @@
 			"PlatformDenyList": [ "Linux" ]
 		},
 		{
-			"Name": "HardwareEncoders",
+			"Name": "AVCodecsCore",
+			"Enabled": true
+		},
+		{
+			"Name": "NVCodecs",
 			"Enabled": true
 		},
+		{
+			"Name": "AMFCodecs",
+			"Enabled": true,
+			"PlatformDenyList": [ "Linux" ]
+		},
 		{
 			"Name": "WebSocketNetworking",
 			"Enabled": true
diff --git a/Engine/Plugins/Media/PixelStreaming/Source/PixelStreaming/PixelStreaming.Build.cs b/Engine/Plugins/Media/PixelStreaming/Source/PixelStreaming/PixelStreaming.Build.cs
--- a/Engine/Plugins/Media/PixelStreaming/Source/PixelStreaming/PixelStreaming.Build.cs
+++ b/Engine/Plugins/Media/PixelStreaming/Source/PixelStreaming/PixelStreaming.Build.cs
@@ -65,22 +65,17 @@ namespace UnrealBuildTool.Rules
 				"ApplicationCore",
 				"InputDevice",
 				"WebRTC",
-				"PixelCapture"
+				"PixelCapture",
+				"PixelStreamingInput"
 			});
 
 			// NOTE: General rule is not to access the private folder of another module
 			PrivateIncludePaths.AddRange(new string[]
 			{
-				System.IO.Path.Combine(GetModuleDirectory("AudioMixer"), "Private"),
-				System.IO.Path.Combine(GetModuleDirectory("Renderer"), "Private"),
+				Path.Combine(GetModuleDirectory("AudioMixer"), "Private"),
+				Path.Combine(GetModuleDirectory("Renderer"), "Private"),
 			});
 
-			// WebRTC third party includes (just libyuv for colour format conversions for now)
-			PublicIncludePaths.AddRange(new string[]
-				{
-					Path.Combine(EngineDir, "Source/ThirdParty/WebRTC/4147/Include/third_party/libyuv/include"),
-				});
-
 			PrivateDependencyModuleNames.AddRange(new string[]
 			{
 				"Core",
@@ -101,11 +96,16 @@ namespace UnrealBuildTool.Rules
 				"Sockets",
 				"MediaUtils",
 				"DeveloperSettings",
-				"AVEncoder",
+				"AVCodecsCore",
+				"AVCodecsCoreRHI",
 				"PixelCaptureShaders",
 				"PixelStreamingServers",
+				"PixelStreamingHMD",
+				"TraceLog",
 			});
 
+			PrivateDefinitions.Add("PIXELSTREAMING_DUMP_ENCODING=0");
+
 			PrivateDependencyModuleNames.Add("VulkanRHI");
 			AddEngineThirdPartyPrivateStaticDependencies(Target, "Vulkan", "CUDA");
 
diff --git a/Engine/Plugins/Media/PixelStreaming/Source/PixelStreaming/Private/AdaptedVideoTrackSource.cpp b/Engine/Plugins/Media/PixelStreaming/Source/PixelStreaming/Private/AdaptedVideoTrackSource.cpp
--- a/Engine/Plugins/Media/PixelStreaming/Source/PixelStreaming/Private/AdaptedVideoTrackSource.cpp
+++ b/Engine/Plugins/Media/PixelStreaming/Source/PixelStreaming/Private/AdaptedVideoTrackSource.cpp
@@ -13,11 +13,7 @@ namespace UE::PixelStreaming
 
 	bool FAdaptedVideoTrackSource::GetStats(Stats* stats)
 	{
-#if WEBRTC_VERSION == 84
-		rtc::CritScope lock(&stats_crit_);
-#elif WEBRTC_VERSION == 96
 		webrtc::MutexLock lock(&stats_mutex_);
-#endif
 
 		if (!stats_)
 		{
@@ -39,7 +35,7 @@ namespace UE::PixelStreaming
 			pending rotation despite some sink with wants.rotation_applied ==
 			true was just added. The VideoBroadcaster enforces
 			synchronization for us in this case, by not passing the frame on
-			to sinks which don't want it. 
+			to sinks which don't want it.
 		*/
 		if (apply_rotation() && frame.rotation() != webrtc::kVideoRotation_0 && buffer->type() == webrtc::VideoFrameBuffer::Type::kI420)
 		{
@@ -93,11 +89,7 @@ namespace UE::PixelStreaming
 		int* crop_y)
 	{
 		{
-#if WEBRTC_VERSION == 84
-			rtc::CritScope lock(&stats_crit_);
-#elif WEBRTC_VERSION == 96
 			webrtc::MutexLock lock(&stats_mutex_);
-#endif
 			stats_ = Stats{ width, height };
 		}
 
@@ -120,4 +112,4 @@ namespace UE::PixelStreaming
 		return true;
 	}
 
-} // namespace rtc
+} // namespace UE::PixelStreaming
diff --git a/Engine/Plugins/Media/PixelStreaming/Source/PixelStreaming/Private/AdaptedVideoTrackSource.h b/Engine/Plugins/Media/PixelStreaming/Source/PixelStreaming/Private/AdaptedVideoTrackSource.h
--- a/Engine/Plugins/Media/PixelStreaming/Source/PixelStreaming/Private/AdaptedVideoTrackSource.h
+++ b/Engine/Plugins/Media/PixelStreaming/Source/PixelStreaming/Private/AdaptedVideoTrackSource.h
@@ -94,14 +94,8 @@ namespace UE::PixelStreaming
 		void RemoveEncodedSink(rtc::VideoSinkInterface<webrtc::RecordableEncodedFrame>* sink) override {}
 
 		cricket::VideoAdapter video_adapter_;
-
-#if WEBRTC_VERSION == 84
-		rtc::CriticalSection stats_crit_;
-		absl::optional<Stats> stats_ RTC_GUARDED_BY(stats_crit_);
-#elif WEBRTC_VERSION == 96
 		webrtc::Mutex stats_mutex_;
 		absl::optional<Stats> stats_ RTC_GUARDED_BY(stats_mutex_);
-#endif
 
 		rtc::VideoBroadcaster broadcaster_;
 	};
diff --git a/Engine/Plugins/Media/PixelStreaming/Source/PixelStreaming/Private/AudioCapturer.cpp b/Engine/Plugins/Media/PixelStreaming/Source/PixelStreaming/Private/AudioCapturer.cpp
--- a/Engine/Plugins/Media/PixelStreaming/Source/PixelStreaming/Private/AudioCapturer.cpp
+++ b/Engine/Plugins/Media/PixelStreaming/Source/PixelStreaming/Private/AudioCapturer.cpp
@@ -3,6 +3,7 @@
 #include "AudioCapturer.h"
 #include "AudioMixerDevice.h"
 #include "SampleBuffer.h"
+#include "PixelStreamingTrace.h"
 #include "Engine/GameEngine.h"
 
 DECLARE_LOG_CATEGORY_EXTERN(LogPixelStreamingAudioCapturer, Log, All);
@@ -31,6 +32,7 @@ namespace UE::PixelStreaming
 
 	void FAudioCapturer::OnNewSubmixBuffer(const USoundSubmix* OwningSubmix, float* AudioData, int32 NumSamples, int32 InNumChannels, const int32 InSampleRate, double AudioClock)
 	{
+		TRACE_CPUPROFILER_EVENT_SCOPE_ON_CHANNEL_STR("PixelStreaming FAudioCapturer::OnNewSubmixBuffer", PixelStreamingChannel);
 		if (!(bInitialized && bRecordingInitialized))
 		{
 			return;
@@ -363,4 +365,4 @@ namespace UE::PixelStreaming
 		*enabled = true;
 		return 0;
 	}
-} // namespace UE::PixelStreaming
\ No newline at end of file
+} // namespace UE::PixelStreaming
diff --git a/Engine/Plugins/Media/PixelStreaming/Source/PixelStreaming/Private/AudioInputMixer.cpp b/Engine/Plugins/Media/PixelStreaming/Source/PixelStreaming/Private/AudioInputMixer.cpp
--- a/Engine/Plugins/Media/PixelStreaming/Source/PixelStreaming/Private/AudioInputMixer.cpp
+++ b/Engine/Plugins/Media/PixelStreaming/Source/PixelStreaming/Private/AudioInputMixer.cpp
@@ -3,6 +3,7 @@
 #include "AudioInputMixer.h"
 #include "Settings.h"
 #include "SampleBuffer.h"
+#include "PixelStreamingTrace.h"
 
 namespace UE::PixelStreaming
 {
@@ -110,6 +111,7 @@ namespace UE::PixelStreaming
 
 	void FMixerRunnable::Tick()
 	{
+		TRACE_CPUPROFILER_EVENT_SCOPE_ON_CHANNEL_STR("PixelStreaming FAudioInputMixer::FMixerRunnable::Tick", PixelStreamingChannel);
 		if (!Mixers.IsValid())
 		{
 			return;
@@ -159,6 +161,7 @@ namespace UE::PixelStreaming
 
 	void FAudioInput::PushAudio(const float* InBuffer, int32 NumSamples, int32 InNumChannels, int32 InSampleRate)
 	{
+		TRACE_CPUPROFILER_EVENT_SCOPE_ON_CHANNEL_STR("PixelStreaming FAudioInputMixer::PushAudio", PixelStreamingChannel);
 		// Todo(Luke): we could resample instead of throwing this assertion.
 		verifyf(SampleRate == InSampleRate, TEXT("Sample rate was mismatched."));
 
@@ -175,4 +178,4 @@ namespace UE::PixelStreaming
 		}
 	}
 
-} // namespace UE::PixelStreaming
\ No newline at end of file
+} // namespace UE::PixelStreaming
diff --git a/Engine/Plugins/Media/PixelStreaming/Source/PixelStreaming/Private/AudioPlayoutRequester.cpp b/Engine/Plugins/Media/PixelStreaming/Source/PixelStreaming/Private/AudioPlayoutRequester.cpp
--- a/Engine/Plugins/Media/PixelStreaming/Source/PixelStreaming/Private/AudioPlayoutRequester.cpp
+++ b/Engine/Plugins/Media/PixelStreaming/Source/PixelStreaming/Private/AudioPlayoutRequester.cpp
@@ -6,6 +6,7 @@
 #include "HAL/RunnableThread.h"
 #include "HAL/PlatformProcess.h"
 #include "WebRTCIncludes.h"
+#include "PixelStreamingTrace.h"
 
 namespace UE::PixelStreaming
 {
@@ -145,8 +146,11 @@ namespace UE::PixelStreaming
 			// Update request time to now seeing as the 10ms delta has elapsed
 			LastAudioRequestTimeMs = Now;
 
-			// Actually request playout
-			RequestPlayoutFunc();
+			{
+				// Actually request playout
+				TRACE_CPUPROFILER_EVENT_SCOPE_ON_CHANNEL_STR("PixelStreaming FAudioPlayoutRequester::Run RequestPlayoutFunc", PixelStreamingChannel);
+				RequestPlayoutFunc();
+			}
 		}
 
 		return 0;
diff --git a/Engine/Plugins/Media/PixelStreaming/Source/PixelStreaming/Private/AudioSink.cpp b/Engine/Plugins/Media/PixelStreaming/Source/PixelStreaming/Private/AudioSink.cpp
--- a/Engine/Plugins/Media/PixelStreaming/Source/PixelStreaming/Private/AudioSink.cpp
+++ b/Engine/Plugins/Media/PixelStreaming/Source/PixelStreaming/Private/AudioSink.cpp
@@ -1,6 +1,7 @@
 // Copyright Epic Games, Inc. All Rights Reserved.
 
 #include "AudioSink.h"
+#include "PixelStreamingTrace.h"
 
 namespace UE::PixelStreaming
 {
@@ -19,7 +20,7 @@ namespace UE::PixelStreaming
 
 	void FAudioSink::OnData(const void* audio_data, int bits_per_sample, int sample_rate, size_t number_of_channels, size_t number_of_frames, absl::optional<int64_t> absolute_capture_timestamp_ms)
 	{
-
+		TRACE_CPUPROFILER_EVENT_SCOPE_ON_CHANNEL_STR("PixelStreaming FAudioSink::OnData", PixelStreamingChannel);
 		// This data is populated from the internals of WebRTC, basically each audio track sent from the browser has its RTP audio source received and decoded.
 		// The sample rate and number of channels here has absolutely no relationship with PixelStreamingAudioDeviceModule.
 		// The sample rate and number of channels here is determined adaptively by WebRTC's NetEQ class that selects sample rate/number of channels
@@ -62,4 +63,4 @@ namespace UE::PixelStreaming
 	{
 		return AudioConsumers.Num() > 0;
 	}
-} // namespace UE::PixelStreaming
\ No newline at end of file
+} // namespace UE::PixelStreaming
diff --git a/Engine/Plugins/Media/PixelStreaming/Source/PixelStreaming/Private/FrameBufferI420.h b/Engine/Plugins/Media/PixelStreaming/Source/PixelStreaming/Private/FrameBufferI420.h
--- a/Engine/Plugins/Media/PixelStreaming/Source/PixelStreaming/Private/FrameBufferI420.h
+++ b/Engine/Plugins/Media/PixelStreaming/Source/PixelStreaming/Private/FrameBufferI420.h
@@ -3,7 +3,7 @@
 #pragma once
 
 #include "WebRTCIncludes.h"
-#include "PixelCaptureI420Buffer.h"
+#include "PixelCaptureBufferI420.h"
 
 namespace UE::PixelStreaming
 {
@@ -13,7 +13,7 @@ namespace UE::PixelStreaming
 	class WebRTCI420BufferWrapper : public webrtc::I420BufferInterface
 	{
 	public:
-		WebRTCI420BufferWrapper(TSharedPtr<FPixelCaptureI420Buffer> InCaptureBuffer)
+		WebRTCI420BufferWrapper(TSharedPtr<FPixelCaptureBufferI420> InCaptureBuffer)
 			: CaptureBuffer(InCaptureBuffer) {}
 		virtual ~WebRTCI420BufferWrapper() = default;
 		virtual int width() const override { return CaptureBuffer->GetWidth(); }
@@ -26,7 +26,7 @@ namespace UE::PixelStreaming
 		virtual int StrideV() const override { return CaptureBuffer->GetStrideUV(); }
 
 	private:
-		TSharedPtr<FPixelCaptureI420Buffer> CaptureBuffer;
+		TSharedPtr<FPixelCaptureBufferI420> CaptureBuffer;
 	};
 
 	/**
@@ -35,7 +35,7 @@ namespace UE::PixelStreaming
 	class FFrameBufferI420 : public webrtc::VideoFrameBuffer
 	{
 	public:
-		FFrameBufferI420(TSharedPtr<FPixelCaptureI420Buffer> Buffer)
+		FFrameBufferI420(TSharedPtr<FPixelCaptureBufferI420> Buffer)
 			: I420BufferWrapper(new rtc::RefCountedObject<WebRTCI420BufferWrapper>(Buffer)) {}
 
 		virtual ~FFrameBufferI420() = default;
diff --git a/Engine/Plugins/Media/PixelStreaming/Source/PixelStreaming/Private/FrameBufferRHI.h b/Engine/Plugins/Media/PixelStreaming/Source/PixelStreaming/Private/FrameBufferRHI.h
--- a/Engine/Plugins/Media/PixelStreaming/Source/PixelStreaming/Private/FrameBufferRHI.h
+++ b/Engine/Plugins/Media/PixelStreaming/Source/PixelStreaming/Private/FrameBufferRHI.h
@@ -2,27 +2,25 @@
 
 #pragma once
 
-#include "RHI.h"
 #include "WebRTCIncludes.h"
+#include "Video/Resources/VideoResourceRHI.h"
 
 namespace UE::PixelStreaming
 {
 	/**
-	 * A frame buffer for our NVENC encoders (h264). Contains a single RHI texture.
+	 * A frame buffer for RHI Video Resources.
 	 */
 	class FFrameBufferRHI : public webrtc::VideoFrameBuffer
 	{
 	public:
-		FFrameBufferRHI(FTexture2DRHIRef InFrameTexture)
-		:FrameTexture(InFrameTexture)
-		{
-		}
+		FFrameBufferRHI(TSharedPtr<FVideoResourceRHI> VideoResourceRHI) : VideoResourceRHI(VideoResourceRHI) {}
+
 
 		virtual ~FFrameBufferRHI() = default;
 
 		virtual webrtc::VideoFrameBuffer::Type type() const override { return webrtc::VideoFrameBuffer::Type::kNative; }
-		virtual int width() const override { return FrameTexture->GetDesc().Extent.X; }
-		virtual int height() const override { return FrameTexture->GetDesc().Extent.Y; }
+		virtual int width() const override { return VideoResourceRHI->GetDescriptor().Width; }
+		virtual int height() const override { return VideoResourceRHI->GetDescriptor().Height; }
 
 		virtual rtc::scoped_refptr<webrtc::I420BufferInterface> ToI420() override
 		{
@@ -36,9 +34,9 @@ namespace UE::PixelStreaming
 			return nullptr;
 		}
 
-		FTexture2DRHIRef GetFrameTexture() const { return FrameTexture; }
+		TSharedPtr<FVideoResourceRHI> GetVideoResource() { return VideoResourceRHI; }
 
 	private:
-		FTexture2DRHIRef FrameTexture;
+		TSharedPtr<FVideoResourceRHI> VideoResourceRHI;
 	};
 } // namespace UE::PixelStreaming
diff --git a/Engine/Plugins/Media/PixelStreaming/Source/PixelStreaming/Private/PixelStreamingDataChannel.cpp b/Engine/Plugins/Media/PixelStreaming/Source/PixelStreaming/Private/PixelStreamingDataChannel.cpp
--- a/Engine/Plugins/Media/PixelStreaming/Source/PixelStreaming/Private/PixelStreamingDataChannel.cpp
+++ b/Engine/Plugins/Media/PixelStreaming/Source/PixelStreaming/Private/PixelStreamingDataChannel.cpp
@@ -31,13 +31,9 @@ FPixelStreamingDataChannel::FPixelStreamingDataChannel(FPixelStreamingPeerConnec
 	SendConfig.negotiated = true;
 	SendConfig.id = SendStreamId;
 
-	#if WEBRTC_VERSION == 84
-		RecvChannel = SendChannel = PeerConnection->CreateDataChannel((SendStreamId == RecvStreamId) ? "datachannel" : "senddatachannel", &SendConfig);
-	#else
-		webrtc::RTCErrorOr<rtc::scoped_refptr<webrtc::DataChannelInterface>> SendRecvResult = PeerConnection->CreateDataChannelOrError((SendStreamId == RecvStreamId) ? "datachannel" : "senddatachannel", &SendConfig);
-		checkf(SendRecvResult.ok(), TEXT("Failed to create Data Channel. Msg=%s"), *FString(SendRecvResult.error().message()));
-		RecvChannel = SendChannel = SendRecvResult.MoveValue();
-	#endif
+	webrtc::RTCErrorOr<rtc::scoped_refptr<webrtc::DataChannelInterface>> SendRecvResult = PeerConnection->CreateDataChannelOrError((SendStreamId == RecvStreamId) ? "datachannel" : "senddatachannel", &SendConfig);
+	checkf(SendRecvResult.ok(), TEXT("Failed to create Data Channel. Msg=%s"), *FString(SendRecvResult.error().message()));
+	RecvChannel = SendChannel = SendRecvResult.MoveValue();
 
 	if (SendStreamId != RecvStreamId)
 	{
@@ -45,13 +41,9 @@ FPixelStreamingDataChannel::FPixelStreamingDataChannel(FPixelStreamingPeerConnec
 		RecvConfig.negotiated = true;
 		RecvConfig.id = RecvStreamId;
 
-		#if WEBRTC_VERSION == 84
-			RecvChannel = PeerConnection->CreateDataChannel("recvdatachannel", &RecvConfig);
-		#else
-			webrtc::RTCErrorOr<rtc::scoped_refptr<webrtc::DataChannelInterface>> RecvResult = PeerConnection->CreateDataChannelOrError("recvdatachannel", &RecvConfig);
-			checkf(RecvResult.ok(), TEXT("Failed to create Data Channel. Msg=%s"), *FString(RecvResult.error().message()));
-			RecvChannel = RecvResult.MoveValue();
-		#endif
+		webrtc::RTCErrorOr<rtc::scoped_refptr<webrtc::DataChannelInterface>> RecvResult = PeerConnection->CreateDataChannelOrError("recvdatachannel", &RecvConfig);
+		checkf(RecvResult.ok(), TEXT("Failed to create Data Channel. Msg=%s"), *FString(RecvResult.error().message()));
+		RecvChannel = RecvResult.MoveValue();
 	}
 
 	checkf(SendChannel, TEXT("Send channel cannot be null"));
diff --git a/Engine/Plugins/Media/PixelStreaming/Source/PixelStreaming/Private/PixelStreamingInputComponent.cpp b/Engine/Plugins/Media/PixelStreaming/Source/PixelStreaming/Private/PixelStreamingInputComponent.cpp
--- a/Engine/Plugins/Media/PixelStreaming/Source/PixelStreaming/Private/PixelStreamingInputComponent.cpp
+++ b/Engine/Plugins/Media/PixelStreaming/Source/PixelStreaming/Private/PixelStreamingInputComponent.cpp
@@ -10,7 +10,7 @@
 #include "Serialization/JsonSerializer.h"
 #include "GameFramework/GameUserSettings.h"
 #include "PixelStreamingPrivate.h"
-#include "PixelStreamingProtocol.h"
+#include "PixelStreamingInputProtocol.h"
 #include "IPixelStreamingStreamer.h"
 #include "Settings.h"
 #include "PixelStreamingUtils.h"
@@ -58,9 +58,8 @@ void UPixelStreamingInput::SendPixelStreamingResponse(const FString& Descriptor)
 {
 	if (PixelStreamingModule)
 	{
-		PixelStreamingModule->ForEachStreamer([&Descriptor, this](TSharedPtr<IPixelStreamingStreamer> Streamer)
-		{
-			Streamer->SendPlayerMessage(PixelStreamingModule->GetProtocol().FromStreamerProtocol.Find("Response")->Id, Descriptor);
+		PixelStreamingModule->ForEachStreamer([&Descriptor, this](TSharedPtr<IPixelStreamingStreamer> Streamer) {
+			Streamer->SendPlayerMessage(FPixelStreamingInputProtocol::FromStreamerProtocol.Find("Response")->GetID(), Descriptor);
 		});
 	}
 	else
diff --git a/Engine/Plugins/Media/PixelStreaming/Source/PixelStreaming/Private/PixelStreamingModule.cpp b/Engine/Plugins/Media/PixelStreaming/Source/PixelStreaming/Private/PixelStreamingModule.cpp
--- a/Engine/Plugins/Media/PixelStreaming/Source/PixelStreaming/Private/PixelStreamingModule.cpp
+++ b/Engine/Plugins/Media/PixelStreaming/Source/PixelStreaming/Private/PixelStreamingModule.cpp
@@ -1,6 +1,7 @@
 // Copyright Epic Games, Inc. All Rights Reserved.
 
 #include "PixelStreamingModule.h"
+#include "IPixelStreamingInputModule.h"
 #include "Streamer.h"
 #include "PixelStreamingInputComponent.h"
 #include "PixelStreamingDelegates.h"
@@ -13,6 +14,7 @@
 #include "UObject/UObjectIterator.h"
 #include "Engine/Texture2D.h"
 #include "Slate/SceneViewport.h"
+#include "PixelStreamingUtils.h"
 #include "Utils.h"
 #include "UtilsRender.h"
 
@@ -43,7 +45,6 @@ THIRD_PARTY_INCLUDES_END
 #include "Misc/MessageDialog.h"
 #include "Async/Async.h"
 #include "Engine/Engine.h"
-#include "VideoEncoderFactory.h"
 #include "VideoEncoderFactoryLayered.h"
 #include "WebRTCLogging.h"
 #include "WebSocketsModule.h"
@@ -56,9 +57,9 @@ THIRD_PARTY_INCLUDES_END
 #include "VideoSourceGroup.h"
 #include "PixelStreamingPeerConnection.h"
 #include "Engine/GameEngine.h"
-#include "PixelStreamingApplicationWrapper.h"
-#include "PixelStreamingInputHandler.h"
-#include "InputHandlers.h"
+#include "Stats.h"
+#include "Video/Resources/VideoResourceRHI.h"
+#include "PixelStreamingInputEnums.h"
 
 DEFINE_LOG_CATEGORY(LogPixelStreaming);
 
@@ -66,7 +67,7 @@ IPixelStreamingModule* UE::PixelStreaming::FPixelStreamingModule::PixelStreaming
 
 namespace UE::PixelStreaming
 {
-	typedef Protocol::EPixelStreamingMessageTypes EType;
+	typedef EPixelStreamingMessageTypes EType;
 	/**
 	 * IModuleInterface implementation
 	 */
@@ -87,54 +88,59 @@ namespace UE::PixelStreaming
 		}
 
 		const ERHIInterfaceType RHIType = GDynamicRHI ? RHIGetInterfaceType() : ERHIInterfaceType::Hidden;
-
-		IModularFeatures::Get().RegisterModularFeature(GetModularFeatureName(), this);
-		PopulateProtocol();
-		
 		// only D3D11/D3D12/Vulkan is supported
-		if (RHIType == ERHIInterfaceType::D3D11 || RHIType == ERHIInterfaceType::D3D12 || RHIType == ERHIInterfaceType::Vulkan)
+		if (!(RHIType == ERHIInterfaceType::D3D11 || RHIType == ERHIInterfaceType::D3D12 || RHIType == ERHIInterfaceType::Vulkan))
 		{
-			// By calling InitDefaultStreamer post engine init we can use pixel streaming in standalone editor mode
-			FCoreDelegates::OnFEngineLoopInitComplete.AddLambda([this]() {
-				// Check to see if we can use the Pixel Streaming plugin on this platform.
-				// If not then we avoid setting up our delegates to prevent access to the plugin.
-				if (!IsPlatformCompatible())
-				{
-					return;
-				}
+#if !WITH_DEV_AUTOMATION_TESTS
+			UE_LOG(LogPixelStreaming, Warning, TEXT("Only D3D11/D3D12/Vulkan Dynamic RHI is supported. Detected %s"), GDynamicRHI != nullptr ? GDynamicRHI->GetName() : TEXT("[null]"));
+#endif
+			return;
+		}
 
-				if (!ensure(GEngine != nullptr))
-				{
-					return;
-				}
+		PopulateProtocol();
+		RegisterCustomHandlers();
 
-				FApp::SetUnfocusedVolumeMultiplier(1.0f);
+		// By calling InitDefaultStreamer post engine init we can use pixel streaming in standalone editor mode
+		FCoreDelegates::OnFEngineLoopInitComplete.AddLambda([this, RHIType]() {
+			// Check to see if we can use the Pixel Streaming plugin on this platform.
+			// If not then we avoid setting up our delegates to prevent access to the plugin.
+			if (!IsPlatformCompatible())
+			{
+				return;
+			}
 
-				// Allow Pixel Streaming to broadcast to various delegates bound in the application-specific blueprint.
-				UPixelStreamingDelegates::CreateInstance();
+			if (!ensure(GEngine != nullptr))
+			{
+				return;
+			}
 
-				// Ensure we have ImageWrapper loaded, used in Freezeframes
-				verify(FModuleManager::Get().LoadModule(FName("ImageWrapper")));
-				InitDefaultStreamer();
-				bModuleReady = true;
-				ReadyEvent.Broadcast(*this);
-				// We don't want to start immediately streaming in editor
-				if (!GIsEditor)
-				{
-					StartStreaming();
-				}
-			});
-		}
-		else
-		{
-			#if !WITH_DEV_AUTOMATION_TESTS
-				UE_LOG(LogPixelStreaming, Warning, TEXT("Only D3D11/D3D12/Vulkan Dynamic RHI is supported. Detected %s"), GDynamicRHI != nullptr ? GDynamicRHI->GetName() : TEXT("[null]"));
-			#endif
-		}
+			// HACK (Luke): Until or if we ever find a workaround for fencing, we need to ensure capture always uses a fence
+			// if we don't then we get frequent and intermittent stuttering as textures are rendered to while being encoded.
+			// From testing NVENC + CUDA pathway seems acceptable without a fence in most cases so we use the faster, unsafer path there.
+			if (RHIType == ERHIInterfaceType::D3D11 || IsRHIDeviceAMD())
+			{
+				Settings::CVarPixelStreamingCaptureUseFence.AsVariable()->Set(true);
+			}
+
+			FApp::SetUnfocusedVolumeMultiplier(1.0f);
+
+			// Allow Pixel Streaming to broadcast to various delegates bound in the application-specific blueprint.
+			UPixelStreamingDelegates::CreateInstance();
+
+			// Ensure we have ImageWrapper loaded, used in Freezeframes
+			verify(FModuleManager::Get().LoadModule(FName("ImageWrapper")));
+			InitDefaultStreamer();
+			bModuleReady = true;
+			ReadyEvent.Broadcast(*this);
+			// We don't want to start immediately streaming in editor
+			if (!GIsEditor)
+			{
+				StartStreaming();
+			}
+		});
 
 		rtc::InitializeSSL();
 		RedirectWebRtcLogsToUnreal(rtc::LoggingSeverity::LS_VERBOSE);
-		FModuleManager::LoadModuleChecked<IModuleInterface>(TEXT("AVEncoder"));
 		FModuleManager::LoadModuleChecked<FWebSocketsModule>("WebSockets");
 
 		// ExternalVideoSourceGroup is used so that we can have a video source without a streamer
@@ -142,6 +148,8 @@ namespace UE::PixelStreaming
 		// ExternalVideoSourceGroup->SetVideoInput(FPixelStreamingVideoInputBackBuffer::Create());
 		// ExternalVideoSourceGroup->Start();
 
+		// Call FStats::Get() to initialize the singleton
+		FStats::Get();
 		bStartupCompleted = true;
 	}
 
@@ -159,7 +167,6 @@ namespace UE::PixelStreaming
 		FPixelStreamingPeerConnection::Shutdown();
 
 		rtc::CleanupSSL();
-		IModularFeatures::Get().UnregisterModularFeature(GetModularFeatureName(), this);
 
 		bStartupCompleted = false;
 	}
@@ -208,8 +215,7 @@ namespace UE::PixelStreaming
 	bool FPixelStreamingModule::StartStreaming()
 	{
 		bool bSuccess = true;
-		ForEachStreamer([&bSuccess](TSharedPtr<IPixelStreamingStreamer> Streamer)
-		{
+		ForEachStreamer([&bSuccess](TSharedPtr<IPixelStreamingStreamer> Streamer) {
 			if (Streamer.IsValid())
 			{
 				Streamer->StartStreaming();
@@ -225,8 +231,7 @@ namespace UE::PixelStreaming
 
 	void FPixelStreamingModule::StopStreaming()
 	{
-		ForEachStreamer([this](TSharedPtr<IPixelStreamingStreamer> Streamer)
-		{
+		ForEachStreamer([this](TSharedPtr<IPixelStreamingStreamer> Streamer) {
 			if (Streamer.IsValid())
 			{
 				Streamer->StopStreaming();
@@ -342,46 +347,6 @@ namespace UE::PixelStreaming
 		}
 	}
 
-	const Protocol::FPixelStreamingProtocol& FPixelStreamingModule::GetProtocol()
-	{
-		return MessageProtocol;
-	}
-
-	void FPixelStreamingModule::RegisterMessage(Protocol::EPixelStreamingMessageDirection MessageDirection, const FString& MessageType, Protocol::FPixelStreamingInputMessage Message, const TFunction<void(FMemoryReader)>& Handler)
-	{
-		if(MessageDirection == Protocol::EPixelStreamingMessageDirection::ToStreamer)
-		{
-			MessageProtocol.ToStreamerProtocol.Add(MessageType, Message);
-			ForEachStreamer([&Handler = Handler, &MessageType = MessageType](TSharedPtr<IPixelStreamingStreamer> Streamer)
-			{
-				TWeakPtr<IPixelStreamingInputHandler> WeakInputHandler = Streamer->GetInputHandler();
-				TSharedPtr<IPixelStreamingInputHandler> InputHandler = WeakInputHandler.Pin();
-				if (InputHandler)
-				{
-					InputHandler->RegisterMessageHandler(MessageType, Handler);
-				}
-			});
-		}
-		else if (MessageDirection == Protocol::EPixelStreamingMessageDirection::FromStreamer)
-		{
-			MessageProtocol.FromStreamerProtocol.Add(MessageType, Message);
-		}
-		OnProtocolUpdated.Broadcast();
-	}
-
-	TFunction<void(FMemoryReader)> FPixelStreamingModule::FindMessageHandler(const FString& MessageType)
-	{
-		// All streamers have the same protocol so we just use the first streamers input channel to get the message handler
-		TSharedPtr<IPixelStreamingStreamer> Streamer = Streamers.begin()->Value;
-		TWeakPtr<IPixelStreamingInputHandler> WeakInputHandler = Streamer->GetInputHandler();
-		if (TSharedPtr<IPixelStreamingInputHandler> InputHandler = WeakInputHandler.Pin())
-		{
-			return InputHandler->FindMessageHandler(MessageType);
-		}
-		// If the channel doesn't exist, just return an empty function
-		return ([](FMemoryReader Ar) {});
-	}
-
 	/**
 	 * End IPixelStreamingModule implementation
 	 */
@@ -405,6 +370,10 @@ namespace UE::PixelStreaming
 		}
 
 		TSharedPtr<IPixelStreamingStreamer> Streamer = CreateStreamer(Settings::GetDefaultStreamerID());
+		TSharedPtr<IPixelStreamingSignallingConnection> SignallingConnection = MakeShared<FPixelStreamingSignallingConnection>(Streamer->GetSignallingConnectionObserver().Pin(), Settings::GetDefaultStreamerID());
+		SignallingConnection->SetAutoReconnect(true);
+		Streamer->SetSignallingConnection(SignallingConnection);
+
 		// The PixelStreamingEditorModule handles setting video input in the editor
 		if (!GIsEditor)
 		{
@@ -428,7 +397,7 @@ namespace UE::PixelStreaming
 		{
 			// The user has specified a URL on the command line meaning their intention is to start streaming immediately
 			// in that case, set up the video input for them (as long as we're not in editor)
-			if(!GIsEditor)
+			if (!GIsEditor)
 			{
 				Streamer->SetVideoInput(FPixelStreamingVideoInputBackBuffer::Create());
 			}
@@ -453,94 +422,249 @@ namespace UE::PixelStreaming
 		}
 #endif
 
-		if (Settings::CVarPixelStreamingEncoderCodec.GetValueOnAnyThread() == "H264"
-			&& !AVEncoder::FVideoEncoderFactory::Get().HasEncoderForCodec(AVEncoder::ECodecType::H264))
+		if ((Settings::CVarPixelStreamingEncoderCodec.GetValueOnAnyThread() == "H264" && !FVideoEncoder::IsSupported<FVideoResourceRHI, FVideoEncoderConfigH264>())
+			|| (Settings::CVarPixelStreamingEncoderCodec.GetValueOnAnyThread() == "H265" && !FVideoEncoder::IsSupported<FVideoResourceRHI, FVideoEncoderConfigH265>()))
 		{
-			UE_LOG(LogPixelStreaming, Warning, TEXT("Could not setup hardware encoder for H.264. This is usually a driver issue, try reinstalling your drivers."));
+			UE_LOG(LogPixelStreaming, Warning, TEXT("Could not setup hardware encoder. This is usually a driver issue, try reinstalling your drivers."));
 			UE_LOG(LogPixelStreaming, Warning, TEXT("Falling back to VP8 software video encoding."));
 			Settings::CVarPixelStreamingEncoderCodec.AsVariable()->Set(TEXT("VP8"), ECVF_SetByCommandline);
 		}
 
 		return bCompatible;
 	}
-	/**
-	 * End own methods
-	 */
-
-	TSharedPtr<IInputDevice> FPixelStreamingModule::CreateInputDevice(const TSharedRef<FGenericApplicationMessageHandler>& InMessageHandler)
-	{
-		return MakeShared<FInputHandlers>(InMessageHandler);
-	}
 
 	void FPixelStreamingModule::PopulateProtocol()
 	{
-		using namespace Protocol;
-
 		// Old EToStreamerMsg Commands
 		/*
 		 * Control Messages.
 		 */
 		// Simple command with no payload
 		// Note, we only specify the ID when creating these messages to preserve backwards compatability
-		// when adding your own message type, you can simply do MessageProtocol.Direction.Add("XXX");
-		MessageProtocol.ToStreamerProtocol.Add("IFrameRequest", FPixelStreamingInputMessage(0, 0, {}));
-		MessageProtocol.ToStreamerProtocol.Add("RequestQualityControl", FPixelStreamingInputMessage(1, 0, {}));
-		MessageProtocol.ToStreamerProtocol.Add("FpsRequest", FPixelStreamingInputMessage(2, 0, {}));
-		MessageProtocol.ToStreamerProtocol.Add("AverageBitrateRequest", FPixelStreamingInputMessage(3, 0, {}));
-		MessageProtocol.ToStreamerProtocol.Add("StartStreaming", FPixelStreamingInputMessage(4, 0, {}));
-		MessageProtocol.ToStreamerProtocol.Add("StopStreaming", FPixelStreamingInputMessage(5, 0, {}));
-		MessageProtocol.ToStreamerProtocol.Add("LatencyTest", FPixelStreamingInputMessage(6, 0, {}));
-		MessageProtocol.ToStreamerProtocol.Add("RequestInitialSettings", FPixelStreamingInputMessage(7, 0, {}));
-		MessageProtocol.ToStreamerProtocol.Add("TestEcho", FPixelStreamingInputMessage(8, 0, {}));
+		// when adding your own message type, you can simply do FPixelStreamingInputProtocol.Direction.Add("XXX");
+		FPixelStreamingInputProtocol::ToStreamerProtocol.Add("IFrameRequest", FPixelStreamingInputMessage(0));
+		FPixelStreamingInputProtocol::ToStreamerProtocol.Add("RequestQualityControl", FPixelStreamingInputMessage(1));
+		FPixelStreamingInputProtocol::ToStreamerProtocol.Add("FpsRequest", FPixelStreamingInputMessage(2));
+		FPixelStreamingInputProtocol::ToStreamerProtocol.Add("AverageBitrateRequest", FPixelStreamingInputMessage(3));
+		FPixelStreamingInputProtocol::ToStreamerProtocol.Add("StartStreaming", FPixelStreamingInputMessage(4));
+		FPixelStreamingInputProtocol::ToStreamerProtocol.Add("StopStreaming", FPixelStreamingInputMessage(5));
+		FPixelStreamingInputProtocol::ToStreamerProtocol.Add("LatencyTest", FPixelStreamingInputMessage(6));
+		FPixelStreamingInputProtocol::ToStreamerProtocol.Add("RequestInitialSettings", FPixelStreamingInputMessage(7));
+		FPixelStreamingInputProtocol::ToStreamerProtocol.Add("TestEcho", FPixelStreamingInputMessage(8));
 
 		/*
 		 * Input Messages.
 		 */
 		// Generic Input Messages.
-		MessageProtocol.ToStreamerProtocol.Add("UIInteraction", FPixelStreamingInputMessage(50, 0, {}));
-		MessageProtocol.ToStreamerProtocol.Add("Command", FPixelStreamingInputMessage(51, 0, {}));
+		FPixelStreamingInputProtocol::ToStreamerProtocol.Add("UIInteraction", FPixelStreamingInputMessage(50));
+		FPixelStreamingInputProtocol::ToStreamerProtocol.Add("Command", FPixelStreamingInputMessage(51));
 
 		// Keyboard Input Message.
 		// Complex command with payload, therefore we specify the length of the payload (bytes) as well as the structure of the payload
-		MessageProtocol.ToStreamerProtocol.Add("KeyDown", FPixelStreamingInputMessage(60, 2, { EType::Uint8, EType::Uint8 }));
-		MessageProtocol.ToStreamerProtocol.Add("KeyUp", FPixelStreamingInputMessage(61, 1, { EType::Uint8 }));
-		MessageProtocol.ToStreamerProtocol.Add("KeyPress", FPixelStreamingInputMessage(62, 2, { EType::Uint16 }));
+		FPixelStreamingInputProtocol::ToStreamerProtocol.Add("KeyDown", FPixelStreamingInputMessage(60, { EType::Uint8, EType::Uint8 }));
+		FPixelStreamingInputProtocol::ToStreamerProtocol.Add("KeyUp", FPixelStreamingInputMessage(61, { EType::Uint8 }));
+		FPixelStreamingInputProtocol::ToStreamerProtocol.Add("KeyPress", FPixelStreamingInputMessage(62, { EType::Uint16 }));
 
 		// Mouse Input Messages.
-		MessageProtocol.ToStreamerProtocol.Add("MouseEnter", FPixelStreamingInputMessage(70));
-		MessageProtocol.ToStreamerProtocol.Add("MouseLeave", FPixelStreamingInputMessage(71));
-		MessageProtocol.ToStreamerProtocol.Add("MouseDown", FPixelStreamingInputMessage(72, 5, { EType::Uint8, EType::Uint16, EType::Uint16 }));
-		MessageProtocol.ToStreamerProtocol.Add("MouseUp", FPixelStreamingInputMessage(73, 5, { EType::Uint8, EType::Uint16, EType::Uint16 }));
-		MessageProtocol.ToStreamerProtocol.Add("MouseMove", FPixelStreamingInputMessage(74, 8, { EType::Uint16, EType::Uint16, EType::Uint16, EType::Uint16 }));
-		MessageProtocol.ToStreamerProtocol.Add("MouseWheel", FPixelStreamingInputMessage(75, 6, { EType::Int16, EType::Uint16, EType::Uint16 }));
-		MessageProtocol.ToStreamerProtocol.Add("MouseDouble", FPixelStreamingInputMessage(76, 5, { EType::Uint8, EType::Uint16, EType::Uint16 }));
+		FPixelStreamingInputProtocol::ToStreamerProtocol.Add("MouseEnter", FPixelStreamingInputMessage(70));
+		FPixelStreamingInputProtocol::ToStreamerProtocol.Add("MouseLeave", FPixelStreamingInputMessage(71));
+		FPixelStreamingInputProtocol::ToStreamerProtocol.Add("MouseDown", FPixelStreamingInputMessage(72, { EType::Uint8, EType::Uint16, EType::Uint16 }));
+		FPixelStreamingInputProtocol::ToStreamerProtocol.Add("MouseUp", FPixelStreamingInputMessage(73, { EType::Uint8, EType::Uint16, EType::Uint16 }));
+		FPixelStreamingInputProtocol::ToStreamerProtocol.Add("MouseMove", FPixelStreamingInputMessage(74, { EType::Uint16, EType::Uint16, EType::Uint16, EType::Uint16 }));
+		FPixelStreamingInputProtocol::ToStreamerProtocol.Add("MouseWheel", FPixelStreamingInputMessage(75, { EType::Int16, EType::Uint16, EType::Uint16 }));
+		FPixelStreamingInputProtocol::ToStreamerProtocol.Add("MouseDouble", FPixelStreamingInputMessage(76, { EType::Uint8, EType::Uint16, EType::Uint16 }));
 
 		// Touch Input Messages.
-		MessageProtocol.ToStreamerProtocol.Add("TouchStart", FPixelStreamingInputMessage(80, 8, { EType::Uint8, EType::Uint16, EType::Uint16, EType::Uint8, EType::Uint8, EType::Uint8}));
-		MessageProtocol.ToStreamerProtocol.Add("TouchEnd", FPixelStreamingInputMessage(81, 8, { EType::Uint8, EType::Uint16, EType::Uint16, EType::Uint8, EType::Uint8, EType::Uint8}));
-		MessageProtocol.ToStreamerProtocol.Add("TouchMove", FPixelStreamingInputMessage(82, 8, { EType::Uint8, EType::Uint16, EType::Uint16, EType::Uint8, EType::Uint8, EType::Uint8}));
+		FPixelStreamingInputProtocol::ToStreamerProtocol.Add("TouchStart", FPixelStreamingInputMessage(80, { EType::Uint8, EType::Uint16, EType::Uint16, EType::Uint8, EType::Uint8, EType::Uint8 }));
+		FPixelStreamingInputProtocol::ToStreamerProtocol.Add("TouchEnd", FPixelStreamingInputMessage(81, { EType::Uint8, EType::Uint16, EType::Uint16, EType::Uint8, EType::Uint8, EType::Uint8 }));
+		FPixelStreamingInputProtocol::ToStreamerProtocol.Add("TouchMove", FPixelStreamingInputMessage(82, { EType::Uint8, EType::Uint16, EType::Uint16, EType::Uint8, EType::Uint8, EType::Uint8 }));
 
 		// Gamepad Input Messages.
-		MessageProtocol.ToStreamerProtocol.Add("GamepadButtonPressed", FPixelStreamingInputMessage(90, 3, { EType::Uint8, EType::Uint8, EType::Uint8 }));
-		MessageProtocol.ToStreamerProtocol.Add("GamepadButtonReleased", FPixelStreamingInputMessage(91, 3, { EType::Uint8, EType::Uint8, EType::Uint8 }));
-		MessageProtocol.ToStreamerProtocol.Add("GamepadAnalog", FPixelStreamingInputMessage(92, 3, { EType::Uint8, EType::Uint8, EType::Double }));
+		FPixelStreamingInputProtocol::ToStreamerProtocol.Add("GamepadButtonPressed", FPixelStreamingInputMessage(90, { EType::Uint8, EType::Uint8, EType::Uint8 }));
+		FPixelStreamingInputProtocol::ToStreamerProtocol.Add("GamepadButtonReleased", FPixelStreamingInputMessage(91, { EType::Uint8, EType::Uint8, EType::Uint8 }));
+		FPixelStreamingInputProtocol::ToStreamerProtocol.Add("GamepadAnalog", FPixelStreamingInputMessage(92, { EType::Uint8, EType::Uint8, EType::Double }));
+
+		// XR Input Messages.
+		// clang-format off
+		FPixelStreamingInputProtocol::ToStreamerProtocol.Add("XRHMDTransform", FPixelStreamingInputMessage(110, {	// 4x4 Transform
+																													EType::Float, EType::Float, EType::Float, EType::Float,
+																													EType::Float, EType::Float, EType::Float, EType::Float,
+																													EType::Float, EType::Float, EType::Float, EType::Float,
+																													EType::Float, EType::Float, EType::Float, EType::Float,
+																												}));
+		
+		FPixelStreamingInputProtocol::ToStreamerProtocol.Add("XRControllerTransform", FPixelStreamingInputMessage(111, {// 4x4 Transform
+																														EType::Float, EType::Float, EType::Float, EType::Float, 
+																														EType::Float, EType::Float, EType::Float, EType::Float, 
+																														EType::Float, EType::Float, EType::Float, EType::Float, 
+																														EType::Float, EType::Float, EType::Float, EType::Float,
+																														// Handedness (L, R, Any)
+																														EType::Uint8 
+																														}));
+		
+		// clang-format on
+		FPixelStreamingInputProtocol::ToStreamerProtocol.Add("XRButtonPressed", FPixelStreamingInputMessage(112, // Handedness,   ButtonIdx,      IsRepeat
+																												{ EType::Uint8, EType::Uint8, EType::Uint8 }));
+		FPixelStreamingInputProtocol::ToStreamerProtocol.Add("XRButtonTouched", FPixelStreamingInputMessage(113, // Handedness,   ButtonIdx,      IsRepeat
+																												{ EType::Uint8, EType::Uint8, EType::Uint8 }));
+		FPixelStreamingInputProtocol::ToStreamerProtocol.Add("XRButtonReleased", FPixelStreamingInputMessage(114, // Handedness,   ButtonIdx,     IsRepeat
+																					 							{ EType::Uint8, EType::Uint8, EType::Uint8 }));
+																												
+		FPixelStreamingInputProtocol::ToStreamerProtocol.Add("XRAnalog", FPixelStreamingInputMessage(115, { EType::Uint8, EType::Uint8, EType::Double }));
+
+		FPixelStreamingInputProtocol::ToStreamerProtocol.Add("XRSystem", FPixelStreamingInputMessage(116, { EType::Uint8 }));
 
 		// Old EToPlayerMsg commands
-		MessageProtocol.FromStreamerProtocol.Add("QualityControlOwnership", FPixelStreamingInputMessage(0));
-		MessageProtocol.FromStreamerProtocol.Add("Response", FPixelStreamingInputMessage(1));
-		MessageProtocol.FromStreamerProtocol.Add("Command", FPixelStreamingInputMessage(2));
-		MessageProtocol.FromStreamerProtocol.Add("FreezeFrame", FPixelStreamingInputMessage(3));
-		MessageProtocol.FromStreamerProtocol.Add("UnfreezeFrame", FPixelStreamingInputMessage(4));
-		MessageProtocol.FromStreamerProtocol.Add("VideoEncoderAvgQP", FPixelStreamingInputMessage(5));
-		MessageProtocol.FromStreamerProtocol.Add("LatencyTest", FPixelStreamingInputMessage(6));
-		MessageProtocol.FromStreamerProtocol.Add("InitialSettings", FPixelStreamingInputMessage(7));
-		MessageProtocol.FromStreamerProtocol.Add("FileExtension", FPixelStreamingInputMessage(8));
-		MessageProtocol.FromStreamerProtocol.Add("FileMimeType", FPixelStreamingInputMessage(9));
-		MessageProtocol.FromStreamerProtocol.Add("FileContents", FPixelStreamingInputMessage(10));
-		MessageProtocol.FromStreamerProtocol.Add("TestEcho", FPixelStreamingInputMessage(11));
-		MessageProtocol.FromStreamerProtocol.Add("InputControlOwnership", FPixelStreamingInputMessage(12));
-		MessageProtocol.FromStreamerProtocol.Add("Protocol", FPixelStreamingInputMessage(255));
+		FPixelStreamingInputProtocol::FromStreamerProtocol.Add("QualityControlOwnership", FPixelStreamingInputMessage(0));
+		FPixelStreamingInputProtocol::FromStreamerProtocol.Add("Response", FPixelStreamingInputMessage(1));
+		FPixelStreamingInputProtocol::FromStreamerProtocol.Add("Command", FPixelStreamingInputMessage(2));
+		FPixelStreamingInputProtocol::FromStreamerProtocol.Add("FreezeFrame", FPixelStreamingInputMessage(3));
+		FPixelStreamingInputProtocol::FromStreamerProtocol.Add("UnfreezeFrame", FPixelStreamingInputMessage(4));
+		FPixelStreamingInputProtocol::FromStreamerProtocol.Add("VideoEncoderAvgQP", FPixelStreamingInputMessage(5));
+		FPixelStreamingInputProtocol::FromStreamerProtocol.Add("LatencyTest", FPixelStreamingInputMessage(6));
+		FPixelStreamingInputProtocol::FromStreamerProtocol.Add("InitialSettings", FPixelStreamingInputMessage(7));
+		FPixelStreamingInputProtocol::FromStreamerProtocol.Add("FileExtension", FPixelStreamingInputMessage(8));
+		FPixelStreamingInputProtocol::FromStreamerProtocol.Add("FileMimeType", FPixelStreamingInputMessage(9));
+		FPixelStreamingInputProtocol::FromStreamerProtocol.Add("FileContents", FPixelStreamingInputMessage(10));
+		FPixelStreamingInputProtocol::FromStreamerProtocol.Add("TestEcho", FPixelStreamingInputMessage(11));
+		FPixelStreamingInputProtocol::FromStreamerProtocol.Add("InputControlOwnership", FPixelStreamingInputMessage(12));
+		FPixelStreamingInputProtocol::FromStreamerProtocol.Add("Protocol", FPixelStreamingInputMessage(255));
+	}
+
+	void FPixelStreamingModule::RegisterCustomHandlers()
+	{
+		IPixelStreamingInputModule& InputModule = IPixelStreamingInputModule::Get();
+		InputModule.RegisterMessage(EPixelStreamingMessageDirection::ToStreamer,
+			"UIInteraction",
+			FPixelStreamingInputMessage(50),
+			[this](FMemoryReader Ar) { HandleUIInteraction(Ar); });
+
+		// The current handler is the function that will currently be executed if a message with type "Command" is received
+		TFunction<void(FMemoryReader)> BaseOnCommandHandler = InputModule.FindMessageHandler("Command");
+		// We then create our new handler which will execute the "base" handler
+		TFunction<void(FMemoryReader)> ExtendedOnCommandHandler = [this, BaseOnCommandHandler](FMemoryReader Ar) {
+			// and then perform out extended functionality after.
+			// equivalent to the super::DoSomeFunc pattern
+			BaseOnCommandHandler(Ar);
+
+			// In this case, the base handler handles raw console commands. Our extended functionality parses the commands
+			// for PS specific parameters
+			HandleOnCommand(Ar);
+		};
+
+		// Handle receiving commands from peers
+		InputModule.RegisterMessage(EPixelStreamingMessageDirection::ToStreamer, "Command", FPixelStreamingInputMessage(51), ExtendedOnCommandHandler);
+		// Handle sending commands to peers
+		InputModule.OnSendMessage.AddRaw(this, &UE::PixelStreaming::FPixelStreamingModule::HandleSendCommand);
+	}
+
+	void FPixelStreamingModule::HandleOnCommand(FMemoryReader Ar)
+	{
+		FString Res;
+		Res.GetCharArray().SetNumUninitialized(Ar.TotalSize() / 2 + 1);
+		Ar.Serialize(Res.GetCharArray().GetData(), Ar.TotalSize());
+		FString Descriptor = Res.Mid(1);
+		bool bSuccess = false;
+
+		/**
+		 * Encoder Settings
+		 */
+		FString MinQPString;
+		UE::PixelStreaming::ExtractJsonFromDescriptor(Descriptor, TEXT("Encoder.MinQP"), MinQPString, bSuccess);
+		if (bSuccess)
+		{
+			int MinQP = FCString::Atoi(*MinQPString);
+			UE::PixelStreaming::Settings::CVarPixelStreamingEncoderMinQP->Set(MinQP, ECVF_SetByCommandline);
+			return;
+		}
+
+		FString MaxQPString;
+		UE::PixelStreaming::ExtractJsonFromDescriptor(Descriptor, TEXT("Encoder.MaxQP"), MaxQPString, bSuccess);
+		if (bSuccess)
+		{
+			int MaxQP = FCString::Atoi(*MaxQPString);
+			UE::PixelStreaming::Settings::CVarPixelStreamingEncoderMaxQP->Set(MaxQP, ECVF_SetByCommandline);
+			return;
+		}
+
+		/**
+		 * WebRTC Settings
+		 */
+		FString FPSString;
+		UE::PixelStreaming::ExtractJsonFromDescriptor(Descriptor, TEXT("WebRTC.Fps"), FPSString, bSuccess);
+		if (bSuccess)
+		{
+			int FPS = FCString::Atoi(*FPSString);
+			UE::PixelStreaming::Settings::CVarPixelStreamingWebRTCFps->Set(FPS, ECVF_SetByCommandline);
+			return;
+		}
+
+		FString MinBitrateString;
+		UE::PixelStreaming::ExtractJsonFromDescriptor(Descriptor, TEXT("WebRTC.MinBitrate"), MinBitrateString, bSuccess);
+		if (bSuccess)
+		{
+			int MinBitrate = FCString::Atoi(*MinBitrateString);
+			UE::PixelStreaming::Settings::CVarPixelStreamingWebRTCMinBitrate->Set(MinBitrate, ECVF_SetByCommandline);
+			return;
+		}
+
+		FString MaxBitrateString;
+		UE::PixelStreaming::ExtractJsonFromDescriptor(Descriptor, TEXT("WebRTC.MaxBitrate"), MaxBitrateString, bSuccess);
+		if (bSuccess)
+		{
+			int MaxBitrate = FCString::Atoi(*MaxBitrateString);
+			UE::PixelStreaming::Settings::CVarPixelStreamingWebRTCMaxBitrate->Set(MaxBitrate, ECVF_SetByCommandline);
+			return;
+		}
 	}
+
+	void FPixelStreamingModule::HandleSendCommand(FMemoryReader Ar)
+	{
+		FString Descriptor;
+		Ar << Descriptor;
+		ForEachStreamer([&Descriptor, this](TSharedPtr<IPixelStreamingStreamer> Streamer) {
+			Streamer->SendPlayerMessage(FPixelStreamingInputProtocol::FromStreamerProtocol.Find("Command")->GetID(), Descriptor);
+		});
+	}
+
+	void FPixelStreamingModule::HandleUIInteraction(FMemoryReader Ar)
+	{
+		FString Res;
+		Res.GetCharArray().SetNumUninitialized(Ar.TotalSize() / 2 + 1);
+		Ar.Serialize(Res.GetCharArray().GetData(), Ar.TotalSize());
+
+		FString Descriptor = Res.Mid(1);
+
+		UE_LOG(LogPixelStreaming, Verbose, TEXT("UIInteraction: %s"), *Descriptor);
+		for (UPixelStreamingInput* InputComponent : InputComponents)
+		{
+			InputComponent->OnInputEvent.Broadcast(Descriptor);
+		}
+	}
+	/**
+	 * End own methods
+	 */
+
+	/**
+	 * Deprecated methods
+	 */
+	const FPixelStreamingInputProtocol FPixelStreamingModule::GetProtocol()
+	{
+		return FPixelStreamingInputProtocol();
+	}
+
+	void FPixelStreamingModule::RegisterMessage(EPixelStreamingMessageDirection MessageDirection, const FString& MessageType, FPixelStreamingInputMessage Message, const TFunction<void(FMemoryReader)>& Handler)
+	{
+		IPixelStreamingInputModule::Get().RegisterMessage(MessageDirection, MessageType, Message, Handler);
+	}
+
+	TFunction<void(FMemoryReader)> FPixelStreamingModule::FindMessageHandler(const FString& MessageType)
+	{
+		return IPixelStreamingInputModule::Get().FindMessageHandler(MessageType);
+	}
+	/**
+	 * End deprecated methods
+	 */
 } // namespace UE::PixelStreaming
 
 IMPLEMENT_MODULE(UE::PixelStreaming::FPixelStreamingModule, PixelStreaming)
diff --git a/Engine/Plugins/Media/PixelStreaming/Source/PixelStreaming/Private/PixelStreamingModule.h b/Engine/Plugins/Media/PixelStreaming/Source/PixelStreaming/Private/PixelStreamingModule.h
--- a/Engine/Plugins/Media/PixelStreaming/Source/PixelStreaming/Private/PixelStreamingModule.h
+++ b/Engine/Plugins/Media/PixelStreaming/Source/PixelStreaming/Private/PixelStreamingModule.h
@@ -5,7 +5,7 @@
 #include "IPixelStreamingModule.h"
 #include "RHI.h"
 #include "Tickable.h"
-#include "PixelStreamingProtocol.h"
+#include "PixelStreamingInputProtocol.h"
 
 class UPixelStreamingInput;
 class SWindow;
@@ -37,10 +37,7 @@ namespace UE::PixelStreaming
 		virtual TSharedPtr<IPixelStreamingStreamer> DeleteStreamer(const FString& StreamerId) override;
 		virtual FString GetDefaultStreamerID() override;
 		virtual FString GetDefaultSignallingURL() override;
-		virtual const Protocol::FPixelStreamingProtocol& GetProtocol() override;
-		
-		virtual void RegisterMessage(Protocol::EPixelStreamingMessageDirection MessageDirection, const FString& MessageType, Protocol::FPixelStreamingInputMessage Message, const TFunction<void(FMemoryReader)>& Handler) override;
-		virtual TFunction<void(FMemoryReader)> FindMessageHandler(const FString& MessageType) override;
+
 		// These are staying on the module at the moment as theres no way of the BPs knowing which streamer they are relevant to
 		virtual void AddInputComponent(UPixelStreamingInput* InInputComponent) override;
 		virtual void RemoveInputComponent(UPixelStreamingInput* InInputComponent) override;
@@ -59,14 +56,20 @@ namespace UE::PixelStreaming
 		void ShutdownModule() override;
 		/** End IModuleInterface implementation */
 
-		/** IInputDeviceModule implementation */
-		virtual TSharedPtr<IInputDevice> CreateInputDevice(const TSharedRef<FGenericApplicationMessageHandler>& InMessageHandler) override;
-		/** End IInputDeviceModule implementation */
+		/** Deprected methods from the IPixelStreamingModule. We're only keeping these here until we can fully remove them */
+		const FPixelStreamingInputProtocol GetProtocol() override;
+		void RegisterMessage(EPixelStreamingMessageDirection MessageDirection, const FString& MessageType, FPixelStreamingInputMessage Message, const TFunction<void(FMemoryReader)>& Handler) override;
+		TFunction<void(FMemoryReader)> FindMessageHandler(const FString& MessageType) override;
+		/** End deprecated methods */
 
 		// Own methods
 		void InitDefaultStreamer();
 		bool IsPlatformCompatible() const;
 		void PopulateProtocol();
+		void RegisterCustomHandlers();
+		void HandleOnCommand(FMemoryReader Ar);
+		void HandleSendCommand(FMemoryReader Ar);
+		void HandleUIInteraction(FMemoryReader Ar);
 
 	private:
 		bool bModuleReady = false;
@@ -78,6 +81,5 @@ namespace UE::PixelStreaming
 		TSharedPtr<FVideoSourceGroup> ExternalVideoSourceGroup;
 		mutable FCriticalSection StreamersCS;
 		TMap<FString, TSharedPtr<IPixelStreamingStreamer>> Streamers;
-		Protocol::FPixelStreamingProtocol MessageProtocol;
 	};
 } // namespace UE::PixelStreaming
diff --git a/Engine/Plugins/Media/PixelStreaming/Source/PixelStreaming/Private/PixelStreamingPeerConnection.cpp b/Engine/Plugins/Media/PixelStreaming/Source/PixelStreaming/Private/PixelStreamingPeerConnection.cpp
--- a/Engine/Plugins/Media/PixelStreaming/Source/PixelStreaming/Private/PixelStreamingPeerConnection.cpp
+++ b/Engine/Plugins/Media/PixelStreaming/Source/PixelStreaming/Private/PixelStreamingPeerConnection.cpp
@@ -9,10 +9,10 @@
 #include "AudioCapturer.h"
 #include "PixelStreamingAudioDeviceModule.h"
 #include "ToStringExtensions.h"
-#include "PixelStreamingProtocol.h"
 #include "PixelStreamingDataChannel.h"
 #include "Stats.h"
 #include "AudioInputMixer.h"
+#include "absl/strings/match.h"
 
 namespace
 {
@@ -82,39 +82,79 @@ namespace
 		return bSyncVideoAndAudio ? "pixelstreaming_av_stream_id" : "pixelstreaming_video_stream_id";
 	}
 
-	void MungeLocalSDP(cricket::SessionDescription* SessionDescription)
+	EPixelStreamingCodec ExtractVideoCodec(cricket::SessionDescription* SessionDescription)
 	{
 		std::vector<cricket::ContentInfo>& ContentInfos = SessionDescription->contents();
 		for (cricket::ContentInfo& ContentInfo : ContentInfos)
 		{
 			cricket::MediaContentDescription* MediaDescription = ContentInfo.media_description();
 			cricket::MediaType MediaType = MediaDescription->type();
-			if (MediaType != cricket::MediaType::MEDIA_TYPE_AUDIO)
+
+			if (MediaDescription->type() == cricket::MediaType::MEDIA_TYPE_VIDEO)
 			{
-				continue;
+				cricket::VideoContentDescription* VideoDescription = MediaDescription->as_video();
+				std::vector<cricket::VideoCodec> CodecsCopy = VideoDescription->codecs();
+				for (cricket::VideoCodec& Codec : CodecsCopy)
+				{
+					if (absl::EqualsIgnoreCase(Codec.name, cricket::kVp8CodecName))
+					{
+						return EPixelStreamingCodec::VP8;
+					}
+					else if (absl::EqualsIgnoreCase(Codec.name, cricket::kVp9CodecName))
+					{
+						return EPixelStreamingCodec::VP9;
+					}
+					else if (absl::EqualsIgnoreCase(Codec.name, cricket::kH265CodecName))
+					{
+						return EPixelStreamingCodec::H265;
+					}
+					else if (absl::EqualsIgnoreCase(Codec.name, cricket::kH264CodecName))
+					{
+						return EPixelStreamingCodec::H264;
+					}
+				}
 			}
-			cricket::AudioContentDescription* AudioDescription = MediaDescription->as_audio();
-			if (AudioDescription == nullptr)
+		}
+		return EPixelStreamingCodec::Invalid;
+	}
+
+	void MungeLocalSDP(cricket::SessionDescription* SessionDescription)
+	{
+		std::vector<cricket::ContentInfo>& ContentInfos = SessionDescription->contents();
+		for (cricket::ContentInfo& ContentInfo : ContentInfos)
+		{
+			cricket::MediaContentDescription* MediaDescription = ContentInfo.media_description();
+			cricket::MediaType MediaType = MediaDescription->type();
+
+			if (MediaDescription->type() == cricket::MediaType::MEDIA_TYPE_VIDEO)
 			{
-				continue;
+				// todo: munge local video sdp if required
 			}
-			std::vector<cricket::AudioCodec> CodecsCopy = AudioDescription->codecs();
-			for (cricket::AudioCodec& Codec : CodecsCopy)
+			else if (MediaDescription->type() == cricket::MediaType::MEDIA_TYPE_AUDIO)
 			{
-				if (Codec.name == "opus")
+				cricket::AudioContentDescription* AudioDescription = MediaDescription->as_audio();
+				if (AudioDescription == nullptr)
+				{
+					continue;
+				}
+				std::vector<cricket::AudioCodec> CodecsCopy = AudioDescription->codecs();
+				for (cricket::AudioCodec& Codec : CodecsCopy)
 				{
-					Codec.SetParam(cricket::kCodecParamPTime, "20");
-					Codec.SetParam(cricket::kCodecParamMaxPTime, "120");
-					Codec.SetParam(cricket::kCodecParamMinPTime, "3");
-					Codec.SetParam(cricket::kCodecParamSPropStereo, "1");
-					Codec.SetParam(cricket::kCodecParamStereo, "1");
-					Codec.SetParam(cricket::kCodecParamUseInbandFec, "1");
-					Codec.SetParam(cricket::kCodecParamUseDtx, "0");
-					Codec.SetParam(cricket::kCodecParamMaxAverageBitrate, "510000");
-					Codec.SetParam(cricket::kCodecParamMaxPlaybackRate, "48000");
+					if (Codec.name == "opus")
+					{
+						Codec.SetParam(cricket::kCodecParamPTime, "20");
+						Codec.SetParam(cricket::kCodecParamMaxPTime, "120");
+						Codec.SetParam(cricket::kCodecParamMinPTime, "3");
+						Codec.SetParam(cricket::kCodecParamSPropStereo, "1");
+						Codec.SetParam(cricket::kCodecParamStereo, "1");
+						Codec.SetParam(cricket::kCodecParamUseInbandFec, "1");
+						Codec.SetParam(cricket::kCodecParamUseDtx, "0");
+						Codec.SetParam(cricket::kCodecParamMaxAverageBitrate, "510000");
+						Codec.SetParam(cricket::kCodecParamMaxPlaybackRate, "48000");
+					}
 				}
+				AudioDescription->set_codecs(CodecsCopy);
 			}
-			AudioDescription->set_codecs(CodecsCopy);
 		}
 	}
 
@@ -128,6 +168,7 @@ namespace
 			if (MediaDescription->type() == cricket::MediaType::MEDIA_TYPE_VIDEO)
 			{
 				cricket::VideoContentDescription* VideoDescription = MediaDescription->as_video();
+
 				std::vector<cricket::VideoCodec> CodecsCopy = VideoDescription->codecs();
 				for (cricket::VideoCodec& Codec : CodecsCopy)
 				{
@@ -143,12 +184,8 @@ namespace
 
 	void SetTransceiverDirection(webrtc::RtpTransceiverInterface& InTransceiver, webrtc::RtpTransceiverDirection InDirection)
 	{
-		#if WEBRTC_VERSION == 84
-			InTransceiver.SetDirection(InDirection);
-		#else
-			webrtc::RTCError Result = InTransceiver.SetDirectionWithError(InDirection);
-			checkf(Result.ok(), TEXT("Failed to add Video transceiver to PeerConnection. Msg=%s"), *FString(Result.message()));
-		#endif
+		webrtc::RTCError Result = InTransceiver.SetDirectionWithError(InDirection);
+		checkf(Result.ok(), TEXT("Failed to add Video transceiver to PeerConnection. Msg=%s"), *FString(Result.message()));
 	}
 } // namespace
 
@@ -176,7 +213,16 @@ public:
 			{
 				if (Transceiver->media_type() == cricket::MediaType::MEDIA_TYPE_VIDEO)
 				{
-					PeerConnection->GetStats(Transceiver->sender(), WebRTCStatsCallback);
+					const webrtc::RtpTransceiverDirection TransceiverDirection = Transceiver->direction();
+					if (TransceiverDirection == webrtc::RtpTransceiverDirection::kSendRecv || TransceiverDirection == webrtc::RtpTransceiverDirection::kSendOnly)
+					{
+						PeerConnection->GetStats(Transceiver->sender(), WebRTCStatsCallback);
+					}
+
+					if (TransceiverDirection == webrtc::RtpTransceiverDirection::kSendRecv || TransceiverDirection == webrtc::RtpTransceiverDirection::kRecvOnly)
+					{
+						PeerConnection->GetStats(Transceiver->receiver(), WebRTCStatsCallback);
+					}
 				}
 			}
 		}
@@ -238,13 +284,9 @@ TUniquePtr<FPixelStreamingPeerConnection> FPixelStreamingPeerConnection::Create(
 	TUniquePtr<FPixelStreamingPeerConnection> NewPeerConnection = TUniquePtr<FPixelStreamingPeerConnection>(new FPixelStreamingPeerConnection());
 	NewPeerConnection->IsSFU = IsSFU;
 
-	#if WEBRTC_VERSION == 84
-		NewPeerConnection->PeerConnection = PeerConnectionFactory->CreatePeerConnection(RTCConfig, nullptr, nullptr, NewPeerConnection.Get());
-	#else
-		webrtc::RTCErrorOr<rtc::scoped_refptr<webrtc::PeerConnectionInterface>> Result = PeerConnectionFactory->CreatePeerConnectionOrError(RTCConfig, webrtc::PeerConnectionDependencies(NewPeerConnection.Get()));
-		checkf(Result.ok(), TEXT("Failed to create Peer Connection. Msg=%s"), *FString(Result.error().message()));
-		NewPeerConnection->PeerConnection = Result.MoveValue();
-	#endif
+	webrtc::RTCErrorOr<rtc::scoped_refptr<webrtc::PeerConnectionInterface>> Result = PeerConnectionFactory->CreatePeerConnectionOrError(RTCConfig, webrtc::PeerConnectionDependencies(NewPeerConnection.Get()));
+	checkf(Result.ok(), TEXT("Failed to create Peer Connection. Msg=%s"), *FString(Result.error().message()));
+	NewPeerConnection->PeerConnection = Result.MoveValue();
 
 	// Setup suggested bitrate settings on the Peer Connection based on our CVars
 	webrtc::BitrateSettings BitrateSettings;
@@ -336,6 +378,10 @@ void FPixelStreamingPeerConnection::SetLocalDescription(webrtc::SessionDescripti
 
 	FPixelStreamingSetSessionDescriptionObserver* SetLocalDescriptionObserver = FPixelStreamingSetSessionDescriptionObserver::Create(OnSetLocalDescriptionSuccess, OnSetLocalDescriptionFail);
 	MungeLocalSDP(SDP->description());
+
+	// Get the negotiated codec from the SDP
+	NegotiatedVideoCodec = ExtractVideoCodec(SDP->description());
+
 	PeerConnection->SetLocalDescription(SetLocalDescriptionObserver, SDP);
 }
 
@@ -358,6 +404,10 @@ void FPixelStreamingPeerConnection::SetRemoteDescription(webrtc::SessionDescript
 
 	FPixelStreamingSetSessionDescriptionObserver* SetRemoteDescriptionObserver = FPixelStreamingSetSessionDescriptionObserver::Create(OnSetRemoteDescriptionSuccess, OnSetRemoteDescriptionFail);
 	MungeRemoteSDP(SDP->description());
+
+	// Get the negotiated codec from the SDP
+	NegotiatedVideoCodec = ExtractVideoCodec(SDP->description());
+
 	PeerConnection->SetRemoteDescription(SetRemoteDescriptionObserver, SDP);
 }
 
@@ -371,6 +421,11 @@ const webrtc::SessionDescriptionInterface* FPixelStreamingPeerConnection::GetRem
 	return PeerConnection->remote_description();
 }
 
+EPixelStreamingCodec FPixelStreamingPeerConnection::GetNegotiatedVideoCodec() const
+{
+	return NegotiatedVideoCodec;
+}
+
 void FPixelStreamingPeerConnection::SetVideoSource(rtc::scoped_refptr<webrtc::VideoTrackSourceInterface> InVideoSource)
 {
 	const static FString VideoTrackLabel = TEXT("pixelstreaming_video_track_label");
@@ -525,6 +580,9 @@ void FPixelStreamingPeerConnection::SetAudioSink(TSharedPtr<IPixelStreamingAudio
 					if (rtc::scoped_refptr<webrtc::MediaStreamTrackInterface> Track = Receiver->track())
 					{
 						webrtc::AudioTrackInterface* AudioTrack = static_cast<webrtc::AudioTrackInterface*>(Track.get());
+						// There's no mechanism to check if this sink already exists on the track, so attempt to remove it first to make sure we
+						// don't accidentally add it twice (note it is always safe to call remove because it is a no-op if the sink is not found).
+						AudioTrack->RemoveSink(AudioSink.Get());
 						AudioTrack->AddSink(AudioSink.Get());
 					}
 				}
@@ -559,13 +617,9 @@ TSharedPtr<FPixelStreamingDataChannel> FPixelStreamingPeerConnection::CreateData
 	DataChannelConfig.negotiated = Negotiated;
 	DataChannelConfig.id = Id;
 
-	#if WEBRTC_VERSION == 84
-		return FPixelStreamingDataChannel::Create(PeerConnection->CreateDataChannel("datachannel", &DataChannelConfig));
-	#else
-		webrtc::RTCErrorOr<rtc::scoped_refptr<webrtc::DataChannelInterface>> Result = PeerConnection->CreateDataChannelOrError("datachannel", &DataChannelConfig);
-		checkf(Result.ok(), TEXT("Failed to create Data Channel. Msg=%s"), *FString(Result.error().message()));
-		return FPixelStreamingDataChannel::Create(Result.MoveValue());
-	#endif
+	webrtc::RTCErrorOr<rtc::scoped_refptr<webrtc::DataChannelInterface>> Result = PeerConnection->CreateDataChannelOrError("datachannel", &DataChannelConfig);
+	checkf(Result.ok(), TEXT("Failed to create Data Channel. Msg=%s"), *FString(Result.error().message()));
+	return FPixelStreamingDataChannel::Create(Result.MoveValue());
 }
 
 void FPixelStreamingPeerConnection::SetWebRTCStatsCallback(rtc::scoped_refptr<webrtc::RTCStatsCollectorCallback> InCallback)
@@ -697,6 +751,10 @@ void FPixelStreamingPeerConnection::OnTrack(rtc::scoped_refptr<webrtc::RtpTransc
 	else if (AudioSink && Transceiver->media_type() == cricket::MediaType::MEDIA_TYPE_AUDIO)
 	{
 		webrtc::AudioTrackInterface* AudioTrack = static_cast<webrtc::AudioTrackInterface*>(Transceiver->receiver()->track().get());
+		// There's no mechanism to check if this sink already exists on the track, so attempt to remove it first to make sure we
+		// don't accidentally add it twice (note it is always safe to call remove because it is a no-op if the sink is not found).
+		// To future maintainers: yes, this code exsists in two spots, this is intentional.
+		AudioTrack->RemoveSink(AudioSink.Get());
 		AudioTrack->AddSink(AudioSink.Get());
 	}
 
@@ -775,11 +833,7 @@ void FPixelStreamingPeerConnection::CreatePeerConnectionFactory()
 {
 	using namespace UE::PixelStreaming;
 
-#if WEBRTC_VERSION == 84
-	SignallingThread = MakeUnique<rtc::Thread>(rtc::SocketServer::CreateDefault());
-#elif WEBRTC_VERSION == 96
 	SignallingThread = MakeUnique<rtc::Thread>(rtc::CreateDefaultSocketServer());
-#endif
 	SignallingThread->SetName("FPixelStreamingPeerConnection SignallingThread", nullptr);
 	SignallingThread->Start();
 
@@ -792,7 +846,7 @@ void FPixelStreamingPeerConnection::CreatePeerConnectionFactory()
 		else
 		{
 			// If experimental audio input is enabled we pass the mixer, if not we don't - no mixer means only use the UE submix.
-			if (Settings::IsExperimentalAudioInputEnabled())
+			if (Settings::CVarPixelStreamingExperimentalAudioInput.GetValueOnAnyThread())
 			{
 				AudioDeviceModule = new rtc::RefCountedObject<FAudioDeviceModule>(AudioMixer);
 				UE_LOG(LogPixelStreaming, Log, TEXT("Using -PixelStreamingExperimentalAudioInput. Pixel Streaming audio will mix the UE submix with user audio inputs."));
diff --git a/Engine/Plugins/Media/PixelStreaming/Source/PixelStreaming/Private/PixelStreamingSignallingConnection.cpp b/Engine/Plugins/Media/PixelStreaming/Source/PixelStreaming/Private/PixelStreamingSignallingConnection.cpp
--- a/Engine/Plugins/Media/PixelStreaming/Source/PixelStreaming/Private/PixelStreamingSignallingConnection.cpp
+++ b/Engine/Plugins/Media/PixelStreaming/Source/PixelStreaming/Private/PixelStreamingSignallingConnection.cpp
@@ -9,16 +9,14 @@
 #include "Settings.h"
 #include "TimerManager.h"
 #include "PixelStreamingDelegates.h"
-#include "PixelStreamingProtocol.h"
 #include "Utils.h"
 #include "ToStringExtensions.h"
 
 DECLARE_LOG_CATEGORY_EXTERN(LogPixelStreamingSS, Log, VeryVerbose);
 DEFINE_LOG_CATEGORY(LogPixelStreamingSS);
 
-FPixelStreamingSignallingConnection::FPixelStreamingSignallingConnection(const FWebSocketFactory& InWebSocketFactory, IPixelStreamingSignallingConnectionObserver& InObserver, FString InStreamerId)
-	: WebSocketFactory(InWebSocketFactory)
-	, Observer(InObserver)
+FPixelStreamingSignallingConnection::FPixelStreamingSignallingConnection(TSharedPtr<IPixelStreamingSignallingConnectionObserver> InObserver, FString InStreamerId)
+	: Observer(InObserver)
 	, StreamerId(InStreamerId)
 {
 	RegisterHandler("identify", [this](FJsonObjectPtr JsonMsg) { OnIdRequested(); });
@@ -58,7 +56,8 @@ void FPixelStreamingSignallingConnection::Connect(FString InUrl, bool bIsReconne
 	Disconnect();
 
 	Url = InUrl;
-	WebSocket = WebSocketFactory(Url);
+
+	WebSocket = FWebSocketsModule::Get().CreateWebSocket(Url, TEXT(""));
 	verifyf(WebSocket, TEXT("Web Socket Factory failed to return a valid Web Socket."));
 
 	OnConnectedHandle = WebSocket->OnConnected().AddLambda([this]() { OnConnected(); });
@@ -124,6 +123,11 @@ void FPixelStreamingSignallingConnection::SendOffer(FPixelStreamingPlayerId Play
 	FString SdpStr = UE::PixelStreaming::ToString(SdpAnsi);
 	OfferJson->SetStringField(TEXT("sdp"), SdpStr);
 
+#if WITH_EDITOR
+	// Default hover preference needs to be sent with the offer to prevent error messages on older infrastructure
+	OfferJson->SetStringField(TEXT("defaultToHover"), TEXT("true"));
+#endif // WITH_EDITOR
+
 	UE_LOG(LogPixelStreamingSS, Log, TEXT("Sending player=%s \"offer\" to SS %s"), *PlayerId, *Url);
 	UE_LOG(LogPixelStreamingSS, Verbose, TEXT("SDP offer\n%s"), *SdpStr);
 
@@ -288,7 +292,7 @@ void FPixelStreamingSignallingConnection::OnConnected()
 	// No need to to do reconnect now that we are connected
 	StopReconnectTimer();
 
-	Observer.OnSignallingConnected();
+	Observer->OnSignallingConnected();
 
 	if (bKeepAliveEnabled)
 	{
@@ -315,7 +319,7 @@ void FPixelStreamingSignallingConnection::OnConnectionError(const FString& Error
 {
 	UE_LOG(LogPixelStreamingSS, Log, TEXT("Failed to connect to SS %s - signalling server may not be up yet. Message: \"%s\""), *Url, *Error);
 
-	Observer.OnSignallingError(Error);
+	Observer->OnSignallingError(Error);
 
 	StopKeepAliveTimer();
 
@@ -335,7 +339,7 @@ void FPixelStreamingSignallingConnection::OnClosed(int32 StatusCode, const FStri
 {
 	UE_LOG(LogPixelStreamingSS, Log, TEXT("Closed connection to SS %s - \n\tstatus %d\n\treason: %s\n\twas clean: %s"), *Url, StatusCode, *Reason, bWasClean ? TEXT("true") : TEXT("false"));
 
-	Observer.OnSignallingDisconnected(StatusCode, Reason, bWasClean);
+	Observer->OnSignallingDisconnected(StatusCode, Reason, bWasClean);
 
 	StopKeepAliveTimer();
 
@@ -472,7 +476,7 @@ void FPixelStreamingSignallingConnection::OnConfig(const FJsonObjectPtr& Json)
 	// force `UnifiedPlan` as we control both ends of WebRTC streaming
 	RTCConfig.sdp_semantics = webrtc::SdpSemantics::kUnifiedPlan;
 
-	Observer.OnSignallingConfig(RTCConfig);
+	Observer->OnSignallingConfig(RTCConfig);
 }
 
 void FPixelStreamingSignallingConnection::OnSessionDescription(const FJsonObjectPtr& Json)
@@ -490,11 +494,11 @@ void FPixelStreamingSignallingConnection::OnSessionDescription(const FJsonObject
 	bool bGotPlayerId = GetPlayerIdJson(Json, PlayerId);
 	if (!bGotPlayerId)
 	{
-		Observer.OnSignallingSessionDescription(Type, Sdp);
+		Observer->OnSignallingSessionDescription(Type, Sdp);
 	}
 	else
 	{
-		Observer.OnSignallingSessionDescription(PlayerId, Type, Sdp);
+		Observer->OnSignallingSessionDescription(PlayerId, Type, Sdp);
 	}
 }
 
@@ -529,11 +533,11 @@ void FPixelStreamingSignallingConnection::OnIceCandidate(const FJsonObjectPtr& J
 
 	if (bGotPlayerId)
 	{
-		Observer.OnSignallingRemoteIceCandidate(PlayerId, SdpMid, SdpMLineIndex, CandidateStr);
+		Observer->OnSignallingRemoteIceCandidate(PlayerId, SdpMid, SdpMLineIndex, CandidateStr);
 	}
 	else
 	{
-		Observer.OnSignallingRemoteIceCandidate(SdpMid, SdpMLineIndex, CandidateStr);
+		Observer->OnSignallingRemoteIceCandidate(SdpMid, SdpMLineIndex, CandidateStr);
 	}
 }
 
@@ -546,7 +550,7 @@ void FPixelStreamingSignallingConnection::OnPlayerCount(const FJsonObjectPtr& Js
 		return;
 	}
 
-	Observer.OnSignallingPlayerCount(Count);
+	Observer->OnSignallingPlayerCount(Count);
 }
 
 void FPixelStreamingSignallingConnection::OnPlayerConnected(const FJsonObjectPtr& Json)
@@ -571,7 +575,7 @@ void FPixelStreamingSignallingConnection::OnPlayerConnected(const FJsonObjectPtr
 	bool bIsSFU = false;
 	Json->TryGetBoolField(TEXT("sfu"), PlayerConfig.IsSFU);
 
-	Observer.OnSignallingPlayerConnected(PlayerId, PlayerConfig);
+	Observer->OnSignallingPlayerConnected(PlayerId, PlayerConfig);
 }
 
 void FPixelStreamingSignallingConnection::OnPlayerDisconnected(const FJsonObjectPtr& Json)
@@ -584,7 +588,7 @@ void FPixelStreamingSignallingConnection::OnPlayerDisconnected(const FJsonObject
 		return;
 	}
 
-	Observer.OnSignallingPlayerDisconnected(PlayerId);
+	Observer->OnSignallingPlayerDisconnected(PlayerId);
 }
 
 void FPixelStreamingSignallingConnection::OnSFUPeerDataChannels(const FJsonObjectPtr& Json)
@@ -621,7 +625,7 @@ void FPixelStreamingSignallingConnection::OnSFUPeerDataChannels(const FJsonObjec
 		return;
 	}
 
-	Observer.OnSignallingSFUPeerDataChannels(SFUId, PlayerId, SendStreamId, RecvStreamId);
+	Observer->OnSignallingSFUPeerDataChannels(SFUId, PlayerId, SendStreamId, RecvStreamId);
 }
 
 void FPixelStreamingSignallingConnection::OnPeerDataChannels(const FJsonObjectPtr& Json)
@@ -638,7 +642,7 @@ void FPixelStreamingSignallingConnection::OnPeerDataChannels(const FJsonObjectPt
 		UE_LOG(LogPixelStreamingSS, Error, TEXT("Failed to get `recvStreamId` from remote `peerDataChannels` message\n%s"), *UE::PixelStreaming::ToString(Json));
 		return;
 	}
-	Observer.OnSignallingPeerDataChannels(SendStreamId, RecvStreamId);
+	Observer->OnSignallingPeerDataChannels(SendStreamId, RecvStreamId);
 }
 
 void FPixelStreamingSignallingConnection::SetPlayerIdJson(FJsonObjectPtr& JsonObject, FPixelStreamingPlayerId PlayerId)
diff --git a/Engine/Plugins/Media/PixelStreaming/Source/PixelStreaming/Private/PixelStreamingTrace.cpp b/Engine/Plugins/Media/PixelStreaming/Source/PixelStreaming/Private/PixelStreamingTrace.cpp
new file mode 100644
--- /dev/null
+++ b/Engine/Plugins/Media/PixelStreaming/Source/PixelStreaming/Private/PixelStreamingTrace.cpp
@@ -0,0 +1,7 @@
+// Copyright Epic Games, Inc. All Rights Reserved.
+
+#include "PixelStreamingTrace.h"
+
+#include "Trace/Trace.inl"
+
+UE_TRACE_CHANNEL_DEFINE(PixelStreamingChannel);
diff --git a/Engine/Plugins/Media/PixelStreaming/Source/PixelStreaming/Private/PixelStreamingVideoInput.cpp b/Engine/Plugins/Media/PixelStreaming/Source/PixelStreaming/Private/PixelStreamingVideoInput.cpp
--- a/Engine/Plugins/Media/PixelStreaming/Source/PixelStreaming/Private/PixelStreamingVideoInput.cpp
+++ b/Engine/Plugins/Media/PixelStreaming/Source/PixelStreaming/Private/PixelStreamingVideoInput.cpp
@@ -3,6 +3,8 @@
 #include "PixelStreamingVideoInput.h"
 #include "Settings.h"
 #include "FrameBufferMultiFormat.h"
+#include "PixelCaptureBufferFormat.h"
+#include "PixelStreamingTrace.h"
 
 FPixelStreamingVideoInput::FPixelStreamingVideoInput()
 {
@@ -17,6 +19,7 @@ void FPixelStreamingVideoInput::AddOutputFormat(int32 Format)
 
 void FPixelStreamingVideoInput::OnFrame(const IPixelCaptureInputFrame& InputFrame)
 {
+	TRACE_CPUPROFILER_EVENT_SCOPE_ON_CHANNEL_STR("PixelStreaming Video Input Frame", PixelStreamingChannel);
 	if (LastFrameWidth != -1 && LastFrameHeight != -1)
 	{
 		if (InputFrame.GetWidth() != LastFrameWidth || InputFrame.GetHeight() != LastFrameHeight)
diff --git a/Engine/Plugins/Media/PixelStreaming/Source/PixelStreaming/Private/PixelStreamingVideoInputRHI.cpp b/Engine/Plugins/Media/PixelStreaming/Source/PixelStreaming/Private/PixelStreamingVideoInputRHI.cpp
--- a/Engine/Plugins/Media/PixelStreaming/Source/PixelStreaming/Private/PixelStreamingVideoInputRHI.cpp
+++ b/Engine/Plugins/Media/PixelStreaming/Source/PixelStreaming/Private/PixelStreamingVideoInputRHI.cpp
@@ -16,9 +16,9 @@ TSharedPtr<FPixelCaptureCapturer> FPixelStreamingVideoInputRHI::CreateCapturer(i
 		case PixelCaptureBufferFormat::FORMAT_RHI:
 		{
 			// "Safe Texture Copy" polls a fence to ensure a GPU copy is complete
-			// the RDG pathway does not poll a fence so is more unsafe but offers 
+			// the RDG pathway does not poll a fence so is more unsafe but offers
 			// a significant performance increase
-			if (UE::PixelStreaming::Settings::IsUsingSafeTextureCopy())
+			if (UE::PixelStreaming::Settings::CVarPixelStreamingCaptureUseFence.GetValueOnAnyThread())
 			{
 				return FPixelCaptureCapturerRHI::Create(FinalScale);
 			}
diff --git a/Engine/Plugins/Media/PixelStreaming/Source/PixelStreaming/Private/RTCStatsCollector.cpp b/Engine/Plugins/Media/PixelStreaming/Source/PixelStreaming/Private/RTCStatsCollector.cpp
--- a/Engine/Plugins/Media/PixelStreaming/Source/PixelStreaming/Private/RTCStatsCollector.cpp
+++ b/Engine/Plugins/Media/PixelStreaming/Source/PixelStreaming/Private/RTCStatsCollector.cpp
@@ -13,6 +13,7 @@ namespace UE::PixelStreaming
 	FRTCStatsCollector::FStreamStatsSink::FStreamStatsSink(FString InRid)
 		: Rid(InRid)
 	{
+		// basic values to display
 		Add(PixelStreamingStatNames::FirCount, 0);
 		Add(PixelStreamingStatNames::PliCount, 0);
 		Add(PixelStreamingStatNames::NackCount, 0);
@@ -26,11 +27,142 @@ namespace UE::PixelStreaming
 		Add(PixelStreamingStatNames::HugeFramesSent, 0);
 		Add(PixelStreamingStatNames::AvgSendDelay, 0);
 
-		AddForCalculation(PixelStreamingStatNames::FramesSent);
-		AddForCalculation(PixelStreamingStatNames::BytesSent);
-		AddForCalculation(PixelStreamingStatNames::QPSum);
-		AddForCalculation(PixelStreamingStatNames::TotalEncodeTime);
-		AddForCalculation(PixelStreamingStatNames::FramesEncoded);
+		// these are values used to calculate extra values (stores time deltas etc)
+		AddMonitored(PixelStreamingStatNames::FramesSent);
+		AddMonitored(PixelStreamingStatNames::FramesReceived);
+		AddMonitored(PixelStreamingStatNames::BytesSent);
+		AddMonitored(PixelStreamingStatNames::BytesReceived);
+		AddMonitored(PixelStreamingStatNames::QPSum);
+		AddMonitored(PixelStreamingStatNames::TotalEncodeTime);
+		AddMonitored(PixelStreamingStatNames::FramesEncoded);
+		AddMonitored(PixelStreamingStatNames::FramesDecoded);
+
+		// below is the logic for calculating each of the stats we want to calculate
+
+		// FrameSent Per Second
+		AddStatCalculator([](FStatsSink& StatSource, double Period) -> TOptional<FStatData> {
+			FMonitoredStat* FramesSentStat = StatSource.GetMonitoredStat(PixelStreamingStatNames::FramesSent);
+			if (FramesSentStat && FramesSentStat->LatestStat.StatValue > 0)
+			{
+				const double FramesSentPerSecond = FramesSentStat->CalculateDelta(Period);
+				return FStatData(PixelStreamingStatNames::FramesSentPerSecond, FramesSentPerSecond, 0);
+			}
+			return {};
+		});
+
+		// FramesReceived Per Second
+		AddStatCalculator([](FStatsSink& StatSource, double Period) -> TOptional<FStatData> {
+			FMonitoredStat* FramesReceivedStat = StatSource.GetMonitoredStat(PixelStreamingStatNames::FramesReceived);
+			if (FramesReceivedStat && FramesReceivedStat->LatestStat.StatValue > 0)
+			{
+				const double FramesReceivedPerSecond = FramesReceivedStat->CalculateDelta(Period);
+				return FStatData(PixelStreamingStatNames::FramesReceivedPerSecond, FramesReceivedPerSecond, 0);
+			}
+			return {};
+		});
+
+		// BytesSent Per Second
+		AddStatCalculator([](FStatsSink& StatSource, double Period) -> TOptional<FStatData> {
+			FMonitoredStat* BytesSentStat = StatSource.GetMonitoredStat(PixelStreamingStatNames::BytesSent);
+			if (BytesSentStat && BytesSentStat->LatestStat.StatValue > 0)
+			{
+				const double BytesSentPerSecond = BytesSentStat->CalculateDelta(Period);
+				const double MegabitsPerSecond = BytesSentPerSecond / 1000.0 * 8.0;
+				return FStatData(PixelStreamingStatNames::Bitrate, MegabitsPerSecond, 0);
+			}
+			return {};
+		});
+
+		// BytesReceived Per Second
+		AddStatCalculator([](FStatsSink& StatSource, double Period) -> TOptional<FStatData> {
+			FMonitoredStat* BytesReceivedStat = StatSource.GetMonitoredStat(PixelStreamingStatNames::BytesReceived);
+			if (BytesReceivedStat && BytesReceivedStat->LatestStat.StatValue > 0)
+			{
+				const double BytesReceivedPerSecond = BytesReceivedStat->CalculateDelta(Period);
+				const double MegabitsPerSecond = BytesReceivedPerSecond / 1000.0 * 8.0;
+				return FStatData(PixelStreamingStatNames::Bitrate, MegabitsPerSecond, 0);
+			}
+			return {};
+		});
+
+		// Encoded fps
+		AddStatCalculator([](FStatsSink& StatSource, double Period) -> TOptional<FStatData> {
+			FMonitoredStat* EncodedFramesStat = StatSource.GetMonitoredStat(PixelStreamingStatNames::FramesEncoded);
+			if (EncodedFramesStat && EncodedFramesStat->LatestStat.StatValue > 0)
+			{
+				const double EncodedFramesPerSecond = EncodedFramesStat->CalculateDelta(Period);
+				return FStatData(PixelStreamingStatNames::EncodedFramesPerSecond, EncodedFramesPerSecond, 0);
+			}
+			return {};
+		});
+
+		// Decoded fps
+		AddStatCalculator([](FStatsSink& StatSource, double Period) -> TOptional<FStatData> {
+			FMonitoredStat* DecodedFramesStat = StatSource.GetMonitoredStat(PixelStreamingStatNames::FramesDecoded);
+			if (DecodedFramesStat && DecodedFramesStat->LatestStat.StatValue > 0)
+			{
+				const double DecodedFramesPerSecond = DecodedFramesStat->CalculateDelta(Period);
+				return FStatData(PixelStreamingStatNames::DecodedFramesPerSecond, DecodedFramesPerSecond, 0);
+			}
+			return {};
+		});
+
+		// Avg QP Per Second
+		AddStatCalculator([](FStatsSink& StatSource, double Period) -> TOptional<FStatData> {
+			FMonitoredStat* QPSumStat = StatSource.GetMonitoredStat(PixelStreamingStatNames::QPSum);
+			FStatData* EncodedFramesPerSecond = StatSource.GetCalculatedStat(PixelStreamingStatNames::EncodedFramesPerSecond);
+			if (QPSumStat && QPSumStat->LatestStat.StatValue > 0
+				&& EncodedFramesPerSecond && EncodedFramesPerSecond->StatValue > 0.0)
+			{
+				const double QPSumDeltaPerSecond = QPSumStat->CalculateDelta(Period);
+				const double MeanQPPerFrame = QPSumDeltaPerSecond / EncodedFramesPerSecond->StatValue;
+				FName StatName = PixelStreamingStatNames::MeanQPPerSecond;
+				return FStatData(StatName, MeanQPPerFrame, 0);
+			}
+			return {};
+		});
+
+		// Mean EncodeTime (ms) Per Frame
+		AddStatCalculator([](FStatsSink& StatSource, double Period) -> TOptional<FStatData> {
+			FMonitoredStat* TotalEncodeTimeStat = StatSource.GetMonitoredStat(PixelStreamingStatNames::TotalEncodeTime);
+			FStatData* EncodedFramesPerSecond = StatSource.GetCalculatedStat(PixelStreamingStatNames::EncodedFramesPerSecond);
+			if (TotalEncodeTimeStat && TotalEncodeTimeStat->LatestStat.StatValue > 0
+				&& EncodedFramesPerSecond && EncodedFramesPerSecond->StatValue > 0.0)
+			{
+				const double TotalEncodeTimePerSecond = TotalEncodeTimeStat->CalculateDelta(Period);
+				const double MeanEncodeTimePerFrameMs = TotalEncodeTimePerSecond / EncodedFramesPerSecond->StatValue * 1000.0;
+				return FStatData(PixelStreamingStatNames::MeanEncodeTime, MeanEncodeTimePerFrameMs, 2);
+			}
+			return {};
+		});
+
+		// Mean SendDelay (ms) Per Frame
+		AddStatCalculator([](FStatsSink& StatSource, double Period) -> TOptional<FStatData> {
+			FMonitoredStat* TotalSendDelayStat = StatSource.GetMonitoredStat(PixelStreamingStatNames::TotalPacketSendDelay);
+			FStatData* FramesSentPerSecond = StatSource.GetCalculatedStat(PixelStreamingStatNames::FramesSentPerSecond);
+			if (TotalSendDelayStat && TotalSendDelayStat->LatestStat.StatValue > 0
+				&& FramesSentPerSecond && FramesSentPerSecond->StatValue > 0.0)
+			{
+				const double TotalSendDelayPerSecond = TotalSendDelayStat->CalculateDelta(Period);
+				const double MeanSendDelayPerFrameMs = TotalSendDelayPerSecond / FramesSentPerSecond->StatValue * 1000.0;
+				return FStatData(PixelStreamingStatNames::MeanSendDelay, MeanSendDelayPerFrameMs, 2);
+			}
+			return {};
+		});
+
+		// JitterBufferDelay (ms)
+		AddStatCalculator([](FStatsSink& StatSource, double Period) -> TOptional<FStatData> {
+			FMonitoredStat* JitterBufferDelayStat = StatSource.GetMonitoredStat(PixelStreamingStatNames::JitterBufferDelay);
+			FStatData* FramesReceivedPerSecond = StatSource.GetCalculatedStat(PixelStreamingStatNames::FramesReceivedPerSecond);
+			if (JitterBufferDelayStat && JitterBufferDelayStat->LatestStat.StatValue > 0
+				&& FramesReceivedPerSecond && FramesReceivedPerSecond->StatValue > 0.0)
+			{
+				const double TotalJitterBufferDelayPerSecond = JitterBufferDelayStat->CalculateDelta(Period);
+				const double MeanJitterBufferDelayMs = TotalJitterBufferDelayPerSecond / FramesReceivedPerSecond->StatValue * 1000.0;
+				return FStatData(PixelStreamingStatNames::JitterBufferDelay, MeanJitterBufferDelayMs, 2);
+			}
+			return {};
+		});
 	}
 
 	// ------------- FRTCStatsCollector-------------------
@@ -47,7 +179,7 @@ namespace UE::PixelStreaming
 	{
 		FStatsSink TrackStatsSink;
 		TrackStatsSink.Add(PixelStreamingStatNames::JitterBufferDelay, 2);
-		TrackStatsSink.Add(PixelStreamingStatNames::FramesReceived, 0);
+		TrackStatsSink.Add(PixelStreamingStatNames::FramesPerSecond, 0);
 		TrackStatsSink.Add(PixelStreamingStatNames::FramesDecoded, 0);
 		TrackStatsSink.Add(PixelStreamingStatNames::FramesDropped, 0);
 		TrackStatsSink.Add(PixelStreamingStatNames::FramesCorrupted, 0);
@@ -82,11 +214,48 @@ namespace UE::PixelStreaming
 		return rtc::RefCountReleaseStatus::kOtherRefsRemained;
 	}
 
-	void FRTCStatsCollector::ExtractValueAndSet(const webrtc::RTCStatsMemberInterface* ExtractFrom, FStatData* SetValueHere)
+	// returns true if the value is worth storing (false if it started and remains zero)
+	bool FRTCStatsCollector::ExtractValueAndSet(const webrtc::RTCStatsMemberInterface* ExtractFrom, FStatData* SetValueHere)
 	{
+		const bool bZeroInitially = SetValueHere->StatValue == 0.0;
 		FString StatValueStr = ExtractFrom->is_defined() ? ToString(ExtractFrom->ValueToString()) : TEXT("0.0");
 		double StatValueDouble = FCString::Atod(*StatValueStr);
 		SetValueHere->StatValue = StatValueDouble;
+		const bool bZeroStill = SetValueHere->StatValue == 0.0;
+		return !(bZeroInitially && bZeroStill);
+	}
+
+	FRTCStatsCollector::FStatsSink* FRTCStatsCollector::FindSink(const FString& StatsType, const std::vector<const webrtc::RTCStatsMemberInterface*>& StatMembers, FString& OutSsrc)
+	{
+		FStatsSink* Sink = nullptr;
+
+		if (StatsType == "outbound-rtp" || StatsType == "inbound-rtp")
+		{
+			// Extract the `ssrc` to uniquely id the stream
+			for (const webrtc::RTCStatsMemberInterface* StatMember : StatMembers)
+			{
+				const FString StatName = FString(StatMember->name());
+				if (StatName == "ssrc")
+				{
+					OutSsrc = StatMember->is_defined() ? ToString(StatMember->ValueToString()) : TEXT("");
+					if (!StreamStatsSinks.Contains(OutSsrc))
+					{
+						Sink = &StreamStatsSinks.Add(OutSsrc, FStreamStatsSink(OutSsrc));
+					}
+					else
+					{
+						Sink = StreamStatsSinks.Find(OutSsrc);
+					}
+					break;
+				}
+			}
+		}
+		else
+		{
+			Sink = StatSinks.Find(StatsType);
+		}
+
+		return Sink;
 	}
 
 	void FRTCStatsCollector::OnStatsDelivered(const rtc::scoped_refptr<const webrtc::RTCStatsReport>& Report)
@@ -100,7 +269,7 @@ namespace UE::PixelStreaming
 
 		for (auto&& Stats : *Report)
 		{
-			FString StatsType = FString(Stats.type());
+			const FString StatsType = FString(Stats.type());
 			std::vector<const webrtc::RTCStatsMemberInterface*> StatMembers = Stats.Members();
 
 			// For debugging
@@ -111,76 +280,37 @@ namespace UE::PixelStreaming
 			// }
 
 			// Find relevant FStatsSink
-			FStatsSink* StatsSink = nullptr;
-			FString Rid;
+			FString Ssrc;
 
-			if (StatsType == "outbound-rtp")
+			if (FStatsSink* StatsSink = FindSink(StatsType, StatMembers, Ssrc))
 			{
-				// Extract the `rid`
+				const bool bSimulcastStat = Ssrc.Contains(TEXT("Simulcast"), ESearchCase::IgnoreCase, ESearchDir::FromStart);
+				const FString Id = bSimulcastStat ? Ssrc : AssociatedPlayerId;
+
 				for (const webrtc::RTCStatsMemberInterface* StatMember : StatMembers)
 				{
-					FString StatName = FString(StatMember->name());
-					if (StatName == "rid")
-					{
+					const FName StatName = FName(StatMember->name());
 
-						Rid = StatMember->is_defined() ? ToString(StatMember->ValueToString()) : TEXT("");
-						if (!StreamStatsSinks.Contains(Rid))
+					// Check if the stat is a stat we use for calculation, if so, extract the value and store it
+					if (FMonitoredStat* MonitoredStat = StatsSink->GetMonitoredStat(StatName))
+					{
+						ExtractValueAndSet(StatMember, &MonitoredStat->LatestStat);
+					}
+					else if (FStatData* StatToEmit = StatsSink->Get(StatName))
+					{
+						if (ExtractValueAndSet(StatMember, StatToEmit))
 						{
-							StreamStatsSinks.Add(Rid, FStreamStatsSink(Rid));
+							PSStats->StorePeerStat(Id, *StatToEmit);
 						}
-						StatsSink = StreamStatsSinks.Find(Rid);
-						break;
 					}
 				}
-			}
-			else
-			{
-				StatsSink = StatSinks.Find(StatsType);
-			}
-
-			if (!StatsSink)
-			{
-				continue;
-			}
-
-			bool bSimulcastStat = Rid.Contains(TEXT("Simulcast"), ESearchCase::IgnoreCase, ESearchDir::FromStart);
-			FString Id = bSimulcastStat ? Rid : AssociatedPlayerId;
-
-			for (const webrtc::RTCStatsMemberInterface* StatMember : StatMembers)
-			{
-				FName StatName = FName(StatMember->name());
-
-				// Check if the stat is a stat we use for calculation, if so, extract the value and store it
-				FCalculatedStat* CalcStat = StatsSink->GetCalcStat(StatName);
-				if (CalcStat)
-				{
-					ExtractValueAndSet(StatMember, &CalcStat->LatestStat);
-					continue;
-				}
-
-				// Check if stat is a stat we emit immediately, if so, extract the value and emit it
-				FStatData* StatToEmit = StatsSink->Get(StatName);
-				if (StatToEmit)
-				{
-					const bool bZeroInitially = StatToEmit->StatValue == 0.0;
-					ExtractValueAndSet(StatMember, StatToEmit);
-					const bool bZeroStill = StatToEmit->StatValue == 0.0;
-
-					// Don't bother emitting a zero value
-					if (!bZeroInitially && !bZeroStill)
-					{
-						PSStats->StorePeerStat(Id, *StatToEmit);
-					}
 
-					continue;
-				}
+				PostDeliverCalculateStats(PSStats);
 			}
-
-			CalculateStats(PSStats);
 		}
 	}
 
-	void FRTCStatsCollector::CalculateStats(FStats* PSStats)
+	void FRTCStatsCollector::PostDeliverCalculateStats(FStats* PSStats)
 	{
 		uint64 CyclesNow = FPlatformTime::Cycles64();
 		double SecondsDelta = FGenericPlatformTime::ToSeconds64(CyclesNow - LastCalculationCycles);
@@ -191,75 +321,23 @@ namespace UE::PixelStreaming
 			return;
 		}
 
-		// This will bring the calculations into a 1 second window
-		double PerSecondRatio = 1.0 / SecondsDelta;
-		double FramesSentPerSecond = 0.0;
-		double EncodedFramesPerSecond = 0.0;
-
 		// Calculate stats based on each of the stream stats sinks
-
 		for (auto& Entry : StreamStatsSinks)
 		{
-			FString& Rid = Entry.Key;
+			const FString& Ssrc = Entry.Key;
+			const bool bSimulcastStat = Ssrc.Contains(TEXT("Simulcast"), ESearchCase::IgnoreCase, ESearchDir::FromStart);
+			const FString Id = bSimulcastStat ? Ssrc : AssociatedPlayerId;
 			FStatsSink& StreamSink = Entry.Value;
-			bool bSimulcastStat = Rid.Contains(TEXT("Simulcast"), ESearchCase::IgnoreCase, ESearchDir::FromStart);
-			FString Id = bSimulcastStat ? Rid : AssociatedPlayerId;
-
-			// --------- FrameSent Per Second -----------
-			FCalculatedStat* FramesSentStat = StreamSink.GetCalcStat(PixelStreamingStatNames::FramesSent);
-			if (FramesSentStat && FramesSentStat->LatestStat.StatValue > 0)
-			{
-				FramesSentPerSecond = FramesSentStat->CalculateDelta(PerSecondRatio);
-				PSStats->StorePeerStat(Id, FStatData(PixelStreamingStatNames::FramesSentPerSecond, FramesSentPerSecond, 0));
-			}
 
-			// --------- BytesSent Per Second -----------
-			FCalculatedStat* BytesSentStat = StreamSink.GetCalcStat(PixelStreamingStatNames::BytesSent);
-			if (BytesSentStat && BytesSentStat->LatestStat.StatValue > 0)
-			{
-				double BytesSentPerSecond = BytesSentStat->CalculateDelta(PerSecondRatio);
-				double MegabitsPerSecond = BytesSentPerSecond / 1000.0 * 8.0;
-				FName StatName = PixelStreamingStatNames::Bitrate;
-				PSStats->StorePeerStat(Id, FStatData(StatName, MegabitsPerSecond, 0));
-			}
-
-			// ------------- Encoded fps -------------
-			FCalculatedStat* EncodedFramesStat = StreamSink.GetCalcStat(PixelStreamingStatNames::FramesEncoded);
-			if (EncodedFramesStat && EncodedFramesStat->LatestStat.StatValue > 0)
+			for (auto& Calculator : StreamSink.Calculators)
 			{
-				EncodedFramesPerSecond = EncodedFramesStat->CalculateDelta(PerSecondRatio);
-				FName StatName = PixelStreamingStatNames::EncodedFramesPerSecond;
-				PSStats->StorePeerStat(Id, FStatData(StatName, EncodedFramesPerSecond, 0));
-			}
-
-			// ------------- Avg QP Per Second -------------
-			FCalculatedStat* QPSumStat = StreamSink.GetCalcStat(PixelStreamingStatNames::QPSum);
-			if (QPSumStat && QPSumStat->LatestStat.StatValue > 0 && EncodedFramesPerSecond > 0.0)
-			{
-				double QPSumDeltaPerSecond = QPSumStat->CalculateDelta(PerSecondRatio);
-				double MeanQPPerFrame = QPSumDeltaPerSecond / EncodedFramesPerSecond;
-				FName StatName = PixelStreamingStatNames::MeanQPPerSecond;
-				PSStats->StorePeerStat(Id, FStatData(StatName, MeanQPPerFrame, 0));
-			}
-
-			// ------------- Mean EncodeTime (ms) Per Frame -------------
-			FCalculatedStat* TotalEncodeTimeStat = StreamSink.GetCalcStat(PixelStreamingStatNames::TotalEncodeTime);
-			if (TotalEncodeTimeStat && TotalEncodeTimeStat->LatestStat.StatValue > 0 && EncodedFramesPerSecond > 0.0)
-			{
-				double TotalEncodeTimePerSecond = TotalEncodeTimeStat->CalculateDelta(PerSecondRatio);
-				double MeanEncodeTimePerFrameMs = TotalEncodeTimePerSecond / EncodedFramesPerSecond * 1000.0;
-				FName StatName = PixelStreamingStatNames::MeanEncodeTime;
-				PSStats->StorePeerStat(Id, FStatData(StatName, MeanEncodeTimePerFrameMs, 2));
-			}
-
-			// ------------- Mean SendDelay (ms) Per Frame -------------
-			FCalculatedStat* TotalSendDelayStat = StreamSink.GetCalcStat(PixelStreamingStatNames::TotalPacketSendDelay);
-			if (TotalSendDelayStat && TotalSendDelayStat->LatestStat.StatValue > 0 && FramesSentPerSecond > 0)
-			{
-				double TotalSendDelayPerSecond = TotalSendDelayStat->CalculateDelta(PerSecondRatio);
-				double MeanSendDelayPerFrameMs = TotalSendDelayPerSecond / FramesSentPerSecond * 1000.0;
-				FName StatName = PixelStreamingStatNames::MeanSendDelay;
-				PSStats->StorePeerStat(Id, FStatData(StatName, MeanSendDelayPerFrameMs, 2));
+				TOptional<FStatData> OptStatData = Calculator(StreamSink, SecondsDelta);
+				if (OptStatData.IsSet())
+				{
+					FStatData& StatData = *OptStatData;
+					StreamSink.CalculatedStats.Add(StatData.StatName, StatData);
+					PSStats->StorePeerStat(Id, StatData);
+				}
 			}
 		}
 
diff --git a/Engine/Plugins/Media/PixelStreaming/Source/PixelStreaming/Private/RTCStatsCollector.h b/Engine/Plugins/Media/PixelStreaming/Source/PixelStreaming/Private/RTCStatsCollector.h
--- a/Engine/Plugins/Media/PixelStreaming/Source/PixelStreaming/Private/RTCStatsCollector.h
+++ b/Engine/Plugins/Media/PixelStreaming/Source/PixelStreaming/Private/RTCStatsCollector.h
@@ -8,7 +8,7 @@
 
 namespace UE::PixelStreaming
 {
-	class FRTCStatsCollector : public webrtc::RTCStatsCollectorCallback
+	class PIXELSTREAMING_API FRTCStatsCollector : public webrtc::RTCStatsCollectorCallback
 	{
 	public:
 		FRTCStatsCollector();
@@ -20,16 +20,12 @@ namespace UE::PixelStreaming
 		rtc::RefCountReleaseStatus Release() const override;
 		// End RTCStatsCollectorCallback interface
 
-	private:
-		void ExtractValueAndSet(const webrtc::RTCStatsMemberInterface* ExtractFrom, FStatData* SetValueHere);
-		void CalculateStats(FStats* PSStats);
-
 	private:
 		// Some stats have a latest value that is update each time that stats comes and
 		// then also a calculated value that might get updates every 1 second.
 		// An example is FPS, where total frames sent is the stat that is tracked but
 		// the calculated stat is frames sent per second.
-		struct FCalculatedStat
+		struct FMonitoredStat
 		{
 			FStatData LatestStat;
 			double PrevValue;
@@ -48,15 +44,19 @@ namespace UE::PixelStreaming
 			virtual ~FStatsSink() = default;
 			// Sink only interested in a single type of stat from WebRTC (e.g. track, outbound-rtp).
 			virtual bool Contains(const FName& StatName) const { return Stats.Contains(StatName); };
-			virtual bool ContainsCalcStat(const FName& StatName) const { return CalcStats.Contains(StatName); }
+			virtual bool ContainsMonitoredStat(const FName& StatName) const { return MonitoredStats.Contains(StatName); }
+			virtual bool ContainsCalculatedStat(const FName& StatName) const { return CalculatedStats.Contains(StatName); }
 			virtual void Add(FName StatName, int NDecimalPlaces) { Stats.Add(StatName, FStatData(StatName, 0.0, NDecimalPlaces)); };
-			virtual void AddForCalculation(FName StatName) { CalcStats.Add(StatName, { FStatData(StatName, 0.0, 2), 0.0 }); };
+			virtual void AddMonitored(FName StatName) { MonitoredStats.Add(StatName, { FStatData(StatName, 0.0, 2), 0.0 }); };
+			virtual void AddStatCalculator(const TFunction<TOptional<FStatData>(FStatsSink&, double)>& Calculator) { Calculators.Add(Calculator); }
 			virtual FStatData* Get(const FName& StatName) { return Stats.Find(StatName); };
-			virtual FCalculatedStat* GetCalcStat(const FName& StatName) { return CalcStats.Find(StatName); }
+			virtual FMonitoredStat* GetMonitoredStat(const FName& StatName) { return MonitoredStats.Find(StatName); }
+			virtual FStatData* GetCalculatedStat(const FName& StatName) { return CalculatedStats.Find(StatName); }
 
-		protected:
-			TMap<FName, FStatData> Stats;
-			TMap<FName, FCalculatedStat> CalcStats;
+			TMap<FName, FStatData> Stats;				// basic stats that are displayed as is.
+			TMap<FName, FMonitoredStat> MonitoredStats;	// stats to monitor changes over time.
+			TMap<FName, FStatData> CalculatedStats;		// the result of calculations that can be used in other calculations
+			TArray<TFunction<TOptional<FStatData>(FStatsSink&, double)>> Calculators;
 		};
 
 		class FStreamStatsSink : public FStatsSink
@@ -70,6 +70,11 @@ namespace UE::PixelStreaming
 			FString Rid;
 		};
 
+	private:
+		FStatsSink* FindSink(const FString& StatsType, const std::vector<const webrtc::RTCStatsMemberInterface*>& StatMembers, FString& OutSsrc);
+		bool ExtractValueAndSet(const webrtc::RTCStatsMemberInterface* ExtractFrom, FStatData* SetValueHere);
+		void PostDeliverCalculateStats(FStats* PSStats);
+
 	private:
 		FPixelStreamingPlayerId AssociatedPlayerId;
 
