Integration of AVCodecs into PixelStreaming and support for new pixel formats

diff --git a/Engine/Plugins/Media/PixelStreaming/Source/PixelStreaming/Private/Tests/SignallingServerConnectionTests.cpp b/Engine/Plugins/Media/PixelStreaming/Source/PixelStreaming/Private/Tests/SignallingServerConnectionTests.cpp
--- a/Engine/Plugins/Media/PixelStreaming/Source/PixelStreaming/Private/Tests/SignallingServerConnectionTests.cpp
+++ b/Engine/Plugins/Media/PixelStreaming/Source/PixelStreaming/Private/Tests/SignallingServerConnectionTests.cpp
@@ -4,10 +4,16 @@
 #include "PixelStreamingSignallingConnection.h"
 #include "Serialization/JsonSerializer.h"
 #include "IWebSocket.h"
+#include "Utils.h"
+#include "ToStringExtensions.h"
+#include "Dom/JsonObject.h"
+
+DECLARE_LOG_CATEGORY_EXTERN(LogMockPixelStreamingSS, Log, VeryVerbose);
+DEFINE_LOG_CATEGORY(LogMockPixelStreamingSS);
 
 namespace
 {
-	//using namespace UE::PixelStreaming;
+	// using namespace UE::PixelStreaming;
 
 	// For validating callbacks from the signalling server connection
 	class FTestSSConnectionObserver : public IPixelStreamingSignallingConnectionObserver
@@ -30,9 +36,9 @@ namespace
 		{
 			SetIceCandidate = { SdpMid, SdpMLineIndex, Sdp, PlayerId };
 		}
-		// virtual void OnPlayerConnected(FPixelStreamingPlayerId PlayerId, int Flags) { unimplemented(); }
-		// virtual void OnPlayerDisconnected(FPixelStreamingPlayerId PlayerId) { unimplemented(); }
-		// virtual void OnSFUPeerDataChannels(FPixelStreamingPlayerId SFUId, FPixelStreamingPlayerId PlayerId, int32 SendStreamId, int32 RecvStreamId) { unimplemented(); }
+		virtual void OnSignallingPlayerConnected(FPixelStreamingPlayerId, const FPixelStreamingPlayerConfig&) override {}
+		virtual void OnSignallingPlayerDisconnected(FPixelStreamingPlayerId PlayerId) override {}
+		virtual void OnSignallingSFUPeerDataChannels(FPixelStreamingPlayerId SFUId, FPixelStreamingPlayerId PlayerId, int32 SendStreamId, int32 RecvStreamId) override {}
 
 		//// Player-only
 		virtual void OnSignallingSessionDescription(webrtc::SdpType Type, const FString& Sdp) override
@@ -43,8 +49,8 @@ namespace
 		{
 			SetIceCandidate = { SdpMid, SdpMLineIndex, Sdp };
 		}
-		// virtual void OnPlayerCount(uint32 Count) { unimplemented(); }
-		// virtual void OnPeerDataChannels(int32 SendStreamId, int32 RecvStreamId) { unimplemented(); }
+		virtual void OnSignallingPlayerCount(uint32 Count) override {}
+		virtual void OnSignallingPeerDataChannels(int32 SendStreamId, int32 RecvStreamId) override {}
 
 		void Reset()
 		{
@@ -115,6 +121,556 @@ namespace
 		bool bConnected = false;
 	};
 
+	class FMockSignallingConnection : public IPixelStreamingSignallingConnection
+	{
+	public:
+		FMockSignallingConnection(IPixelStreamingSignallingConnectionObserver& InObserver, FString InStreamerId)
+			: Observer(InObserver)
+			, StreamerId(InStreamerId)
+		{
+			RegisterHandler("identify", [this](FJsonObjectPtr JsonMsg) {
+				FJsonObjectPtr Json = MakeShared<FJsonObject>();
+				Json->SetStringField(TEXT("type"), TEXT("endpointId"));
+				Json->SetStringField(TEXT("id"), StreamerId);
+				FString Msg = UE::PixelStreaming::ToString(Json, false);
+				UE_LOG(LogMockPixelStreamingSS, Verbose, TEXT("-> SS: endpointId\n%s"), *Msg);
+				SendMessage(Msg);
+			});
+			RegisterHandler("config", [this](FJsonObjectPtr JsonMsg) {
+				const TSharedPtr<FJsonObject>* PeerConnectionOptions = nullptr;
+				if (!JsonMsg->TryGetObjectField(TEXT("peerConnectionOptions"), PeerConnectionOptions))
+				{
+					UE_LOG(LogMockPixelStreamingSS, Error, TEXT("Cannot find `peerConnectionOptions` field in SS config\n%s"), *UE::PixelStreaming::ToString(JsonMsg));
+					return;
+				}
+
+				webrtc::PeerConnectionInterface::RTCConfiguration RTCConfig;
+
+				const TArray<TSharedPtr<FJsonValue>>* IceServers = nullptr;
+				if ((*PeerConnectionOptions)->TryGetArrayField(TEXT("iceServers"), IceServers))
+				{
+					for (const TSharedPtr<FJsonValue>& IceServerVal : *IceServers)
+					{
+						const TSharedPtr<FJsonObject>* IceServerJson;
+						if (!IceServerVal->TryGetObject(IceServerJson))
+						{
+							UE_LOG(LogMockPixelStreamingSS, Error, TEXT("Failed to parse SS config: `iceServer` - not an object\n%s"), *UE::PixelStreaming::ToString(*PeerConnectionOptions));
+							return;
+						}
+
+						RTCConfig.servers.push_back(webrtc::PeerConnectionInterface::IceServer{});
+						webrtc::PeerConnectionInterface::IceServer& IceServer = RTCConfig.servers.back();
+
+						TArray<FString> Urls;
+						if ((*IceServerJson)->TryGetStringArrayField(TEXT("urls"), Urls))
+						{
+							for (const FString& UrlElem : Urls)
+							{
+								IceServer.urls.push_back(UE::PixelStreaming::ToString(UrlElem));
+							}
+						}
+						else
+						{
+							// in the RTC Spec, "urls" can be an array or a single string
+							// https://www.w3.org/TR/webrtc/#dictionary-rtciceserver-members
+							FString UrlsSingle;
+							if ((*IceServerJson)->TryGetStringField(TEXT("urls"), UrlsSingle))
+							{
+								IceServer.urls.push_back(UE::PixelStreaming::ToString(UrlsSingle));
+							}
+						}
+
+						FString Username;
+						if ((*IceServerJson)->TryGetStringField(TEXT("username"), Username))
+						{
+							IceServer.username = UE::PixelStreaming::ToString(Username);
+						}
+
+						FString Credential;
+						if ((*IceServerJson)->TryGetStringField(TEXT("credential"), Credential))
+						{
+							IceServer.password = UE::PixelStreaming::ToString(Credential);
+						}
+					}
+				}
+
+				// force `UnifiedPlan` as we control both ends of WebRTC streaming
+				RTCConfig.sdp_semantics = webrtc::SdpSemantics::kUnifiedPlan;
+
+				Observer.OnSignallingConfig(RTCConfig);
+			});
+
+			RegisterHandler("offer", [this](FJsonObjectPtr JsonMsg) {
+				webrtc::SdpType Type = JsonMsg->GetStringField(TEXT("type")) == TEXT("offer") ? webrtc::SdpType::kOffer : webrtc::SdpType::kAnswer;
+
+				FString Sdp;
+				if (!JsonMsg->TryGetStringField(TEXT("sdp"), Sdp))
+				{
+					UE_LOG(LogMockPixelStreamingSS, Error, TEXT("Cannot find `sdp` in Streamer's answer\n%s"), *UE::PixelStreaming::ToString(JsonMsg));
+					return;
+				}
+
+				FPixelStreamingPlayerId PlayerId;
+				bool bGotPlayerId = GetPlayerIdJson(JsonMsg, PlayerId);
+				if (!bGotPlayerId)
+				{
+					Observer.OnSignallingSessionDescription(Type, Sdp);
+				}
+				else
+				{
+					Observer.OnSignallingSessionDescription(PlayerId, Type, Sdp);
+				}
+			});
+			RegisterHandler("answer", [this](FJsonObjectPtr JsonMsg) {
+				webrtc::SdpType Type = JsonMsg->GetStringField(TEXT("type")) == TEXT("offer") ? webrtc::SdpType::kOffer : webrtc::SdpType::kAnswer;
+
+				FString Sdp;
+				if (!JsonMsg->TryGetStringField(TEXT("sdp"), Sdp))
+				{
+					UE_LOG(LogMockPixelStreamingSS, Error, TEXT("Cannot find `sdp` in Streamer's answer\n%s"), *UE::PixelStreaming::ToString(JsonMsg));
+					return;
+				}
+
+				FPixelStreamingPlayerId PlayerId;
+				bool bGotPlayerId = GetPlayerIdJson(JsonMsg, PlayerId);
+				if (!bGotPlayerId)
+				{
+					Observer.OnSignallingSessionDescription(Type, Sdp);
+				}
+				else
+				{
+					Observer.OnSignallingSessionDescription(PlayerId, Type, Sdp);
+				}
+			});
+			RegisterHandler("iceCandidate", [this](FJsonObjectPtr JsonMsg) {
+				FPixelStreamingPlayerId PlayerId;
+				bool bGotPlayerId = GetPlayerIdJson(JsonMsg, PlayerId);
+
+				const FJsonObjectPtr* CandidateJson;
+				if (!JsonMsg->TryGetObjectField(TEXT("candidate"), CandidateJson))
+				{
+					PlayerError(PlayerId, TEXT("Failed to get `candiate` from remote `iceCandidate` message\n%s"), *UE::PixelStreaming::ToString(JsonMsg));
+				}
+
+				FString SdpMid;
+				if (!(*CandidateJson)->TryGetStringField(TEXT("sdpMid"), SdpMid))
+				{
+					PlayerError(PlayerId, TEXT("Failed to get `sdpMid` from remote `iceCandidate` message\n%s"), *UE::PixelStreaming::ToString(JsonMsg));
+				}
+
+				int32 SdpMLineIndex = 0;
+				if (!(*CandidateJson)->TryGetNumberField(TEXT("sdpMlineIndex"), SdpMLineIndex))
+				{
+					PlayerError(PlayerId, TEXT("Failed to get `sdpMlineIndex` from remote `iceCandidate` message\n%s"), *UE::PixelStreaming::ToString(JsonMsg));
+				}
+
+				FString CandidateStr;
+				if (!(*CandidateJson)->TryGetStringField(TEXT("candidate"), CandidateStr))
+				{
+					PlayerError(PlayerId, TEXT("Failed to get `candidate` from remote `iceCandidate` message\n%s"), *UE::PixelStreaming::ToString(JsonMsg));
+				}
+
+				if (bGotPlayerId)
+				{
+					Observer.OnSignallingRemoteIceCandidate(PlayerId, SdpMid, SdpMLineIndex, CandidateStr);
+				}
+				else
+				{
+					Observer.OnSignallingRemoteIceCandidate(SdpMid, SdpMLineIndex, CandidateStr);
+				}
+			});
+			RegisterHandler("ping", [this](FJsonObjectPtr JsonMsg) { /* nothing */ });
+			RegisterHandler("pong", [this](FJsonObjectPtr JsonMsg) { /* nothing */ });
+			RegisterHandler("playerCount", [this](FJsonObjectPtr JsonMsg) {
+				uint32 Count;
+				if (!JsonMsg->TryGetNumberField(TEXT("count"), Count))
+				{
+					UE_LOG(LogMockPixelStreamingSS, Error, TEXT("Failed to get `count` from `playerCount` message\n%s"), *UE::PixelStreaming::ToString(JsonMsg));
+					return;
+				}
+
+				Observer.OnSignallingPlayerCount(Count);
+			});
+			RegisterHandler("playerConnected", [this](FJsonObjectPtr JsonMsg) {
+				FPixelStreamingPlayerId PlayerId;
+				bool bGotPlayerId = GetPlayerIdJson(JsonMsg, PlayerId);
+				if (!bGotPlayerId)
+				{
+					UE_LOG(LogMockPixelStreamingSS, Error, TEXT("Failed to get `playerId` from `join` message\n%s"), *UE::PixelStreaming::ToString(JsonMsg));
+					return;
+				}
+
+				UE_LOG(LogMockPixelStreamingSS, Log, TEXT("Got player connected, player id=%s"), *PlayerId);
+
+				FPixelStreamingPlayerConfig PlayerConfig;
+
+				// Default to always making datachannel, unless explicitly set to false.
+				bool bMakeDataChannel = true;
+				JsonMsg->TryGetBoolField(TEXT("datachannel"), PlayerConfig.SupportsDataChannel);
+
+				// Default peer is not an SFU, unless explictly set as SFU
+				bool bIsSFU = false;
+				JsonMsg->TryGetBoolField(TEXT("sfu"), PlayerConfig.IsSFU);
+
+				Observer.OnSignallingPlayerConnected(PlayerId, PlayerConfig);
+			});
+			RegisterHandler("playerDisconnected", [this](FJsonObjectPtr JsonMsg) {
+				FPixelStreamingPlayerId PlayerId;
+				bool bGotPlayerId = GetPlayerIdJson(JsonMsg, PlayerId);
+				if (!bGotPlayerId)
+				{
+					UE_LOG(LogMockPixelStreamingSS, Error, TEXT("Failed to get `playerId` from `playerDisconnected` message\n%s"), *UE::PixelStreaming::ToString(JsonMsg));
+					return;
+				}
+
+				Observer.OnSignallingPlayerDisconnected(PlayerId);
+			});
+			RegisterHandler("streamerDataChannels", [this](FJsonObjectPtr JsonMsg) {
+				FPixelStreamingPlayerId SFUId;
+				bool bSuccess = GetPlayerIdJson(JsonMsg, SFUId, TEXT("sfuId"));
+				if (!bSuccess)
+				{
+					UE_LOG(LogMockPixelStreamingSS, Error, TEXT("Failed to get `sfuId` from `streamerDataChannels` message\n%s"), *UE::PixelStreaming::ToString(JsonMsg));
+					return;
+				}
+
+				FPixelStreamingPlayerId PlayerId;
+				bSuccess = GetPlayerIdJson(JsonMsg, PlayerId, TEXT("playerId"));
+				if (!bSuccess)
+				{
+					UE_LOG(LogMockPixelStreamingSS, Error, TEXT("Failed to get `playerId` from `streamerDataChannels` message\n%s"), *UE::PixelStreaming::ToString(JsonMsg));
+					return;
+				}
+
+				int32 SendStreamId;
+				bSuccess = JsonMsg->TryGetNumberField(TEXT("sendStreamId"), SendStreamId);
+				if (!bSuccess)
+				{
+					UE_LOG(LogMockPixelStreamingSS, Error, TEXT("Failed to get `sendStreamId` from `streamerDataChannels` message\n%s"), *UE::PixelStreaming::ToString(JsonMsg));
+					return;
+				}
+
+				int32 RecvStreamId;
+				bSuccess = JsonMsg->TryGetNumberField(TEXT("recvStreamId"), RecvStreamId);
+				if (!bSuccess)
+				{
+					UE_LOG(LogMockPixelStreamingSS, Error, TEXT("Failed to get `recvStreamId` from `streamerDataChannels` message\n%s"), *UE::PixelStreaming::ToString(JsonMsg));
+					return;
+				}
+
+				Observer.OnSignallingSFUPeerDataChannels(SFUId, PlayerId, SendStreamId, RecvStreamId);
+			});
+			RegisterHandler("peerDataChannels", [this](FJsonObjectPtr JsonMsg) {
+				int32 SendStreamId = 0;
+				if (!JsonMsg->TryGetNumberField(TEXT("sendStreamId"), SendStreamId))
+				{
+					UE_LOG(LogMockPixelStreamingSS, Error, TEXT("Failed to get `sendStreamId` from remote `peerDataChannels` message\n%s"), *UE::PixelStreaming::ToString(JsonMsg));
+					return;
+				}
+				int32 RecvStreamId = 0;
+				if (!JsonMsg->TryGetNumberField(TEXT("recvStreamId"), RecvStreamId))
+				{
+					UE_LOG(LogMockPixelStreamingSS, Error, TEXT("Failed to get `recvStreamId` from remote `peerDataChannels` message\n%s"), *UE::PixelStreaming::ToString(JsonMsg));
+					return;
+				}
+				Observer.OnSignallingPeerDataChannels(SendStreamId, RecvStreamId);
+			});
+		}
+		virtual ~FMockSignallingConnection() = default;
+
+		virtual void TryConnect(FString URL) override
+		{
+			if (WebSocket && WebSocket->IsConnected())
+			{
+				UE_LOG(LogMockPixelStreamingSS, Log, TEXT("Skipping `Connect()` because we are already connected. If you want to reconnect then disconnect first."));
+				return;
+			}
+
+			// Reconnecting on an existing websocket can be problematic depending what state it was
+			// left in. Easier and safer to disconnect any existing socket/delegates and start fresh.
+			Disconnect();
+
+			Url = URL;
+
+			WebSocket = MakeShared<FMockWebSocket>();
+
+			OnConnectedHandle = WebSocket->OnConnected().AddLambda([this]() {
+				UE_LOG(LogMockPixelStreamingSS, Log, TEXT("Connected to SS %s"), *Url);
+				Observer.OnSignallingConnected();
+			});
+
+			OnConnectionErrorHandle = WebSocket->OnConnectionError().AddLambda([this](const FString& Error) {
+				UE_LOG(LogMockPixelStreamingSS, Log, TEXT("Failed to connect to SS %s - signalling server may not be up yet. Message: \"%s\""), *Url, *Error);
+				Observer.OnSignallingError(Error);
+			});
+
+			OnClosedHandle = WebSocket->OnClosed().AddLambda([this](int32 StatusCode, const FString& Reason, bool bWasClean) {
+				UE_LOG(LogMockPixelStreamingSS, Log, TEXT("Closed connection to SS %s - \n\tstatus %d\n\treason: %s\n\twas clean: %s"), *Url, StatusCode, *Reason, bWasClean ? TEXT("true") : TEXT("false"));
+				Observer.OnSignallingDisconnected(StatusCode, Reason, bWasClean);
+			});
+
+			OnMessageHandle = WebSocket->OnMessage().AddLambda([this](const FString& Msg) {
+				FJsonObjectPtr JsonMsg;
+				const auto JsonReader = TJsonReaderFactory<TCHAR>::Create(Msg);
+
+				if (!FJsonSerializer::Deserialize(JsonReader, JsonMsg))
+				{
+					UE_LOG(LogMockPixelStreamingSS, Error, TEXT("Failed to parse SS message:\n%s"), *Msg);
+					return;
+				}
+
+				FString MsgType;
+				if (!JsonMsg->TryGetStringField(TEXT("type"), MsgType))
+				{
+					UE_LOG(LogMockPixelStreamingSS, Error, TEXT("Cannot find `type` field in SS message:\n%s"), *Msg);
+					return;
+				}
+
+				TFunction<void(FJsonObjectPtr)>* Handler = MessageHandlers.Find(MsgType);
+				if (Handler != nullptr)
+				{
+					(*Handler)(JsonMsg);
+				}
+			});
+
+			UE_LOG(LogMockPixelStreamingSS, Log, TEXT("Connecting to SS %s"), *Url);
+			WebSocket->Connect();
+		}
+
+		virtual void Disconnect() override
+		{
+			if (!WebSocket)
+			{
+				return;
+			}
+
+			WebSocket->OnConnected().Remove(OnConnectedHandle);
+			WebSocket->OnConnectionError().Remove(OnConnectionErrorHandle);
+			WebSocket->OnClosed().Remove(OnClosedHandle);
+			WebSocket->OnMessage().Remove(OnMessageHandle);
+
+			WebSocket->Close();
+			WebSocket = nullptr;
+			UE_LOG(LogMockPixelStreamingSS, Log, TEXT("Closing websocket to SS %s"), *Url);
+		}
+		virtual bool IsConnected() const override
+		{
+			return WebSocket != nullptr && WebSocket.IsValid() && WebSocket->IsConnected();
+		}
+
+		virtual void SendOffer(FPixelStreamingPlayerId PlayerId, const webrtc::SessionDescriptionInterface& SDP) override
+		{
+			FJsonObjectPtr OfferJson = MakeShared<FJsonObject>();
+			OfferJson->SetStringField(TEXT("type"), TEXT("offer"));
+			SetPlayerIdJson(OfferJson, PlayerId);
+
+			std::string SdpAnsi;
+			SDP.ToString(&SdpAnsi);
+			FString SdpStr = UE::PixelStreaming::ToString(SdpAnsi);
+			OfferJson->SetStringField(TEXT("sdp"), SdpStr);
+
+#if WITH_EDITOR
+			// Default hover preference needs to be sent with the offer to prevent error messages on older infrastructure
+			OfferJson->SetStringField(TEXT("defaultToHover"), TEXT("true"));
+#endif // WITH_EDITOR
+
+			UE_LOG(LogMockPixelStreamingSS, Log, TEXT("Sending player=%s \"offer\" to SS %s"), *PlayerId, *Url);
+			UE_LOG(LogMockPixelStreamingSS, Verbose, TEXT("SDP offer\n%s"), *SdpStr);
+
+			SendMessage(UE::PixelStreaming::ToString(OfferJson, false));
+		}
+
+		virtual void SendAnswer(FPixelStreamingPlayerId PlayerId, const webrtc::SessionDescriptionInterface& SDP) override
+		{
+			FJsonObjectPtr AnswerJson = MakeShared<FJsonObject>();
+			AnswerJson->SetStringField(TEXT("type"), TEXT("answer"));
+			SetPlayerIdJson(AnswerJson, PlayerId);
+
+			std::string SdpAnsi;
+
+			if (SDP.ToString(&SdpAnsi))
+			{
+				FString SdpStr = UE::PixelStreaming::ToString(SdpAnsi);
+				AnswerJson->SetStringField(TEXT("sdp"), SdpStr);
+
+				UE_LOG(LogMockPixelStreamingSS, Log, TEXT("Sending player=%s \"answer\" to SS %s"), *PlayerId, *Url);
+				UE_LOG(LogMockPixelStreamingSS, Verbose, TEXT("SDP answer\n%s"), *SdpStr);
+
+				SendMessage(UE::PixelStreaming::ToString(AnswerJson, false));
+			}
+		}
+
+		virtual void SendIceCandidate(FPixelStreamingPlayerId PlayerId, const webrtc::IceCandidateInterface& IceCandidate) override
+		{
+			FJsonObjectPtr IceCandidateJson = MakeShared<FJsonObject>();
+
+			IceCandidateJson->SetStringField(TEXT("type"), TEXT("iceCandidate"));
+			SetPlayerIdJson(IceCandidateJson, PlayerId);
+
+			FJsonObjectPtr CandidateJson = MakeShared<FJsonObject>();
+			CandidateJson->SetStringField(TEXT("sdpMid"), IceCandidate.sdp_mid().c_str());
+			CandidateJson->SetNumberField(TEXT("sdpMLineIndex"), static_cast<double>(IceCandidate.sdp_mline_index()));
+
+			std::string CandidateAnsi;
+			if (IceCandidate.ToString(&CandidateAnsi))
+			{
+				FString CandidateStr = UE::PixelStreaming::ToString(CandidateAnsi);
+				CandidateJson->SetStringField(TEXT("candidate"), CandidateStr);
+				IceCandidateJson->SetObjectField(TEXT("candidate"), CandidateJson);
+				UE_LOG(LogMockPixelStreamingSS, Verbose, TEXT("-> SS: iceCandidate\n%s"), *UE::PixelStreaming::ToString(IceCandidateJson));
+				SendMessage(UE::PixelStreaming::ToString(IceCandidateJson, false));
+			}
+		}
+
+		virtual void SendDisconnectPlayer(FPixelStreamingPlayerId PlayerId, const FString& Reason) override
+		{
+			FJsonObjectPtr JsonMsg = MakeShared<FJsonObject>();
+
+			JsonMsg->SetStringField(TEXT("type"), TEXT("disconnectPlayer"));
+			SetPlayerIdJson(JsonMsg, PlayerId);
+			JsonMsg->SetStringField(TEXT("reason"), Reason);
+
+			FString Msg = UE::PixelStreaming::ToString(JsonMsg, false);
+			UE_LOG(LogMockPixelStreamingSS, Verbose, TEXT("Sending player=%s \"disconnectPlayer\" to SS %s"), *PlayerId, *Url);
+
+			SendMessage(Msg);
+		}
+
+		virtual void SendOffer(const webrtc::SessionDescriptionInterface& SDP) override
+		{
+			FJsonObjectPtr OfferJson = MakeShared<FJsonObject>();
+			OfferJson->SetStringField(TEXT("type"), TEXT("offer"));
+
+			FString SdpStr = UE::PixelStreaming::ToString(SDP);
+			OfferJson->SetStringField(TEXT("sdp"), UE::PixelStreaming::ToString(SDP));
+
+			UE_LOG(LogMockPixelStreamingSS, Verbose, TEXT("-> SS: offer\n%s"), *SdpStr);
+
+			SendMessage(UE::PixelStreaming::ToString(OfferJson, false));
+		}
+
+		virtual void SendAnswer(const webrtc::SessionDescriptionInterface& SDP) override
+		{
+			FJsonObjectPtr AnswerJson = MakeShared<FJsonObject>();
+			AnswerJson->SetStringField(TEXT("type"), TEXT("answer"));
+
+			FString SdpStr = UE::PixelStreaming::ToString(SDP);
+			AnswerJson->SetStringField(TEXT("sdp"), UE::PixelStreaming::ToString(SDP));
+
+			UE_LOG(LogMockPixelStreamingSS, Verbose, TEXT("-> SS: answer\n%s"), *SdpStr);
+
+			SendMessage(UE::PixelStreaming::ToString(AnswerJson, false));
+		}
+
+		virtual void SendIceCandidate(const webrtc::IceCandidateInterface& IceCandidate) override
+		{
+			FJsonObjectPtr IceCandidateJson = MakeShared<FJsonObject>();
+
+			IceCandidateJson->SetStringField(TEXT("type"), TEXT("iceCandidate"));
+
+			FJsonObjectPtr CandidateJson = MakeShared<FJsonObject>();
+			CandidateJson->SetStringField(TEXT("sdpMid"), IceCandidate.sdp_mid().c_str());
+			CandidateJson->SetNumberField(TEXT("sdpMLineIndex"), static_cast<double>(IceCandidate.sdp_mline_index()));
+
+			std::string CandidateAnsi;
+			verifyf(IceCandidate.ToString(&CandidateAnsi), TEXT("Failed to serialize IceCandidate"));
+			FString CandidateStr = UE::PixelStreaming::ToString(CandidateAnsi);
+			CandidateJson->SetStringField(TEXT("candidate"), CandidateStr);
+
+			IceCandidateJson->SetObjectField(TEXT("candidate"), CandidateJson);
+
+			UE_LOG(LogMockPixelStreamingSS, Verbose, TEXT("-> SS: ice-candidate\n%s"), *UE::PixelStreaming::ToString(IceCandidateJson));
+
+			SendMessage(UE::PixelStreaming::ToString(IceCandidateJson, false));
+		}
+
+		virtual void SetKeepAlive(bool bKeepAlive) override
+		{
+			if (bKeepAliveEnabled == bKeepAlive)
+			{
+				return;
+			}
+
+			bKeepAliveEnabled = bKeepAlive;
+		}
+
+		virtual void SetAutoReconnect(bool bAutoReconnect) override
+		{
+			if (bAutoReconnectEnabled == bAutoReconnect)
+			{
+				return;
+			}
+
+			bAutoReconnectEnabled = bAutoReconnect;
+		}
+
+		virtual void SendMessage(const FString& Msg) override
+		{
+			if (WebSocket && WebSocket->IsConnected())
+			{
+				WebSocket->Send(Msg);
+			}
+		}
+
+		TSharedPtr<FMockWebSocket> WebSocket;
+
+	private:
+		using FJsonObjectPtr = TSharedPtr<FJsonObject>;
+		void RegisterHandler(const FString& messageType, const TFunction<void(FJsonObjectPtr)>& handler)
+		{
+			MessageHandlers.Add(messageType, handler);
+		}
+
+		bool GetPlayerIdJson(const FJsonObjectPtr& Json, FPixelStreamingPlayerId& OutPlayerId, const FString& FieldId = TEXT("playerId"))
+		{
+			uint32 PlayerIdInt;
+			if (Json->TryGetNumberField(FieldId, PlayerIdInt))
+			{
+				OutPlayerId = ToPlayerId(PlayerIdInt);
+				return true;
+			}
+
+			return false;
+		}
+
+		void SetPlayerIdJson(FJsonObjectPtr& JsonObject, FPixelStreamingPlayerId PlayerId)
+		{
+			int32 PlayerIdAsInt = PlayerIdToInt(PlayerId);
+			JsonObject->SetNumberField(TEXT("playerId"), PlayerIdAsInt);
+		}
+
+		void PlayerError(FPixelStreamingPlayerId PlayerId, const FString& Msg)
+		{
+			UE_LOG(LogMockPixelStreamingSS, Error, TEXT("player %s: %s"), *PlayerId, *Msg);
+			SendDisconnectPlayer(PlayerId, Msg);
+		}
+
+		template <typename FmtType, typename... T>
+		void PlayerError(FPixelStreamingPlayerId PlayerId, const FmtType& Msg, T... args)
+		{
+			const FString FormattedMsg = FString::Printf(Msg, args...);
+			PlayerError(PlayerId, FormattedMsg);
+		}
+
+		IPixelStreamingSignallingConnectionObserver& Observer;
+		FString StreamerId;
+		FString Url;
+
+		FDelegateHandle OnConnectedHandle;
+		FDelegateHandle OnConnectionErrorHandle;
+		FDelegateHandle OnClosedHandle;
+		FDelegateHandle OnMessageHandle;
+
+		bool bAutoReconnectEnabled = true;
+		bool bKeepAliveEnabled = true;
+		/** Handle for efficient management of KeepAlive timer */
+		FTimerHandle TimerHandle_KeepAlive;
+		FTimerHandle TimerHandle_Reconnect;
+		const float KEEP_ALIVE_INTERVAL = 60.0f;
+
+		TMap<FString, TFunction<void(FJsonObjectPtr)>> MessageHandlers;
+	};
+
 	class FMockSessionDescriptionInterface : public webrtc::SessionDescriptionInterface
 	{
 	public:
@@ -145,13 +701,11 @@ namespace
 		cricket::Candidate MockCandidate;
 	};
 
-	BEGIN_DEFINE_SPEC(FSignallingServerConnectionSpec, "System.Plugins.PixelStreaming.SignallingServerConnection",  EAutomationTestFlags::EditorContext | EAutomationTestFlags::ClientContext | EAutomationTestFlags::ProductFilter)
-	TUniquePtr<FPixelStreamingSignallingConnection> SignallingServerConnection; // the object being tested
+	BEGIN_DEFINE_SPEC(FSignallingServerConnectionSpec, "System.Plugins.PixelStreaming.SignallingServerConnection", EAutomationTestFlags::EditorContext | EAutomationTestFlags::ClientContext | EAutomationTestFlags::ProductFilter)
+	TUniquePtr<FMockSignallingConnection> SignallingServerConnection; // the object being tested
 	FTestSSConnectionObserver Observer;
-	FPixelStreamingSignallingConnection::FWebSocketFactory WebSocketFactory;
-	TSharedPtr<FMockWebSocket> SocketEndpoint; // the socket created for the SignallingServerConnection. Used to fake send data.
-	FDelegateHandle OnMessageDelegateHandle;   // so we can remove the delegate after tests
-	TSharedPtr<FJsonObject> LastJsonMsg;	   // the parsed json of the last message received from the SignallingServerConnection
+	FDelegateHandle OnMessageDelegateHandle; // so we can remove the delegate after tests
+	TSharedPtr<FJsonObject> LastJsonMsg;	 // the parsed JsonMsg of the last message received from the SignallingServerConnection
 	const FString ServerID = "MOCK SERVER";
 	const FPixelStreamingPlayerId MockPlayerId = "MOCK";
 	const FMockSessionDescriptionInterface MockSDP;
@@ -162,23 +716,17 @@ namespace
 	{
 		Describe("Mocked connection tests", [this]() {
 			BeforeEach([this]() {
-				WebSocketFactory = [this](const FString&) {
-					SocketEndpoint = MakeShared<FMockWebSocket>();
-					return SocketEndpoint;
-				};
-				SignallingServerConnection = MakeUnique<FPixelStreamingSignallingConnection>(WebSocketFactory, Observer, ServerID);
+				SignallingServerConnection = MakeUnique<FMockSignallingConnection>(Observer, ServerID);
 				SignallingServerConnection->TryConnect("Fake Url");
-				OnMessageDelegateHandle = SocketEndpoint->OnMockResponseEvent.AddLambda([this](const FString& message) {
+				OnMessageDelegateHandle = SignallingServerConnection->WebSocket->OnMockResponseEvent.AddLambda([this](const FString& message) {
 					OnMessageReceived(message);
 				});
 			});
 
 			AfterEach([this]() {
 				LastJsonMsg.Reset();
-				SocketEndpoint->OnMockResponseEvent.Remove(OnMessageDelegateHandle);
-				SocketEndpoint.Reset();
+				SignallingServerConnection->WebSocket->OnMockResponseEvent.Remove(OnMessageDelegateHandle);
 				SignallingServerConnection.Reset();
-				WebSocketFactory.Reset();
 			});
 
 			It("should be connected", [this]() {
@@ -212,19 +760,23 @@ namespace
 				SignallingServerConnection->SendDisconnectPlayer(MockPlayerId, "No Reason");
 			});
 
+			It("should be able to send a custom message", [this]() {
+				SignallingServerConnection->SendMessage(R"({ "type": "custom message"})");
+			});
+
 			It("should reply to an identity request", [this]() {
-				SocketEndpoint->MockSend(R"({"type" : "identify"})");
+				SignallingServerConnection->WebSocket->MockSend(R"({"type" : "identify"})");
 				TestEqual("Type", LastJsonMsg->GetStringField("type"), "endpointId");
 				TestEqual("Id", LastJsonMsg->GetStringField("id"), ServerID);
 			});
 
 			It("should receive configurations", [this]() {
-				SocketEndpoint->MockSend(R"({ "type": "config", "peerConnectionOptions": {} })");
+				SignallingServerConnection->WebSocket->MockSend(R"({ "type": "config", "peerConnectionOptions": {} })");
 				TestTrue("Config IsSet", Observer.SetConfig.IsSet());
 			});
 
 			It("should be able to receive offers with no player id", [this]() {
-				SocketEndpoint->MockSend(R"({ "type": "offer", "sdp": "sdp goes here" })");
+				SignallingServerConnection->WebSocket->MockSend(R"({ "type": "offer", "sdp": "sdp goes here" })");
 				if (TestTrue("SDP IsSet", Observer.SetSDP.IsSet()))
 				{
 					TestEqual("Offer type", Observer.SetSDP.GetValue().Type, webrtc::SdpType::kOffer);
@@ -234,7 +786,7 @@ namespace
 			});
 
 			It("should be able to receive offers with a player id", [this]() {
-				SocketEndpoint->MockSend(R"({ "type": "offer", "playerId": "101", "sdp": "sdp goes here" })");
+				SignallingServerConnection->WebSocket->MockSend(R"({ "type": "offer", "playerId": "101", "sdp": "sdp goes here" })");
 				if (TestTrue("SDP IsSet", Observer.SetSDP.IsSet()))
 				{
 					TestEqual("Offer type", Observer.SetSDP.GetValue().Type, webrtc::SdpType::kOffer);
@@ -247,7 +799,7 @@ namespace
 			});
 
 			It("should be able to receive an answer", [this]() {
-				SocketEndpoint->MockSend(R"({ "type": "answer", "playerId": "102", "sdp": "sdp goes here" })");
+				SignallingServerConnection->WebSocket->MockSend(R"({ "type": "answer", "playerId": "102", "sdp": "sdp goes here" })");
 				if (TestTrue("SDP IsSet", Observer.SetSDP.IsSet()))
 				{
 					TestEqual("Answer type", Observer.SetSDP.GetValue().Type, webrtc::SdpType::kAnswer);
@@ -260,7 +812,7 @@ namespace
 			});
 
 			It("should be able to receive an ice candidate with no player id", [this]() {
-				SocketEndpoint->MockSend(R"({ "type": "iceCandidate", "candidate": { "sdpMid": "mid", "sdpMlineIndex": 10, "candidate": "sdp string" } })");
+				SignallingServerConnection->WebSocket->MockSend(R"({ "type": "iceCandidate", "candidate": { "sdpMid": "mid", "sdpMlineIndex": 10, "candidate": "sdp string" } })");
 				if (TestTrue("Ice Candidate IsSet", Observer.SetIceCandidate.IsSet()))
 				{
 					TestEqual("sdp mid", Observer.SetIceCandidate.GetValue().SdpMid, "mid");
@@ -271,7 +823,7 @@ namespace
 			});
 
 			It("should be able to receive an ice candidate with a player id", [this]() {
-				SocketEndpoint->MockSend(R"({ "type": "iceCandidate", "playerId": "103", "candidate": { "sdpMid": "mid", "sdpMlineIndex": 10, "candidate": "sdp string" } })");
+				SignallingServerConnection->WebSocket->MockSend(R"({ "type": "iceCandidate", "playerId": "103", "candidate": { "sdpMid": "mid", "sdpMlineIndex": 10, "candidate": "sdp string" } })");
 				if (TestTrue("Ice Candidate IsSet", Observer.SetIceCandidate.IsSet()))
 				{
 					TestEqual("sdp mid", Observer.SetIceCandidate.GetValue().SdpMid, "mid");
@@ -286,10 +838,10 @@ namespace
 		});
 	}
 
-	// used to store the parsed json response from the signalling server connection
+	// used to store the parsed JsonMsg response from the signalling server connection
 	void FSignallingServerConnectionSpec::OnMessageReceived(const FString& Response)
 	{
 		const auto JsonReader = TJsonReaderFactory<TCHAR>::Create(Response);
-		TestTrue("Deserialised json", FJsonSerializer::Deserialize(JsonReader, LastJsonMsg));
+		TestTrue("Deserialised JsonMsg", FJsonSerializer::Deserialize(JsonReader, LastJsonMsg));
 	}
 } // namespace
diff --git a/Engine/Plugins/Media/PixelStreaming/Source/PixelStreaming/Private/Tests/TestUtils.cpp b/Engine/Plugins/Media/PixelStreaming/Source/PixelStreaming/Private/Tests/TestUtils.cpp
--- a/Engine/Plugins/Media/PixelStreaming/Source/PixelStreaming/Private/Tests/TestUtils.cpp
+++ b/Engine/Plugins/Media/PixelStreaming/Source/PixelStreaming/Private/Tests/TestUtils.cpp
@@ -16,7 +16,7 @@ namespace UE::PixelStreaming
 
 	bool FSendSolidColorFrame::Update()
 	{
-		TSharedPtr<FPixelCaptureI420Buffer> Buffer = MakeShared<FPixelCaptureI420Buffer>(FrameConfig.Width, FrameConfig.Height);
+		TSharedPtr<FPixelCaptureBufferI420> Buffer = MakeShared<FPixelCaptureBufferI420>(FrameConfig.Width, FrameConfig.Height);
 
 		uint8_t* yData = Buffer->GetMutableDataY();
 		uint8_t* uData = Buffer->GetMutableDataU();
@@ -181,7 +181,7 @@ namespace UE::PixelStreaming
 			UE_LOG(LogPixelStreaming, Error, TEXT("Timed out waiting for streamer to connect to signalling server."));
 			return true;
 		}
-		return false;		
+		return false;
 	}
 
 	bool FWaitForPlayerConnectedOrTimeout::Update()
@@ -247,9 +247,14 @@ namespace UE::PixelStreaming
 
 	TSharedPtr<IPixelStreamingStreamer> CreateStreamer(int StreamerPort)
 	{
-		TSharedPtr<IPixelStreamingStreamer> OutStreamer = IPixelStreamingModule::Get().CreateStreamer("Mock Streamer");
+		FString StreamerName("Mock Streamer");
+		TSharedPtr<IPixelStreamingStreamer> OutStreamer = IPixelStreamingModule::Get().CreateStreamer(StreamerName);
 		OutStreamer->SetVideoInput(FPixelStreamingVideoInputBackBuffer::Create());
 		OutStreamer->SetSignallingServerURL(FString::Printf(TEXT("ws://127.0.0.1:%d"), StreamerPort));
+
+		TSharedPtr<IPixelStreamingSignallingConnection> SignallingConnection = MakeShared<FPixelStreamingSignallingConnection>(OutStreamer->GetSignallingConnectionObserver().Pin(), StreamerName);
+		SignallingConnection->SetAutoReconnect(true);
+		OutStreamer->SetSignallingConnection(SignallingConnection);
 		return OutStreamer;
 	}
 
diff --git a/Engine/Plugins/Media/PixelStreaming/Source/PixelStreaming/Private/Utils.h b/Engine/Plugins/Media/PixelStreaming/Source/PixelStreaming/Private/Utils.h
--- a/Engine/Plugins/Media/PixelStreaming/Source/PixelStreaming/Private/Utils.h
+++ b/Engine/Plugins/Media/PixelStreaming/Source/PixelStreaming/Private/Utils.h
@@ -78,20 +78,7 @@ namespace UE::PixelStreaming
 			FPlatformProcess::ReturnSynchEventToPool(TaskEvent);
 		}
 	}
-#if WEBRTC_VERSION == 84
-	inline webrtc::SdpVideoFormat CreateH264Format(webrtc::H264::Profile profile, webrtc::H264::Level level)
-	{
-		const absl::optional<std::string> ProfileString =
-			webrtc::H264::ProfileLevelIdToString(webrtc::H264::ProfileLevelId(profile, level));
-		check(ProfileString);
-		return webrtc::SdpVideoFormat(
-			cricket::kH264CodecName,
-			{ { cricket::kH264FmtpProfileLevelId, *ProfileString },
-				{ cricket::kH264FmtpLevelAsymmetryAllowed, "1" },
-				{ cricket::kH264FmtpPacketizationMode, "1" } });
 
-	}
-#elif WEBRTC_VERSION == 96
 	inline webrtc::SdpVideoFormat CreateH264Format(webrtc::H264Profile profile, webrtc::H264Level level)
 	{
 		const absl::optional<std::string> ProfileString =
@@ -102,9 +89,8 @@ namespace UE::PixelStreaming
 			{ { cricket::kH264FmtpProfileLevelId, *ProfileString },
 				{ cricket::kH264FmtpLevelAsymmetryAllowed, "1" },
 				{ cricket::kH264FmtpPacketizationMode, "1" } });
-
 	}
-#endif
+
 	inline void MemCpyStride(void* Dest, const void* Src, size_t DestStride, size_t SrcStride, size_t Height)
 	{
 		char* DestPtr = static_cast<char*>(Dest);
@@ -118,11 +104,7 @@ namespace UE::PixelStreaming
 
 	inline size_t SerializeToBuffer(rtc::CopyOnWriteBuffer& Buffer, size_t Pos, const void* Data, size_t DataSize)
 	{
-#if WEBRTC_VERSION == 84
-		FMemory::Memcpy(&Buffer[Pos], Data, DataSize);
-#elif WEBRTC_VERSION == 96
 		FMemory::Memcpy(Buffer.MutableData() + Pos, Data, DataSize);
-#endif
 		return Pos + DataSize;
 	}
 
diff --git a/Engine/Plugins/Media/PixelStreaming/Source/PixelStreaming/Private/UtilsRender.h b/Engine/Plugins/Media/PixelStreaming/Source/PixelStreaming/Private/UtilsRender.h
--- a/Engine/Plugins/Media/PixelStreaming/Source/PixelStreaming/Private/UtilsRender.h
+++ b/Engine/Plugins/Media/PixelStreaming/Source/PixelStreaming/Private/UtilsRender.h
@@ -9,6 +9,9 @@
 #include "ScreenPass.h"
 #include "ScreenRendering.h"
 #include "Utils.h"
+#include "MediaShaders.h"
+#include "RHI.h"
+#include "Runtime/Renderer/Private/ScreenPass.h"
 
 namespace UE::PixelStreaming
 {
@@ -105,6 +108,68 @@ namespace UE::PixelStreaming
 		}
 	}
 
+	inline void CopyTextureRDG(FRHICommandListImmediate& RHICmdList, FTextureRHIRef SourceTexture, FTextureRHIRef DestTexture)
+	{
+		FRDGBuilder GraphBuilder(RHICmdList);
+
+		// Register an external RDG rexture from the source texture
+		FRDGTexture* InputTexture = GraphBuilder.RegisterExternalTexture(CreateRenderTarget(SourceTexture, TEXT("PixelCaptureCopySourceTexture")));
+
+		// Register an external RDG texture from the output buffer
+		FRDGTexture* OutputTexture = GraphBuilder.RegisterExternalTexture(CreateRenderTarget(DestTexture, TEXT("PixelCaptureCopyDestTexture")));
+
+		if (InputTexture->Desc.Format == OutputTexture->Desc.Format &&
+			InputTexture->Desc.Extent.X == OutputTexture->Desc.Extent.X &&
+			InputTexture->Desc.Extent.Y == OutputTexture->Desc.Extent.Y)
+		{
+			// The formats are the same and size are the same. simple copy
+			AddDrawTexturePass(
+				GraphBuilder,
+				GetGlobalShaderMap(GMaxRHIFeatureLevel),
+				InputTexture,
+				OutputTexture,
+				FRDGDrawTextureInfo()
+			);
+		}
+		else
+		{
+			// The formats or size differ to pixel shader stuff
+			//Configure source/output viewport to get the right UV scaling from source texture to output texture
+			FScreenPassTextureViewport InputViewport(InputTexture);
+			FScreenPassTextureViewport OutputViewport(OutputTexture);
+
+			FGlobalShaderMap* GlobalShaderMap = GetGlobalShaderMap(GMaxRHIFeatureLevel);
+			TShaderMapRef<FScreenPassVS> VertexShader(GlobalShaderMap);
+
+			// In cases where texture is converted from a format that doesn't have A channel, we want to force set it to 1.
+			int32 ConversionOperation = 0; // None
+			FModifyAlphaSwizzleRgbaPS::FPermutationDomain PermutationVector;
+			PermutationVector.Set<FModifyAlphaSwizzleRgbaPS::FConversionOp>(ConversionOperation);
+
+			// Rectangle area to use from source
+			const FIntRect ViewRect(FIntPoint(0, 0), InputTexture->Desc.Extent);
+
+			//Dummy ViewFamily/ViewInfo created to use built in Draw Screen/Texture Pass
+			FSceneViewFamily ViewFamily(FSceneViewFamily::ConstructionValues(nullptr, nullptr, FEngineShowFlags(ESFIM_Game))
+				.SetTime(FGameTime())
+				.SetGammaCorrection(1.0f));
+			FSceneViewInitOptions ViewInitOptions;
+			ViewInitOptions.ViewFamily = &ViewFamily;
+			ViewInitOptions.SetViewRectangle(ViewRect);
+			ViewInitOptions.ViewOrigin = FVector::ZeroVector;
+			ViewInitOptions.ViewRotationMatrix = FMatrix::Identity;
+			ViewInitOptions.ProjectionMatrix = FMatrix::Identity;
+			FViewInfo ViewInfo = FViewInfo(ViewInitOptions);
+
+			TShaderMapRef<FModifyAlphaSwizzleRgbaPS> PixelShader(GlobalShaderMap, PermutationVector);
+			FModifyAlphaSwizzleRgbaPS::FParameters* Parameters = PixelShader->AllocateAndSetParameters(GraphBuilder, InputTexture, OutputTexture);
+			AddDrawScreenPass(GraphBuilder, RDG_EVENT_NAME("PixelCapturerSwizzle"), ViewInfo, OutputViewport, InputViewport, VertexShader, PixelShader, Parameters);
+		}
+		GraphBuilder.Execute();
+	}
+
+
+
 	inline void CopyTextureToReadbackTexture(FTextureRHIRef SourceTexture, TSharedPtr<FRHIGPUTextureReadback> GPUTextureReadback, void* OutBuffer)
 	{
 		FRHICommandListImmediate& RHICmdList = FRHICommandListExecutor::GetImmediateCommandList();
diff --git a/Engine/Plugins/Media/PixelStreaming/Source/PixelStreaming/Private/VideoDecoderFactory.cpp b/Engine/Plugins/Media/PixelStreaming/Source/PixelStreaming/Private/VideoDecoderFactory.cpp
--- a/Engine/Plugins/Media/PixelStreaming/Source/PixelStreaming/Private/VideoDecoderFactory.cpp
+++ b/Engine/Plugins/Media/PixelStreaming/Source/PixelStreaming/Private/VideoDecoderFactory.cpp
@@ -4,21 +4,96 @@
 #include "VideoDecoderStub.h"
 #include "absl/strings/match.h"
 #include "VideoDecoderVPX.h"
+#include "VideoDecoderH265.h"
+#include "PixelStreamingPrivate.h"
+#include "Settings.h"
 
 namespace UE::PixelStreaming
 {
+	// the list of each individual codec we have decoder support for (the order of this array is preference order after the selected codec)
+	const TArray<EPixelStreamingCodec> SupportedDecoderCodecList{ EPixelStreamingCodec::VP8, EPixelStreamingCodec::VP9, /* EPixelStreamingCodec::H264, */ EPixelStreamingCodec::H265 };
+
+	// mapping of codec to a list of video formats
+	// done this way so we can order the list of formats based on selected codec in GetSupportedFormats
+	TMap<EPixelStreamingCodec, std::vector<webrtc::SdpVideoFormat>> CreateSupportedDecoderFormatMap()
+	{
+		TMap<EPixelStreamingCodec, std::vector<webrtc::SdpVideoFormat>> Codecs;
+		for (auto& Codec : SupportedDecoderCodecList)
+		{
+			Codecs.Add(Codec);
+		}
+
+		Codecs[EPixelStreamingCodec::VP8].push_back(webrtc::SdpVideoFormat(cricket::kVp8CodecName));
+		Codecs[EPixelStreamingCodec::VP9].push_back(webrtc::SdpVideoFormat(cricket::kVp9CodecName));
+		// Codecs[EPixelStreamingCodec::H264].push_back(UE::PixelStreaming::CreateH264Format(webrtc::H264Profile::kProfileConstrainedBaseline, webrtc::H264Level::kLevel3_1));
+		// Codecs[EPixelStreamingCodec::H264].push_back(UE::PixelStreaming::CreateH264Format(webrtc::H264Profile::kProfileBaseline, webrtc::H264Level::kLevel3_1));
+		Codecs[EPixelStreamingCodec::H265].push_back(webrtc::SdpVideoFormat(cricket::kH265CodecName));
+
+		return Codecs;
+	}
+
 	std::vector<webrtc::SdpVideoFormat> FVideoDecoderFactory::GetSupportedFormats() const
 	{
-		std::vector<webrtc::SdpVideoFormat> video_formats;
-		video_formats.push_back(webrtc::SdpVideoFormat(cricket::kVp8CodecName));
-		video_formats.push_back(webrtc::SdpVideoFormat(cricket::kVp9CodecName));
-		// uncomment these when we get a h264 decoder
-// #if WEBRTC_VERSION == 84
-// 		video_formats.push_back(CreateH264Format(webrtc::H264::kProfileConstrainedBaseline, webrtc::H264::kLevel3_1));
-// #elif WEBRTC_VERSION == 96
-// 		video_formats.push_back(CreateH264Format(webrtc::H264Profile::kProfileConstrainedBaseline, webrtc::H264Level::kLevel3_1));
-// #endif
-		return video_formats;
+		// static so we dont create the list every time this is called since the list will not change during runtime.
+		static TMap<EPixelStreamingCodec, std::vector<webrtc::SdpVideoFormat>> CodecMap = CreateSupportedDecoderFormatMap();
+
+		// since this method is const we need to store this state statically. it means all instances will share this state
+		// but that actually works in our favor since we're describing more about the plugin state than the actual
+		// instance of this factory.
+		static std::vector<webrtc::SdpVideoFormat> SupportedFormats;
+
+		// If we are not negotiating codecs simply return just the one codec that is selected in UE
+		if (!Settings::CVarPixelStreamingWebRTCNegotiateCodecs.GetValueOnAnyThread())
+		{
+			const EPixelStreamingCodec SelectedCodec = UE::PixelStreaming::Settings::GetSelectedCodec();
+			if (CodecMap.Contains(SelectedCodec))
+			{
+				for (auto& Format : CodecMap[SelectedCodec])
+				{
+					SupportedFormats.push_back(Format);
+				}
+				return SupportedFormats;
+			}
+			else
+			{
+				UE_LOG(LogPixelStreaming, Error, TEXT("Selected codec was not a supported codec, falling back to negotiating codecs..."));
+			}
+		}
+
+		static EPixelStreamingCodec LastSelectedCodec = EPixelStreamingCodec::Invalid;
+
+		const EPixelStreamingCodec SelectedCodec = UE::PixelStreaming::Settings::GetSelectedCodec();
+		if (LastSelectedCodec != SelectedCodec)
+		{
+			// build a new format list
+			LastSelectedCodec = SelectedCodec;
+			SupportedFormats.clear();
+
+			// order the codecs so the selected is first
+			TArray<EPixelStreamingCodec> OrderedCodecList;
+			OrderedCodecList.Add(SelectedCodec);
+			for (auto& SupportedCodec : SupportedDecoderCodecList)
+			{
+				if (SupportedCodec != SelectedCodec)
+				{
+					OrderedCodecList.Add(SupportedCodec);
+				}
+			}
+
+			// now just add each of the formats in order
+			for (auto& Codec : OrderedCodecList)
+			{
+				if (CodecMap.Contains(Codec))
+				{
+					for (auto& Format : CodecMap[Codec])
+					{
+						SupportedFormats.push_back(Format);
+					}
+				}
+			}
+		}
+
+		return SupportedFormats;
 	}
 
 	std::unique_ptr<webrtc::VideoDecoder> FVideoDecoderFactory::CreateVideoDecoder(const webrtc::SdpVideoFormat& format)
@@ -31,6 +106,10 @@ namespace UE::PixelStreaming
 		{
 			return std::make_unique<VideoDecoderVPX>(9);
 		}
+		else if (absl::EqualsIgnoreCase(format.name, cricket::kH265CodecName))
+		{
+			return std::make_unique<VideoDecoderH265>();
+		}
 		return std::make_unique<FVideoDecoderStub>();
 	}
 } // namespace UE::PixelStreaming
diff --git a/Engine/Plugins/Media/PixelStreaming/Source/PixelStreaming/Private/VideoDecoderH265.cpp b/Engine/Plugins/Media/PixelStreaming/Source/PixelStreaming/Private/VideoDecoderH265.cpp
new file mode 100644
--- /dev/null
+++ b/Engine/Plugins/Media/PixelStreaming/Source/PixelStreaming/Private/VideoDecoderH265.cpp
@@ -0,0 +1,89 @@
+// Copyright Epic Games, Inc. All Rights Reserved.
+
+#include "VideoDecoderH265.h"
+
+#include "FrameBufferRHI.h"
+#include "PixelStreamingPrivate.h"
+#include "PixelStreamingTrace.h"
+
+namespace UE::PixelStreaming
+{
+	// HACK (aidan.possemiers) the API surface SendPacket surface wants a SharedPtr for the encoded data but WebRTC already owns that and the types are incompatible
+	struct FFakeDeleter
+	{
+		void operator()(uint8* Object) const
+		{
+		}
+	};
+
+	VideoDecoderH265::VideoDecoderH265()
+	{
+		FVideoDecoderConfigH265 DecoderConfig;
+		Decoder = FVideoDecoder::CreateChecked<FVideoResourceRHI>(FAVDevice::GetHardwareDevice(), DecoderConfig);
+	}
+
+	bool VideoDecoderH265::Configure(const Settings& settings)
+	{
+		return true;
+	}
+
+	int32 VideoDecoderH265::Decode(const webrtc::EncodedImage& input_image, bool missing_frames, int64_t render_time_ms)
+	{
+		const int64 TimestampDecodeStart = rtc::TimeMillis();
+		TRACE_CPUPROFILER_EVENT_SCOPE_ON_CHANNEL_STR("PixelStreaming Decoding H265 Video", PixelStreamingChannel);
+
+		FAVResult Result = Decoder->SendPacket(FVideoPacket(
+			MakeShareable<uint8>(input_image.GetEncodedData()->data(), FFakeDeleter()),
+			input_image.GetEncodedData()->size(),
+			input_image.Timestamp(),
+			FrameCount++,
+			input_image.qp_,
+			input_image._frameType == webrtc::VideoFrameType::kVideoFrameKey));
+
+		if (Output != nullptr && Result.IsNotError())
+		{
+			FResolvableVideoResourceRHI DecoderResource;
+			const FAVResult DecodeResult = Decoder->ReceiveFrame(DecoderResource);
+
+			if (DecodeResult.IsSuccess())
+			{
+				rtc::scoped_refptr<webrtc::VideoFrameBuffer> FrameBuffer = new rtc::RefCountedObject<UE::PixelStreaming::FFrameBufferRHI>(DecoderResource);
+				check(FrameBuffer->width() != 0 && FrameBuffer->height() != 0); // TODO we should probably check that we are getting the frame back that we are expecting
+
+				webrtc::VideoFrame Frame = webrtc::VideoFrame::Builder()
+											   .set_video_frame_buffer(FrameBuffer)
+											   .set_timestamp_rtp(input_image.Timestamp()) // NOTE This timestamp is load bearing as WebRTC stores the frames in a map with this as the index
+											   .set_color_space(input_image.ColorSpace())
+											   .build();
+
+				Output->Decoded(Frame, rtc::TimeMillis() - TimestampDecodeStart, input_image.qp_);
+
+				return WEBRTC_VIDEO_CODEC_OK;
+			}
+		}
+		else
+		{
+			UE_LOG(LogTemp, Warning, TEXT("VideoDecoderH265::Decode FAILED"));
+
+			return WEBRTC_VIDEO_CODEC_OK_REQUEST_KEYFRAME;
+		}
+
+		UE_LOG(LogTemp, Error, TEXT("VideoDecoderH265::Decode ERROR"));
+
+		return WEBRTC_VIDEO_CODEC_ERROR;
+	}
+
+	int32 VideoDecoderH265::RegisterDecodeCompleteCallback(webrtc::DecodedImageCallback* callback)
+	{
+		Output = callback;
+
+		return 0;
+	}
+
+	int32 VideoDecoderH265::Release()
+	{
+		Decoder.Reset();
+
+		return 0;
+	}
+} // namespace UE::PixelStreaming
diff --git a/Engine/Plugins/Media/PixelStreaming/Source/PixelStreaming/Private/VideoDecoderH265.h b/Engine/Plugins/Media/PixelStreaming/Source/PixelStreaming/Private/VideoDecoderH265.h
new file mode 100644
--- /dev/null
+++ b/Engine/Plugins/Media/PixelStreaming/Source/PixelStreaming/Private/VideoDecoderH265.h
@@ -0,0 +1,31 @@
+// Copyright Epic Games, Inc. All Rights Reserved.
+
+#pragma once
+
+#include "WebRTCIncludes.h"
+
+#include "Video/Decoders/Configs/VideoDecoderConfigH265.h"
+#include "Video/Resources/VideoResourceRHI.h"
+
+namespace UE::PixelStreaming
+{
+	class VideoDecoderH265 : public webrtc::VideoDecoder
+	{
+	public:
+		VideoDecoderH265();
+		virtual ~VideoDecoderH265() = default;
+
+		virtual bool Configure(const Settings& settings) override;
+		virtual int32 Decode(const webrtc::EncodedImage& input_image, bool missing_frames, int64_t render_time_ms) override;
+		virtual int32 RegisterDecodeCompleteCallback(webrtc::DecodedImageCallback* callback) override;
+		virtual int32 Release() override;
+		virtual const char* ImplementationName() const { return "VideoDecoderH265"; }
+
+	private:
+		TSharedPtr<TVideoDecoder<FVideoResourceRHI, FVideoDecoderConfigH265>> Decoder;
+
+		webrtc::DecodedImageCallback* Output = nullptr;
+
+		uint32 FrameCount = 0;
+	};
+} // namespace UE::PixelStreaming
diff --git a/Engine/Plugins/Media/PixelStreaming/Source/PixelStreaming/Private/VideoDecoderStub.h b/Engine/Plugins/Media/PixelStreaming/Source/PixelStreaming/Private/VideoDecoderStub.h
--- a/Engine/Plugins/Media/PixelStreaming/Source/PixelStreaming/Private/VideoDecoderStub.h
+++ b/Engine/Plugins/Media/PixelStreaming/Source/PixelStreaming/Private/VideoDecoderStub.h
@@ -8,14 +8,9 @@ namespace UE::PixelStreaming
 	class FVideoDecoderStub : public webrtc::VideoDecoder
 	{
 	public:
-#if WEBRTC_VERSION == 84
-		virtual int32_t InitDecode(const webrtc::VideoCodec* codec_settings, int32_t number_of_cores) override { return 0; }
-#elif WEBRTC_VERSION == 96
 		virtual bool Configure(const Settings& settings) override { return false; }
-#endif
 		virtual int32_t Decode(const webrtc::EncodedImage& input_image, bool missing_frames, int64_t render_time_ms) override { return 0; }
 		virtual int32_t RegisterDecodeCompleteCallback(webrtc::DecodedImageCallback* callback) override { return 0; }
 		virtual int32_t Release() override { return 0; }
-
 	};
 } // namespace UE::PixelStreaming
diff --git a/Engine/Plugins/Media/PixelStreaming/Source/PixelStreaming/Private/VideoDecoderVPX.cpp b/Engine/Plugins/Media/PixelStreaming/Source/PixelStreaming/Private/VideoDecoderVPX.cpp
--- a/Engine/Plugins/Media/PixelStreaming/Source/PixelStreaming/Private/VideoDecoderVPX.cpp
+++ b/Engine/Plugins/Media/PixelStreaming/Source/PixelStreaming/Private/VideoDecoderVPX.cpp
@@ -2,7 +2,7 @@
 
 #include "VideoDecoderVPX.h"
 #include "PixelStreamingPrivate.h"
-
+#include "PixelStreamingTrace.h"
 
 namespace UE::PixelStreaming
 {
@@ -22,20 +22,14 @@ namespace UE::PixelStreaming
 		}
 	}
 
-#if WEBRTC_VERSION == 84
-	int32 VideoDecoderVPX::InitDecode(const webrtc::VideoCodec* codec_settings, int32 number_of_cores)
-	{
-		return VideoDecoder->InitDecode(codec_settings, number_of_cores);
-	}
-#elif WEBRTC_VERSION == 96
 	bool VideoDecoderVPX::Configure(const Settings& settings)
 	{
 		return VideoDecoder->Configure(settings);
 	}
-#endif
 
 	int32 VideoDecoderVPX::Decode(const webrtc::EncodedImage& input_image, bool missing_frames, int64_t render_time_ms)
 	{
+		TRACE_CPUPROFILER_EVENT_SCOPE_ON_CHANNEL_STR("PixelStreaming Decoding VPX Video", PixelStreamingChannel);
 		return VideoDecoder->Decode(input_image, missing_frames, render_time_ms);
 	}
 
diff --git a/Engine/Plugins/Media/PixelStreaming/Source/PixelStreaming/Private/VideoDecoderVPX.h b/Engine/Plugins/Media/PixelStreaming/Source/PixelStreaming/Private/VideoDecoderVPX.h
--- a/Engine/Plugins/Media/PixelStreaming/Source/PixelStreaming/Private/VideoDecoderVPX.h
+++ b/Engine/Plugins/Media/PixelStreaming/Source/PixelStreaming/Private/VideoDecoderVPX.h
@@ -12,11 +12,7 @@ namespace UE::PixelStreaming
 		VideoDecoderVPX(int Version);
 		virtual ~VideoDecoderVPX() = default;
 
-#if WEBRTC_VERSION == 84
-		virtual int32 InitDecode(const webrtc::VideoCodec* codec_settings, int32 number_of_cores) override;
-#elif WEBRTC_VERSION == 96
 		virtual bool Configure(const Settings& settings) override;
-#endif
 		virtual int32 Decode(const webrtc::EncodedImage& input_image, bool missing_frames, int64_t render_time_ms) override;
 		virtual int32 RegisterDecodeCompleteCallback(webrtc::DecodedImageCallback* callback) override;
 		virtual int32 Release() override;
diff --git a/Engine/Plugins/Media/PixelStreaming/Source/PixelStreaming/Private/VideoEncoderFactorySingleLayer.cpp b/Engine/Plugins/Media/PixelStreaming/Source/PixelStreaming/Private/VideoEncoderFactorySingleLayer.cpp
--- a/Engine/Plugins/Media/PixelStreaming/Source/PixelStreaming/Private/VideoEncoderFactorySingleLayer.cpp
+++ b/Engine/Plugins/Media/PixelStreaming/Source/PixelStreaming/Private/VideoEncoderFactorySingleLayer.cpp
@@ -2,43 +2,36 @@
 
 #include "VideoEncoderFactorySingleLayer.h"
 #include "absl/strings/match.h"
-#include "VideoEncoderSingleLayerH264.h"
+#include "VideoEncoderSingleLayerHardware.h"
 #include "VideoEncoderSingleLayerVPX.h"
 #include "Settings.h"
-#include "VideoEncoderFactory.h"
 #include "PixelStreamingPrivate.h"
 #include "Utils.h"
+#include "Stats.h"
 
-namespace
+namespace UE::PixelStreaming
 {
-	// the list of each individual codec we have encoder support for
-	const TArray<EPixelStreamingCodec> SupportedCodecList{ EPixelStreamingCodec::H264, EPixelStreamingCodec::VP8, EPixelStreamingCodec::VP9 };
+	// the list of each individual codec we have encoder support for (order of this array is preference order after selected codec)
+	const TArray<EPixelStreamingCodec> SupportedEncoderCodecList{ EPixelStreamingCodec::VP8, EPixelStreamingCodec::VP9, EPixelStreamingCodec::H264, EPixelStreamingCodec::H265 };
 
 	// mapping of codec to a list of video formats
 	// done this way so we can order the list of formats based on selected codec in GetSupportedFormats
-	TMap<EPixelStreamingCodec, std::vector<webrtc::SdpVideoFormat>> CreateSupportedFormatMap()
+	TMap<EPixelStreamingCodec, std::vector<webrtc::SdpVideoFormat>> CreateSupportedEncoderFormatMap()
 	{
 		TMap<EPixelStreamingCodec, std::vector<webrtc::SdpVideoFormat>> Codecs;
-		for (auto& Codec : SupportedCodecList)
+		for (auto& Codec : SupportedEncoderCodecList)
 		{
 			Codecs.Add(Codec);
 		}
 
 		Codecs[EPixelStreamingCodec::VP8].push_back(webrtc::SdpVideoFormat(cricket::kVp8CodecName));
 		Codecs[EPixelStreamingCodec::VP9].push_back(webrtc::SdpVideoFormat(cricket::kVp9CodecName));
-#if WEBRTC_VERSION == 84
-		Codecs[EPixelStreamingCodec::H264].push_back(UE::PixelStreaming::CreateH264Format(webrtc::H264::kProfileConstrainedBaseline, webrtc::H264::kLevel3_1));
-		Codecs[EPixelStreamingCodec::H264].push_back(UE::PixelStreaming::CreateH264Format(webrtc::H264::kProfileBaseline, webrtc::H264::kLevel3_1));
-#elif WEBRTC_VERSION == 96
 		Codecs[EPixelStreamingCodec::H264].push_back(UE::PixelStreaming::CreateH264Format(webrtc::H264Profile::kProfileConstrainedBaseline, webrtc::H264Level::kLevel3_1));
 		Codecs[EPixelStreamingCodec::H264].push_back(UE::PixelStreaming::CreateH264Format(webrtc::H264Profile::kProfileBaseline, webrtc::H264Level::kLevel3_1));
-#endif
+		Codecs[EPixelStreamingCodec::H265].push_back(webrtc::SdpVideoFormat(cricket::kH265CodecName));
+
 		return Codecs;
 	}
-} // namespace
-
-namespace UE::PixelStreaming
-{
 
 	// we want this method to return all the formats we have encoders for but the selected codecs formats should be first in the list.
 	// the reason for this is weird. when we receive video from another pixel streaming source, for some reason webrtc will query
@@ -48,7 +41,7 @@ namespace UE::PixelStreaming
 	{
 
 		// static so we dont create the list every time this is called since the list will not change during runtime.
-		static TMap<EPixelStreamingCodec, std::vector<webrtc::SdpVideoFormat>> CodecMap = CreateSupportedFormatMap();
+		static TMap<EPixelStreamingCodec, std::vector<webrtc::SdpVideoFormat>> CodecMap = CreateSupportedEncoderFormatMap();
 
 		// since this method is const we need to store this state statically. it means all instances will share this state
 		// but that actually works in our favor since we're describing more about the plugin state than the actual
@@ -56,7 +49,7 @@ namespace UE::PixelStreaming
 		static std::vector<webrtc::SdpVideoFormat> SupportedFormats;
 
 		// If we are not negotiating codecs simply return just the one codec that is selected in UE
-		if(!Settings::ShouldNegotiateCodecs())
+		if (!Settings::CVarPixelStreamingWebRTCNegotiateCodecs.GetValueOnAnyThread())
 		{
 			const EPixelStreamingCodec SelectedCodec = UE::PixelStreaming::Settings::GetSelectedCodec();
 			if (CodecMap.Contains(SelectedCodec))
@@ -85,7 +78,7 @@ namespace UE::PixelStreaming
 			// order the codecs so the selected is first
 			TArray<EPixelStreamingCodec> OrderedCodecList;
 			OrderedCodecList.Add(SelectedCodec);
-			for (auto& SupportedCodec : SupportedCodecList)
+			for (auto& SupportedCodec : SupportedEncoderCodecList)
 			{
 				if (SupportedCodec != SelectedCodec)
 				{
@@ -112,9 +105,6 @@ namespace UE::PixelStreaming
 	FVideoEncoderFactorySingleLayer::CodecInfo FVideoEncoderFactorySingleLayer::QueryVideoEncoder(const webrtc::SdpVideoFormat& format) const
 	{
 		webrtc::VideoEncoderFactory::CodecInfo CodecInfo;
-#if WEBRTC_VERSION == 84
-		CodecInfo.is_hardware_accelerated = true;
-#endif
 		CodecInfo.has_internal_source = false;
 		return CodecInfo;
 	}
@@ -123,104 +113,131 @@ namespace UE::PixelStreaming
 	{
 		if (absl::EqualsIgnoreCase(format.name, cricket::kVp8CodecName))
 		{
+			FStats::Get()->StoreApplicationStat(FStatData(FName(TEXT("Video Codec - VP8")), 1, 0));
 			return std::make_unique<FVideoEncoderSingleLayerVPX>(8);
 		}
 		else if (absl::EqualsIgnoreCase(format.name, cricket::kVp9CodecName))
 		{
+			FStats::Get()->StoreApplicationStat(FStatData(FName(TEXT("Video Codec - VP9")), 1, 0));
 			return std::make_unique<FVideoEncoderSingleLayerVPX>(9);
 		}
+		else if (absl::EqualsIgnoreCase(format.name, cricket::kH265CodecName))
+		{
+			FStats::Get()->StoreApplicationStat(FStatData(FName(TEXT("Video Codec - H265")), 1, 0));
+			FScopeLock Lock(&ActiveEncodersGuard);
+			auto VideoEncoder = std::make_unique<FVideoEncoderSingleLayerHardware>(*this, EPixelStreamingCodec::H265);
+			ActiveEncoders.Add(VideoEncoder.get());
+			return VideoEncoder;
+		}
 		else
 		{
 			// Lock during encoder creation
+			FStats::Get()->StoreApplicationStat(FStatData(FName(TEXT("Video Codec - H264")), 1, 0));
 			FScopeLock Lock(&ActiveEncodersGuard);
-			auto VideoEncoder = std::make_unique<FVideoEncoderSingleLayerH264>(*this);
+			auto VideoEncoder = std::make_unique<FVideoEncoderSingleLayerHardware>(*this, EPixelStreamingCodec::H264);
 			ActiveEncoders.Add(VideoEncoder.get());
 			return VideoEncoder;
 		}
 	}
 
-#if WEBRTC_VERSION == 84
-	void FVideoEncoderFactorySingleLayer::OnEncodedImage(const webrtc::EncodedImage& encoded_image, const webrtc::CodecSpecificInfo* codec_specific_info, const webrtc::RTPFragmentationHeader* fragmentation)
-#elif WEBRTC_VERSION == 96
 	void FVideoEncoderFactorySingleLayer::OnEncodedImage(const webrtc::EncodedImage& encoded_image, const webrtc::CodecSpecificInfo* codec_specific_info)
-#endif
 	{
 		// Lock as we send encoded image to each encoder.
 		FScopeLock Lock(&ActiveEncodersGuard);
 
-		// Go through each encoder and send our encoded image to its callback
-		for (FVideoEncoderSingleLayerH264* Encoder : ActiveEncoders)
+		if (codec_specific_info->codecType == webrtc::kVideoCodecH264 || codec_specific_info->codecType == webrtc::kVideoCodecH265)
 		{
-			Encoder->SendEncodedImage(encoded_image, codec_specific_info
-#if WEBRTC_VERSION == 84
-				,
-				fragmentation
-#endif
-			);
+			// Go through each encoder and send our encoded image to its callback
+			for (FVideoEncoderSingleLayerHardware* Encoder : ActiveEncoders)
+			{
+				Encoder->SendEncodedImage(encoded_image, codec_specific_info);
+			}
 		}
 	}
 
-	void FVideoEncoderFactorySingleLayer::ReleaseVideoEncoder(FVideoEncoderSingleLayerH264* Encoder)
+	void FVideoEncoderFactorySingleLayer::ReleaseVideoEncoder(FVideoEncoderSingleLayerHardware* Encoder)
 	{
 		// Lock during deleting an encoder
 		FScopeLock Lock(&ActiveEncodersGuard);
 		ActiveEncoders.Remove(Encoder);
 	}
 
-	FVideoEncoderWrapperHardware* FVideoEncoderFactorySingleLayer::GetOrCreateHardwareEncoder(int Width, int Height, int MaxBitrate, int TargetBitrate, int MaxFramerate)
+	TWeakPtr<TVideoEncoder<FVideoResourceRHI>> FVideoEncoderFactorySingleLayer::GetHardwareEncoder() const
+	{
+		return HardwareEncoder;
+	}
+
+	TWeakPtr<TVideoEncoder<FVideoResourceRHI>> FVideoEncoderFactorySingleLayer::GetOrCreateHardwareEncoder(const FVideoEncoderConfigH264& VideoConfig)
 	{
 		FScopeLock InitLock(&InitEncoderGuard);
 
-		if (HardwareEncoder != nullptr)
- 		{
-			return HardwareEncoder.Get();
+		if (HardwareEncoder.IsValid())
+		{
+			return HardwareEncoder;
 		}
 		else
 		{
-			// Make AVEncoder frame factory.
-			TUniquePtr<FEncoderFrameFactory> FrameFactory = MakeUnique<FEncoderFrameFactory>();
-
-			// Make the encoder config
-			AVEncoder::FVideoEncoder::FLayerConfig EncoderConfig;
-			EncoderConfig.Width = Width;
-			EncoderConfig.Height = Height;
-			EncoderConfig.MaxFramerate = MaxFramerate;
-			EncoderConfig.TargetBitrate = TargetBitrate;
-			EncoderConfig.MaxBitrate = MaxBitrate;
-
-			// Make the actual AVEncoder encoder.
-			const TArray<AVEncoder::FVideoEncoderInfo>& Available = AVEncoder::FVideoEncoderFactory::Get().GetAvailable();
-			TUniquePtr<AVEncoder::FVideoEncoder> Encoder = AVEncoder::FVideoEncoderFactory::Get().Create(Available[0].ID, FrameFactory->GetOrCreateVideoEncoderInput(), EncoderConfig);
-			if (Encoder.IsValid())
+			// Make the hardware encoder wrapper
+			HardwareEncoder = FVideoEncoder::Create<FVideoResourceRHI>(FAVDevice::GetHardwareDevice(), VideoConfig);
+
+			if (!HardwareEncoder.IsValid())
 			{
-				Encoder->SetOnEncodedPacket([this](uint32 InLayerIndex, const TSharedPtr<AVEncoder::FVideoEncoderInputFrame> InFrame, const AVEncoder::FCodecPacket& InPacket) {
-					// Note: this is a static method call.
-					FVideoEncoderWrapperHardware::OnEncodedPacket(this, InLayerIndex, InFrame, InPacket);
-				});
-
-				// Make the hardware encoder wrapper
-				HardwareEncoder = MakeUnique<FVideoEncoderWrapperHardware>(MoveTemp(FrameFactory), MoveTemp(Encoder));
-				return HardwareEncoder.Get();
+				UE_LOG(LogPixelStreaming, Error, TEXT("Could not create encoder. Check encoder config or perhaps you used up all your HW encoders."));
+				// We could not make the encoder, so indicate the id was not set successfully.
+				return nullptr;
 			}
-			else
+
+			return HardwareEncoder;
+		}
+	}
+
+	TWeakPtr<TVideoEncoder<FVideoResourceRHI>> FVideoEncoderFactorySingleLayer::GetOrCreateHardwareEncoder(const FVideoEncoderConfigH265& VideoConfig)
+	{
+		FScopeLock InitLock(&InitEncoderGuard);
+
+		if (HardwareEncoder.IsValid())
+		{
+			return HardwareEncoder;
+		}
+		else
+		{
+			// Make the hardware encoder wrapper
+			HardwareEncoder = FVideoEncoder::Create<FVideoResourceRHI>(FAVDevice::GetHardwareDevice(), VideoConfig);
+
+			if (!HardwareEncoder.IsValid())
 			{
 				UE_LOG(LogPixelStreaming, Error, TEXT("Could not create encoder. Check encoder config or perhaps you used up all your HW encoders."));
 				// We could not make the encoder, so indicate the id was not set successfully.
 				return nullptr;
 			}
+
+			return HardwareEncoder;
 		}
 	}
 
-	FVideoEncoderWrapperHardware* FVideoEncoderFactorySingleLayer::GetHardwareEncoder()
+	void FVideoEncoderFactorySingleLayer::ForceKeyFrame()
 	{
-		return HardwareEncoder.Get();
+		if (HardwareEncoder)
+		{
+			bForceNextKeyframe = true;
+		}
 	}
 
-	void FVideoEncoderFactorySingleLayer::ForceKeyFrame()
+	bool FVideoEncoderFactorySingleLayer::ShouldForceKeyframe() const
+	{
+		if (HardwareEncoder)
+		{
+			return bForceNextKeyframe;
+		}
+
+		return false;
+	}
+
+	void FVideoEncoderFactorySingleLayer::UnforceKeyFrame()
 	{
 		if (HardwareEncoder)
 		{
-			HardwareEncoder->SetForceNextKeyframe();
+			bForceNextKeyframe = false;
 		}
 	}
 } // namespace UE::PixelStreaming
diff --git a/Engine/Plugins/Media/PixelStreaming/Source/PixelStreaming/Private/VideoEncoderFactorySingleLayer.h b/Engine/Plugins/Media/PixelStreaming/Source/PixelStreaming/Private/VideoEncoderFactorySingleLayer.h
--- a/Engine/Plugins/Media/PixelStreaming/Source/PixelStreaming/Private/VideoEncoderFactorySingleLayer.h
+++ b/Engine/Plugins/Media/PixelStreaming/Source/PixelStreaming/Private/VideoEncoderFactorySingleLayer.h
@@ -2,22 +2,27 @@
 #pragma once
 
 #include "WebRTCIncludes.h"
-#include "VideoEncoderWrapperHardware.h"
 
-namespace UE::PixelStreaming 
+#include "Video/VideoEncoder.h"
+#include "Video/Encoders/Configs/VideoEncoderConfigH264.h"
+#include "Video/Encoders/Configs/VideoEncoderConfigH265.h"
+#include "Video/Resources/VideoResourceRHI.h"
+
+namespace UE::PixelStreaming
 {
-	class FVideoEncoderSingleLayerH264;
-	class FVideoEncoderWrapperHardware;
-	
+	using FVideoEncoderHardware = TVideoEncoder<FVideoResourceRHI>;
+
+	class FVideoEncoderSingleLayerHardware;
+
 	/**
 	 * An encoder factory for a single layer. Do not use this directly. Use FVideoEncoderFactoryLayered even
 	 * when not using simulcast specifically.
 	 */
-	class FVideoEncoderFactorySingleLayer : public webrtc::VideoEncoderFactory
+	class PIXELSTREAMING_API FVideoEncoderFactorySingleLayer : public webrtc::VideoEncoderFactory
 	{
 	public:
 		FVideoEncoderFactorySingleLayer() = default;
-		virtual ~FVideoEncoderFactorySingleLayer() = default;
+		virtual ~FVideoEncoderFactorySingleLayer() override = default;
 
 		virtual std::vector<webrtc::SdpVideoFormat> GetSupportedFormats() const override;
 
@@ -27,31 +32,27 @@ namespace UE::PixelStreaming
 		// Always returns our H264 hardware encoder for now.
 		virtual std::unique_ptr<webrtc::VideoEncoder> CreateVideoEncoder(const webrtc::SdpVideoFormat& Format) override;
 
-		void ReleaseVideoEncoder(FVideoEncoderSingleLayerH264* Encoder);
-			
-		FVideoEncoderWrapperHardware* GetOrCreateHardwareEncoder(int Width, int Height, int MaxBitrate, int TargetBitrate, int MaxFramerate);
-		FVideoEncoderWrapperHardware* GetHardwareEncoder();
+		void ReleaseVideoEncoder(FVideoEncoderSingleLayerHardware* Encoder);
 
-		void ForceKeyFrame();
+		TWeakPtr<FVideoEncoderHardware> GetHardwareEncoder() const;
+		TWeakPtr<FVideoEncoderHardware> GetOrCreateHardwareEncoder(const FVideoEncoderConfigH264& VideoConfig);
+		TWeakPtr<FVideoEncoderHardware> GetOrCreateHardwareEncoder(const FVideoEncoderConfigH265& VideoConfig);
 
-#if WEBRTC_VERSION == 84
-		void OnEncodedImage(const webrtc::EncodedImage& Encoded_image, const webrtc::CodecSpecificInfo* CodecSpecificInfo, const webrtc::RTPFragmentationHeader* Fragmentation);
-#elif WEBRTC_VERSION == 96
+		void ForceKeyFrame();
+		bool ShouldForceKeyframe() const;
+		void UnforceKeyFrame();
 		void OnEncodedImage(const webrtc::EncodedImage& Encoded_image, const webrtc::CodecSpecificInfo* CodecSpecificInfo);
-#endif
 
 	private:
-#if WEBRTC_VERSION == 84
-		static webrtc::SdpVideoFormat CreateH264Format(webrtc::H264::Profile Profile, webrtc::H264::Level Level);
-#elif WEBRTC_VERSION == 96
 		static webrtc::SdpVideoFormat CreateH264Format(webrtc::H264Profile Profile, webrtc::H264Level Level);
-#endif
 
-		TUniquePtr<FVideoEncoderWrapperHardware> HardwareEncoder;
+		TSharedPtr<FVideoEncoderHardware> HardwareEncoder;
+
+		uint8 bForceNextKeyframe : 1;
 
 		// Encoders assigned to each peer. Each one of these will be assigned to one of the hardware encoders
 		// raw pointers here because the factory interface wants a unique_ptr here so we cant own the object.
-		TArray<FVideoEncoderSingleLayerH264*> ActiveEncoders;
+		TArray<FVideoEncoderSingleLayerHardware*> ActiveEncoders;
 		FCriticalSection ActiveEncodersGuard;
 
 		// Init encoder guard
diff --git a/Engine/Plugins/Media/PixelStreaming/Source/PixelStreaming/Private/VideoEncoderLayered.cpp b/Engine/Plugins/Media/PixelStreaming/Source/PixelStreaming/Private/VideoEncoderLayered.cpp
--- a/Engine/Plugins/Media/PixelStreaming/Source/PixelStreaming/Private/VideoEncoderLayered.cpp
+++ b/Engine/Plugins/Media/PixelStreaming/Source/PixelStreaming/Private/VideoEncoderLayered.cpp
@@ -5,6 +5,7 @@
 #include "modules/video_coding/codecs/vp9/include/vp9.h"
 #include "VideoEncoderFactorySingleLayer.h"
 #include "Settings.h"
+#include "PixelStreamingTrace.h"
 
 namespace
 {
@@ -75,22 +76,12 @@ namespace
 		{
 		}
 
-#if WEBRTC_VERSION == 84
-		EncodedImageCallback::Result OnEncodedImage(
-			const webrtc::EncodedImage& encoded_image,
-			const webrtc::CodecSpecificInfo* codec_specific_info,
-			const webrtc::RTPFragmentationHeader* fragmentation) override
-		{
-			return adapter_->OnEncodedImage(stream_idx_, encoded_image, codec_specific_info, fragmentation);
-		}
-#elif WEBRTC_VERSION == 96
 		EncodedImageCallback::Result OnEncodedImage(
 			const webrtc::EncodedImage& encoded_image,
 			const webrtc::CodecSpecificInfo* codec_specific_info) override
 		{
 			return adapter_->OnEncodedImage(stream_idx_, encoded_image, codec_specific_info);
 		}
-#endif
 
 	private:
 		UE::PixelStreaming::FVideoEncoderLayered* const adapter_;
@@ -146,6 +137,7 @@ namespace UE::PixelStreaming
 		// clang-format off
 		const webrtc::SdpVideoFormat Format(CurrentCodec.codecType == webrtc::kVideoCodecVP8 ? "VP8"
 										  : CurrentCodec.codecType == webrtc::kVideoCodecVP9 ? "VP9"
+										  : CurrentCodec.codecType == webrtc::kVideoCodecH265 ? "H265"
 																							 : "H264",
 			                                VideoFormat.parameters);
 		// clang-format on
@@ -211,6 +203,7 @@ namespace UE::PixelStreaming
 
 	int FVideoEncoderLayered::Encode(const webrtc::VideoFrame& input_image, const std::vector<webrtc::VideoFrameType>* frame_types)
 	{
+		TRACE_CPUPROFILER_EVENT_SCOPE_ON_CHANNEL_STR("PixelStreaming Encoding Layer", PixelStreamingChannel);
 		if (!IsInitialized())
 		{
 			return WEBRTC_VIDEO_CODEC_UNINITIALIZED;
@@ -261,12 +254,8 @@ namespace UE::PixelStreaming
 			rtc::scoped_refptr<FFrameBufferMultiFormat> LayerFrameBuffer = FrameBuffer->GetLayer(LayerIndex);
 			NewFrame.set_video_frame_buffer(LayerFrameBuffer);
 
-#if WEBRTC_VERSION == 84
-			const uint32_t FrameTimestampMs = 1000 * NewFrame.timestamp() / 90000; // kVideoPayloadTypeFrequency;
-#elif WEBRTC_VERSION == 96
 			// Convert timestamp from RTP 90kHz clock.
 			webrtc::Timestamp FrameTimestamp = webrtc::Timestamp::Micros((1000 * NewFrame.timestamp()) / 90);
-#endif
 			// If adapter is passed through and only one sw encoder does simulcast,
 			// frame types for all streams should be passed to the encoder unchanged.
 			// Otherwise a single per-encoder frame type is passed.
@@ -278,19 +267,12 @@ namespace UE::PixelStreaming
 			}
 			else
 			{
-#if WEBRTC_VERSION == 84
-				if (StreamInfos[StreamIdx].FramerateController->DropFrame(FrameTimestampMs))
-#elif WEBRTC_VERSION == 96
 				if (StreamInfos[StreamIdx].FramerateController->ShouldDropFrame(FrameTimestamp.us() * 1000))
-#endif
 				{
 					continue;
 				}
 				std::fill(StreamFrameTypes.begin(), StreamFrameTypes.end(), webrtc::VideoFrameType::kVideoFrameDelta);
 			}
-#if WEBRTC_VERSION == 84
-			StreamInfos[StreamIdx].FramerateController->AddFrame(FrameTimestampMs);
-#endif
 			const int RtcError = StreamInfos[StreamIdx].Encoder->Encode(NewFrame, &StreamFrameTypes);
 			if (RtcError != WEBRTC_VIDEO_CODEC_OK)
 			{
@@ -365,12 +347,7 @@ namespace UE::PixelStreaming
 						webrtc::DataRate::BitsPerSec(StreamParameters.bitrate.get_sum_bps());
 				}
 			}
-#if WEBRTC_VERSION == 84
-			StreamParameters.framerate_fps = std::min<double>(parameters.framerate_fps, StreamInfos[StreamIdx].FramerateController->GetTargetRate());
-#elif WEBRTC_VERSION == 96
 			StreamParameters.framerate_fps = std::min<double>(parameters.framerate_fps, StreamInfos[StreamIdx].FramerateController->GetMaxFramerate());
-#endif
-
 			StreamInfos[StreamIdx].Encoder->SetRates(StreamParameters);
 		}
 	}
@@ -402,24 +379,14 @@ namespace UE::PixelStreaming
 	webrtc::EncodedImageCallback::Result FVideoEncoderLayered::OnEncodedImage(
 		size_t stream_idx,
 		const webrtc::EncodedImage& encodedImage,
-		const webrtc::CodecSpecificInfo* codecSpecificInfo
-#if WEBRTC_VERSION == 84
-		,
-		const webrtc::RTPFragmentationHeader* fragmentation
-#endif
-	)
+		const webrtc::CodecSpecificInfo* codecSpecificInfo)
 	{
 		webrtc::EncodedImage StreamImage(encodedImage);
 		webrtc::CodecSpecificInfo StreamCodecSpecific = *codecSpecificInfo;
 
 		StreamImage.SetSpatialIndex(stream_idx);
 
-		return EncodedCompleteCallback->OnEncodedImage(StreamImage, &StreamCodecSpecific
-#if WEBRTC_VERSION == 84
-			,
-			fragmentation
-#endif
-		);
+		return EncodedCompleteCallback->OnEncodedImage(StreamImage, &StreamCodecSpecific);
 	}
 
 	bool FVideoEncoderLayered::IsInitialized() const
diff --git a/Engine/Plugins/Media/PixelStreaming/Source/PixelStreaming/Private/VideoEncoderLayered.h b/Engine/Plugins/Media/PixelStreaming/Source/PixelStreaming/Private/VideoEncoderLayered.h
--- a/Engine/Plugins/Media/PixelStreaming/Source/PixelStreaming/Private/VideoEncoderLayered.h
+++ b/Engine/Plugins/Media/PixelStreaming/Source/PixelStreaming/Private/VideoEncoderLayered.h
@@ -41,11 +41,7 @@ namespace UE::PixelStreaming
 		webrtc::EncodedImageCallback::Result OnEncodedImage(
 			size_t stream_idx,
 			const webrtc::EncodedImage& encoded_image,
-			const webrtc::CodecSpecificInfo* codec_specific_info
-#if WEBRTC_VERSION == 84
-			, const webrtc::RTPFragmentationHeader* fragmentation
-#endif
-		);
+			const webrtc::CodecSpecificInfo* codec_specific_info);
 
 		EncoderInfo GetEncoderInfo() const override;
 
diff --git a/Engine/Plugins/Media/PixelStreaming/Source/PixelStreaming/Private/VideoEncoderSingleLayerVPX.cpp b/Engine/Plugins/Media/PixelStreaming/Source/PixelStreaming/Private/VideoEncoderSingleLayerVPX.cpp
--- a/Engine/Plugins/Media/PixelStreaming/Source/PixelStreaming/Private/VideoEncoderSingleLayerVPX.cpp
+++ b/Engine/Plugins/Media/PixelStreaming/Source/PixelStreaming/Private/VideoEncoderSingleLayerVPX.cpp
@@ -7,6 +7,7 @@
 
 #include "PixelCaptureBufferFormat.h"
 #include "PixelCaptureOutputFrameI420.h"
+#include "PixelStreamingTrace.h"
 
 namespace UE::PixelStreaming
 {
@@ -67,6 +68,7 @@ namespace UE::PixelStreaming
 
 	int32 FVideoEncoderSingleLayerVPX::Encode(webrtc::VideoFrame const& frame, std::vector<webrtc::VideoFrameType> const* frame_types)
 	{
+		TRACE_CPUPROFILER_EVENT_SCOPE_ON_CHANNEL_STR("PixelStreaming VPX Encoding", PixelStreamingChannel);
 		checkf(WebRTCVPXEncoder, TEXT("WebRTCVPXEncoder was null. Should never happen."));
 
 		const FFrameBufferMultiFormat* FrameBuffer = StaticCast<FFrameBufferMultiFormat*>(frame.video_frame_buffer().get());
diff --git a/Engine/Plugins/Media/PixelStreaming/Source/PixelStreaming/Private/VideoSource.cpp b/Engine/Plugins/Media/PixelStreaming/Source/PixelStreaming/Private/VideoSource.cpp
--- a/Engine/Plugins/Media/PixelStreaming/Source/PixelStreaming/Private/VideoSource.cpp
+++ b/Engine/Plugins/Media/PixelStreaming/Source/PixelStreaming/Private/VideoSource.cpp
@@ -3,6 +3,7 @@
 #include "VideoSource.h"
 #include "Settings.h"
 #include "FrameBufferMultiFormat.h"
+#include "PixelStreamingTrace.h"
 
 namespace UE::PixelStreaming
 {
@@ -23,6 +24,7 @@ namespace UE::PixelStreaming
 
 	void FVideoSource::PushFrame()
 	{
+	    TRACE_CPUPROFILER_EVENT_SCOPE_ON_CHANNEL_STR("PixelStreaming Push Video Frame", PixelStreamingChannel);
 		static int32 FrameId = 1;
 
 		CurrentState = webrtc::MediaSourceInterface::SourceState::kLive;
@@ -31,7 +33,7 @@ namespace UE::PixelStreaming
 		check(FrameBuffer->width() != 0 && FrameBuffer->height() != 0);
 		webrtc::VideoFrame Frame = webrtc::VideoFrame::Builder()
 									   .set_video_frame_buffer(FrameBuffer)
-									   .set_timestamp_us(webrtc::Clock::GetRealTimeClock()->TimeInMicroseconds())
+									   .set_timestamp_us(rtc::TimeMicros())
 									   .set_rotation(webrtc::VideoRotation::kVideoRotation_0)
 									   .set_id(FrameId++)
 									   .build();
diff --git a/Engine/Plugins/Media/PixelStreaming/Source/PixelStreaming/Private/VideoSourceGroup.cpp b/Engine/Plugins/Media/PixelStreaming/Source/PixelStreaming/Private/VideoSourceGroup.cpp
--- a/Engine/Plugins/Media/PixelStreaming/Source/PixelStreaming/Private/VideoSourceGroup.cpp
+++ b/Engine/Plugins/Media/PixelStreaming/Source/PixelStreaming/Private/VideoSourceGroup.cpp
@@ -4,6 +4,7 @@
 #include "VideoSource.h"
 #include "Settings.h"
 #include "PixelStreamingPrivate.h"
+#include "PixelStreamingTrace.h"
 
 namespace UE::PixelStreaming
 {
@@ -13,7 +14,7 @@ namespace UE::PixelStreaming
 	}
 
 	FVideoSourceGroup::FVideoSourceGroup()
-		: bCoupleFramerate(Settings::CoupleFrameRate() || !Settings::IsUsingSafeTextureCopy())
+		: bCoupleFramerate(Settings::IsCoupledFramerate() || !Settings::CVarPixelStreamingCaptureUseFence.GetValueOnAnyThread()) /* Only couple framerate if not using fence */
 		, FramesPerSecond(Settings::CVarPixelStreamingWebRTCFps.GetValueOnAnyThread())
 	{
 	}
@@ -107,6 +108,7 @@ namespace UE::PixelStreaming
 
 	void FVideoSourceGroup::Tick()
 	{
+	    TRACE_CPUPROFILER_EVENT_SCOPE_ON_CHANNEL_STR("PixelStreaming Video Source Group Tick", PixelStreamingChannel);
 		FScopeLock Lock(&CriticalSection);
 		// for each player session, push a frame
 		for (auto& VideoSource : VideoSources)
diff --git a/Engine/Plugins/Media/PixelStreaming/Source/PixelStreaming/Private/WebRTCIncludes.h b/Engine/Plugins/Media/PixelStreaming/Source/PixelStreaming/Private/WebRTCIncludes.h
--- a/Engine/Plugins/Media/PixelStreaming/Source/PixelStreaming/Private/WebRTCIncludes.h
+++ b/Engine/Plugins/Media/PixelStreaming/Source/PixelStreaming/Private/WebRTCIncludes.h
@@ -44,11 +44,7 @@ THIRD_PARTY_INCLUDES_START
 #include "rtc_base/arraysize.h"
 #include "rtc_base/net_helpers.h"
 #include "rtc_base/string_utils.h"
-#if WEBRTC_VERSION == 84
-#include "rtc_base/signal_thread.h"
-#elif WEBRTC_VERSION == 96
 #include "rtc_base/internal/default_socket_server.h"
-#endif 
 #include "rtc_base/physical_socket_server.h"
 #include "rtc_base/experiments/rate_control_settings.h"
 
@@ -57,11 +53,6 @@ THIRD_PARTY_INCLUDES_START
 
 #include "media/engine/internal_decoder_factory.h"
 #include "media/engine/internal_encoder_factory.h"
-#if WEBRTC_VERSION == 84
-#include "media/base/h264_profile_level_id.h"
-#elif WEBRTC_VERSION == 96
-#include "api/video_codecs/h264_profile_level_id.h"
-#endif 
 #include "media/base/adapted_video_track_source.h"
 #include "media/base/media_channel.h"
 #include "media/base/video_common.h"
@@ -72,15 +63,13 @@ THIRD_PARTY_INCLUDES_START
 #include "modules/audio_device/audio_device_buffer.h"
 #include "modules/audio_processing/include/audio_processing.h"
 #include "modules/video_coding/codecs/h264/include/h264.h"
-#if WEBRTC_VERSION ==  84
-#include "modules/video_coding/utility/framerate_controller.h"
-#endif
 #include "modules/video_coding/utility/simulcast_rate_allocator.h"
 #include "modules/video_coding/codecs/vp8/include/vp8.h"
 #include "modules/video_coding/codecs/vp9/include/vp9.h"
 
 #include "common_video/h264/h264_bitstream_parser.h"
 #include "common_video/h264/h264_common.h"
+#include "common_video/h265/h265_common.h"
 
 #include "media/base/video_broadcaster.h"
 
@@ -96,6 +85,6 @@ THIRD_PARTY_INCLUDES_START
 
 	#include "Windows/PostWindowsApi.h"
 	#include "Windows/HideWindowsPlatformTypes.h"
-#endif //PLATFORM_WINDOWS
+#endif // PLATFORM_WINDOWS
 
 THIRD_PARTY_INCLUDES_END
diff --git a/Engine/Plugins/Media/PixelStreaming/Source/PixelStreaming/Public/IPixelStreamingModule.h b/Engine/Plugins/Media/PixelStreaming/Source/PixelStreaming/Public/IPixelStreamingModule.h
--- a/Engine/Plugins/Media/PixelStreaming/Source/PixelStreaming/Public/IPixelStreamingModule.h
+++ b/Engine/Plugins/Media/PixelStreaming/Source/PixelStreaming/Public/IPixelStreamingModule.h
@@ -9,19 +9,17 @@
 #include "Templates/RefCounting.h"
 #include "PixelStreamingPlayerId.h"
 #include "IPixelStreamingAudioSink.h"
-#include "IPixelStreamingInputHandler.h"
 #include "PixelStreamingWebRTCIncludes.h"
 #include "IPixelStreamingStreamer.h"
-#include "IInputDeviceModule.h"
 #include "PixelStreamingCodec.h"
-#include "PixelStreamingProtocol.h"
+#include "PixelStreamingInputProtocol.h"
 
 class UPixelStreamingInput;
 
 /**
  * The public interface of the Pixel Streaming module.
  */
-class PIXELSTREAMING_API IPixelStreamingModule : public IInputDeviceModule
+class PIXELSTREAMING_API IPixelStreamingModule : public IModuleInterface
 {
 public:
 	/**
@@ -110,30 +108,34 @@ public:
 	virtual TSharedPtr<IPixelStreamingStreamer> DeleteStreamer(const FString& StreamerId) = 0;
 
 	/**
-	 * Get the protocol currently used by each peer. 
+	 * Get the protocol currently used by each peer.
 	 */
-	virtual const Protocol::FPixelStreamingProtocol& GetProtocol() = 0;
+	UE_DEPRECATED(5.2, "GetProtocol() is deprecated. Simply access the required direction with FPixelStreamingInputProtocol::ToStreamerProcol or FPixelStreamingInputProtocol::FromStreamerProcol")
+	virtual const FPixelStreamingInputProtocol GetProtocol() = 0;
 
 	/**
-	 * Register a new message type that peers can send and pixel streaming can receive
+	 * @brief Register a new message that peers can send and pixel streaming can receive
+	 *
 	 * @param MessageDirection The direction the message will travel. eg Streamer->Player or Player->Streamer
 	 * @param MessageType The human readable identifier (eg "TouchStarted")
 	 * @param Message The object used to define the structure of the message
-	 * @param Handler The handler for this message type. This function will be executed whenever the corresponding message type is received 
+	 * @param Handler The handler for this message. This function will be executed whenever the corresponding message is received
 	 */
-	virtual void RegisterMessage(Protocol::EPixelStreamingMessageDirection MessageDirection, const FString& MessageType, Protocol::FPixelStreamingInputMessage Message, const TFunction<void(FMemoryReader)>& Handler) = 0;
+	UE_DEPRECATED(5.2, "RegisterMessage(...) has been moved from the PixelStreaming module to the PixelStreamingInput module")
+	virtual void RegisterMessage(EPixelStreamingMessageDirection MessageDirection, const FString& MessageType, FPixelStreamingInputMessage Message, const TFunction<void(FMemoryReader)>& Handler) = 0;
 
 	/**
-	 * @brief Find the function to be called whenever the specified message type is received.
-	 * 
+	 * @brief Find the function to be called whenever the specified message is received.
+	 *
 	 * @param MessageType The human readable identifier for the message
-	 * @return TFunction<void(FMemoryReader)> The function called when this message type is received.
+	 * @return TFunction<void(FMemoryReader)> The function called when this message is received.
 	 */
+	UE_DEPRECATED(5.2, "FindMessageHandler(...) has been moved from the PixelStreaming module to the PixelStreamingInput module")
 	virtual TFunction<void(FMemoryReader)> FindMessageHandler(const FString& MessageType) = 0;
 
 	/**
 	 * Sets the target FPS for Externally Consumed video Tracks
-	 * @param InFPS new FPS for the ExternalVideoSource to output at. 
+	 * @param InFPS new FPS for the ExternalVideoSource to output at.
 	 */
 	virtual void SetExternalVideoSourceFPS(uint32 InFPS) = 0;
 
@@ -185,11 +187,8 @@ public:
 
 	/**
 	 * @brief A method for iterating through all of the streamers on the module
-	 * 
+	 *
 	 * @param Func The lambda to execute with each streamer
 	 */
 	virtual void ForEachStreamer(const TFunction<void(TSharedPtr<IPixelStreamingStreamer>)>& Func) = 0;
-
-	DECLARE_MULTICAST_DELEGATE(FOnProtocolUpdated);
-	FOnProtocolUpdated OnProtocolUpdated;
 };
diff --git a/Engine/Plugins/Media/PixelStreaming/Source/PixelStreaming/Public/IPixelStreamingSignallingConnectionObserver.h b/Engine/Plugins/Media/PixelStreaming/Source/PixelStreaming/Public/IPixelStreamingSignallingConnectionObserver.h
--- a/Engine/Plugins/Media/PixelStreaming/Source/PixelStreaming/Public/IPixelStreamingSignallingConnectionObserver.h
+++ b/Engine/Plugins/Media/PixelStreaming/Source/PixelStreaming/Public/IPixelStreamingSignallingConnectionObserver.h
@@ -10,22 +10,21 @@
 class PIXELSTREAMING_API IPixelStreamingSignallingConnectionObserver
 {
 public:
-	virtual ~IPixelStreamingSignallingConnectionObserver() = default;
 	virtual void OnSignallingConnected() = 0;
 	virtual void OnSignallingDisconnected(int32 StatusCode, const FString& Reason, bool bWasClean) = 0;
 	virtual void OnSignallingError(const FString& ErrorMsg) = 0;
 	virtual void OnSignallingConfig(const webrtc::PeerConnectionInterface::RTCConfiguration& Config) = 0;
 
 	// Streamer-only
-	virtual void OnSignallingSessionDescription(FPixelStreamingPlayerId PlayerId, webrtc::SdpType Type, const FString& Sdp) { unimplemented(); }
-	virtual void OnSignallingRemoteIceCandidate(FPixelStreamingPlayerId PlayerId, const FString& SdpMid, int SdpMLineIndex, const FString& Sdp) { unimplemented(); }
-	virtual void OnSignallingPlayerConnected(FPixelStreamingPlayerId PlayerId, const FPixelStreamingPlayerConfig& PlayerConfig) { unimplemented(); }
-	virtual void OnSignallingPlayerDisconnected(FPixelStreamingPlayerId PlayerId) { unimplemented(); }
-	virtual void OnSignallingSFUPeerDataChannels(FPixelStreamingPlayerId SFUId, FPixelStreamingPlayerId PlayerId, int32 SendStreamId, int32 RecvStreamId) { unimplemented(); }
+	virtual void OnSignallingSessionDescription(FPixelStreamingPlayerId PlayerId, webrtc::SdpType Type, const FString& Sdp) = 0;
+	virtual void OnSignallingRemoteIceCandidate(FPixelStreamingPlayerId PlayerId, const FString& SdpMid, int SdpMLineIndex, const FString& Sdp) = 0;
+	virtual void OnSignallingPlayerConnected(FPixelStreamingPlayerId PlayerId, const FPixelStreamingPlayerConfig& PlayerConfig) = 0;
+	virtual void OnSignallingPlayerDisconnected(FPixelStreamingPlayerId PlayerId) = 0;
+	virtual void OnSignallingSFUPeerDataChannels(FPixelStreamingPlayerId SFUId, FPixelStreamingPlayerId PlayerId, int32 SendStreamId, int32 RecvStreamId) = 0;
 
 	// Player-only
-	virtual void OnSignallingSessionDescription(webrtc::SdpType Type, const FString& Sdp) { unimplemented(); }
-	virtual void OnSignallingRemoteIceCandidate(const FString& SdpMid, int SdpMLineIndex, const FString& Sdp) { unimplemented(); }
-	virtual void OnSignallingPlayerCount(uint32 Count) { unimplemented(); }
-	virtual void OnSignallingPeerDataChannels(int32 SendStreamId, int32 RecvStreamId) { unimplemented(); }
+	virtual void OnSignallingSessionDescription(webrtc::SdpType Type, const FString& Sdp) = 0;
+	virtual void OnSignallingRemoteIceCandidate(const FString& SdpMid, int SdpMLineIndex, const FString& Sdp) = 0;
+	virtual void OnSignallingPlayerCount(uint32 Count) = 0;
+	virtual void OnSignallingPeerDataChannels(int32 SendStreamId, int32 RecvStreamId) = 0;
 };
diff --git a/Engine/Plugins/Media/PixelStreaming/Source/PixelStreaming/Public/IPixelStreamingStreamer.h b/Engine/Plugins/Media/PixelStreaming/Source/PixelStreaming/Public/IPixelStreamingStreamer.h
--- a/Engine/Plugins/Media/PixelStreaming/Source/PixelStreaming/Public/IPixelStreamingStreamer.h
+++ b/Engine/Plugins/Media/PixelStreaming/Source/PixelStreaming/Public/IPixelStreamingStreamer.h
@@ -2,7 +2,6 @@
 
 #pragma once
 
-#include "PixelStreamingProtocol.h"
 #include "PixelStreamingPlayerId.h"
 #include "PixelStreamingVideoInput.h"
 #include "CoreMinimal.h"
@@ -12,8 +11,11 @@
 #include "IPixelStreamingAudioSink.h"
 #include "IPixelStreamingAudioInput.h"
 #include "IPixelStreamingInputHandler.h"
+#include "IPixelStreamingSignallingConnection.h"
+#include "PixelStreamingInputEnums.h"
 
 class UTexture2D;
+class FPixelStreamingSignallingConnection;
 
 class PIXELSTREAMING_API IPixelStreamingStreamer
 {
@@ -104,7 +106,7 @@ public:
 
 	/**
 	 * @brief Get this streamer's ID
-	 * 
+	 *
 	 */
 	virtual FString GetId() = 0;
 
@@ -209,6 +211,26 @@ public:
 	 */
 	virtual TWeakPtr<IPixelStreamingInputHandler> GetInputHandler() = 0;
 
+	/**
+	 * @brief Get the connection to the signalling server.
+	 * @return The connection to the signalling server.
+	 */
+	virtual TWeakPtr<IPixelStreamingSignallingConnection> GetSignallingConnection() = 0;
+
+	/**
+	 * @brief Set the connection to the signalling server
+	 *
+	 * @param InSignallingConnection The connection to the signalling server
+	 */
+	virtual void SetSignallingConnection(TSharedPtr<IPixelStreamingSignallingConnection> InSignallingConnection) = 0;
+
+	/**
+	 * @brief Get the Signalling Connection Observer object
+	 * 
+	 * @return TWeakPtr<IPixelStreamingSignallingConnectionObserver> 
+	 */
+	virtual TWeakPtr<IPixelStreamingSignallingConnectionObserver> GetSignallingConnectionObserver() = 0;
+
 	/**
 	 * @brief Set the type for this streamers input handler. This controls whether input is routed to widgets or windows
 	 */
diff --git a/Engine/Plugins/Media/PixelStreaming/Source/PixelStreaming/Public/PixelStreamingBufferBuilder.h b/Engine/Plugins/Media/PixelStreaming/Source/PixelStreaming/Public/PixelStreamingBufferBuilder.h
--- a/Engine/Plugins/Media/PixelStreaming/Source/PixelStreaming/Public/PixelStreamingBufferBuilder.h
+++ b/Engine/Plugins/Media/PixelStreaming/Source/PixelStreaming/Public/PixelStreamingBufferBuilder.h
@@ -45,12 +45,7 @@ namespace UE::PixelStreaming
 		{
 			const size_t VSize = ValueSize(Forward<T>(Value));
 			const void* VLoc = ValueLoc(Forward<T>(Value));
-
-#if WEBRTC_VERSION == 84
-			FMemory::Memcpy(&Buffer[Pos], VLoc, VSize);
-#elif WEBRTC_VERSION == 96
 			FMemory::Memcpy(Buffer.MutableData() + Pos, VLoc, VSize);
-#endif
 			Pos += VSize;
 		}
 	};
diff --git a/Engine/Plugins/Media/PixelStreaming/Source/PixelStreaming/Public/PixelStreamingCodec.h b/Engine/Plugins/Media/PixelStreaming/Source/PixelStreaming/Public/PixelStreamingCodec.h
--- a/Engine/Plugins/Media/PixelStreaming/Source/PixelStreaming/Public/PixelStreamingCodec.h
+++ b/Engine/Plugins/Media/PixelStreaming/Source/PixelStreaming/Public/PixelStreamingCodec.h
@@ -6,6 +6,7 @@ enum class EPixelStreamingCodec
 {
 	Invalid,
 	H264,
+	H265,
 	VP8,
 	VP9,
 };
diff --git a/Engine/Plugins/Media/PixelStreaming/Source/PixelStreaming/Public/PixelStreamingDataChannel.h b/Engine/Plugins/Media/PixelStreaming/Source/PixelStreaming/Public/PixelStreamingDataChannel.h
--- a/Engine/Plugins/Media/PixelStreaming/Source/PixelStreaming/Public/PixelStreamingDataChannel.h
+++ b/Engine/Plugins/Media/PixelStreaming/Source/PixelStreaming/Public/PixelStreamingDataChannel.h
@@ -3,7 +3,6 @@
 #pragma once
 
 #include "PixelStreamingBufferBuilder.h"
-#include "PixelStreamingProtocol.h"
 #include "PixelStreamingWebRTCIncludes.h"
 
 class FPixelStreamingPeerConnection;
diff --git a/Engine/Plugins/Media/PixelStreaming/Source/PixelStreaming/Public/PixelStreamingPeerConnection.h b/Engine/Plugins/Media/PixelStreaming/Source/PixelStreaming/Public/PixelStreamingPeerConnection.h
--- a/Engine/Plugins/Media/PixelStreaming/Source/PixelStreaming/Public/PixelStreamingPeerConnection.h
+++ b/Engine/Plugins/Media/PixelStreaming/Source/PixelStreaming/Public/PixelStreamingPeerConnection.h
@@ -2,10 +2,10 @@
 
 #pragma once
 
-#include "PixelStreamingProtocol.h"
 #include "IPixelStreamingAudioSink.h"
 #include "IPixelStreamingAudioInput.h"
 #include "PixelStreamingWebRTCIncludes.h"
+#include "PixelStreamingCodec.h"
 
 class FPixelStreamingDataChannel;
 
@@ -190,6 +190,12 @@ public:
 	 */
 	void SetWebRTCStatsCallback(rtc::scoped_refptr<webrtc::RTCStatsCollectorCallback> InCallback);
 
+	/**
+	 * @return The negotiated video codec as parsed from the SDP offer/answer.
+	 * Note: This changes as new offer/answers come in potentially.
+	 */
+	EPixelStreamingCodec GetNegotiatedVideoCodec() const;
+
 	/**
 	 * Called when the ICE candidate is emitted.
 	 */
@@ -225,7 +231,11 @@ public:
 	template <typename FunctorT>
 	static void PostSignalingTask(FunctorT&& InFunc)
 	{
-		SignallingThread->PostTask(RTC_FROM_HERE, Forward<FunctorT>(InFunc));
+		// Someone may accidentally call this static function without calling FPixelStreamingPeerConnection::Create first
+		if (SignallingThread != nullptr)
+		{
+			SignallingThread->PostTask(RTC_FROM_HERE, Forward<FunctorT>(InFunc));
+		}
 	}
 
 	static TSharedPtr<IPixelStreamingAudioInput> CreateAudioInput();
@@ -267,6 +277,9 @@ private:
 	rtc::VideoSinkInterface<webrtc::VideoFrame>* VideoSink = nullptr;
 	TSharedPtr<IPixelStreamingAudioSink> AudioSink;
 
+	/* The video codec that is negotiated and extracted from the SDP. */
+	EPixelStreamingCodec NegotiatedVideoCodec = EPixelStreamingCodec::Invalid;
+
 	// using unique_ptr here because TUniquePtr does not like incomplete types with -disableunity and -nopch
 	std::unique_ptr<class FPeerWebRTCStatsSource> StatsSource;
 
diff --git a/Engine/Plugins/Media/PixelStreaming/Source/PixelStreaming/Public/PixelStreamingSignallingConnection.h b/Engine/Plugins/Media/PixelStreaming/Source/PixelStreaming/Public/PixelStreamingSignallingConnection.h
--- a/Engine/Plugins/Media/PixelStreaming/Source/PixelStreaming/Public/PixelStreamingSignallingConnection.h
+++ b/Engine/Plugins/Media/PixelStreaming/Source/PixelStreaming/Public/PixelStreamingSignallingConnection.h
@@ -7,96 +7,33 @@
 #include "Engine/EngineTypes.h"
 #include "PixelStreamingPlayerId.h"
 #include "IPixelStreamingSignallingConnectionObserver.h"
+#include "IPixelStreamingSignallingConnection.h"
 
 class IWebSocket;
 
 /**
- * A specialized signalling server connection object for Pixel Streaming signalling servers.
+ * A specialized signalling server connection object for web socket based Pixel Streaming signalling servers.
  */
-class PIXELSTREAMING_API FPixelStreamingSignallingConnection final
+class PIXELSTREAMING_API FPixelStreamingSignallingConnection : public IPixelStreamingSignallingConnection
 {
 public:
-	using FWebSocketFactory = TFunction<TSharedPtr<IWebSocket>(const FString&)>;
-
-	FPixelStreamingSignallingConnection(const FWebSocketFactory& InWebSocketFactory, IPixelStreamingSignallingConnectionObserver& Observer, FString InStreamerId = "");
-	~FPixelStreamingSignallingConnection();
-
-	/**
-	 * Connect to a specified signalling server at ths given URL
-	 * @param URL The url of the destination signalling server.
-	 */
-	void TryConnect(FString URL);
-
-	/**
-	 * Disconnects from the signalling server. Safe to call even when not connected.
-	 */
-	void Disconnect();
-
-	/**
-	 * Checks if the connection is established.
-	 * @returns True when the signalling server is connected.
-	 */
-	bool IsConnected() const;
-
-	/**
-	 * [Streamer only] Sends an offer to the specified player.
-	 * @param PlayerId The Id of the destination player.
-	 * @param SDP The offer session description for the player.
-	 */
-	void SendOffer(FPixelStreamingPlayerId PlayerId, const webrtc::SessionDescriptionInterface& SDP);
-
-	/**
-	 * [Streamer only] Sends an answer to the specified player.
-	 * @param PlayerId The Id of the destination player.
-	 * @param SDP The answer session description for the player.
-	 */
-	void SendAnswer(FPixelStreamingPlayerId PlayerId, const webrtc::SessionDescriptionInterface& SDP);
-
-	/**
-	 * [Streamer only] Sends ice candidate information to the specified player.
-	 * @param PlayerId The Id of the destination player.
-	 * @param IceCandidate The ice candidate information to send.
-	 */
-	void SendIceCandidate(FPixelStreamingPlayerId PlayerId, const webrtc::IceCandidateInterface& IceCandidate);
-
-	/**
-	 * [Streamer only] Send a disconnect message to the specified player.
-	 * @param PlayerId The Id of the destination player.
-	 * @param Reason An optional reason for the disconnection.
-	 */
-	void SendDisconnectPlayer(FPixelStreamingPlayerId PlayerId, const FString& Reason);
-
-	/**
-	 * [Player only] Sends an offer to the streamer.
-	 * @param SDP The local description offer.
-	 */
-	void SendOffer(const webrtc::SessionDescriptionInterface& SDP);
-
-	/**
-	 * [Player only] Sends an answer back to the streamer after receiving an offer.
-	 * @param SDP The local description answer.
-	 */
-	void SendAnswer(const webrtc::SessionDescriptionInterface& SDP);
-
-	/**
-	 * [Player only] Sends ICE candidate information to the streamer.
-	 * @param IceCandidate The ICE candidate information.
-	 */
-	void SendIceCandidate(const webrtc::IceCandidateInterface& IceCandidate);
-
-	/**
-	 * Enables or disables the keep alive pings on this connection. Receiving connections
-	 * should disable the keep alive since the reference signalling server will reject
-	 * pings from players and close the connection.
-	 * @param bKeepAlive True to enable keep alive pings. False to disable.
-	 */
-	void SetKeepAlive(bool bKeepAlive);
-
-	/**
-	 * Toggles automatic reconnecting when websocket is closed or unreachable.
-	 * @param bAutoReconnect True to enable auto reconnect. False to disable.
-	 */
-	void SetAutoReconnect(bool bAutoReconnect);
+	FPixelStreamingSignallingConnection(TSharedPtr<IPixelStreamingSignallingConnectionObserver> InObserver, FString InStreamerId = "");
+	virtual ~FPixelStreamingSignallingConnection();
+
+	/* IPixelStreamingSignallingConnection Interface */
+	virtual void TryConnect(FString URL) override;
+	virtual void Disconnect() override;
+	virtual bool IsConnected() const override;
+	virtual void SendOffer(FPixelStreamingPlayerId PlayerId, const webrtc::SessionDescriptionInterface& SDP) override;
+	virtual void SendAnswer(FPixelStreamingPlayerId PlayerId, const webrtc::SessionDescriptionInterface& SDP) override;
+	virtual void SendIceCandidate(FPixelStreamingPlayerId PlayerId, const webrtc::IceCandidateInterface& IceCandidate) override;
+	virtual void SendDisconnectPlayer(FPixelStreamingPlayerId PlayerId, const FString& Reason) override;
+	virtual void SendOffer(const webrtc::SessionDescriptionInterface& SDP) override;
+	virtual void SendAnswer(const webrtc::SessionDescriptionInterface& SDP) override;
+	virtual void SendIceCandidate(const webrtc::IceCandidateInterface& IceCandidate) override;
+	virtual void SetKeepAlive(bool bKeepAlive) override;
+	virtual void SetAutoReconnect(bool bAutoReconnect) override;
+	virtual void SendMessage(const FString& Msg) override;
 
 private:
 	void Connect(FString Url, bool bIsReconnect);
@@ -130,8 +67,6 @@ private:
 	void StartReconnectTimer();
 	void StopReconnectTimer();
 
-	void SendMessage(const FString& Msg);
-
 	template <typename FmtType, typename... T>
 	void PlayerError(FPixelStreamingPlayerId PlayerId, const FmtType& Msg, T... args)
 	{
@@ -141,8 +76,7 @@ private:
 	void PlayerError(FPixelStreamingPlayerId PlayerId, const FString& Msg);
 
 private:
-	FWebSocketFactory WebSocketFactory;
-	IPixelStreamingSignallingConnectionObserver& Observer;
+	TSharedPtr<IPixelStreamingSignallingConnectionObserver> Observer;
 	FString StreamerId;
 	FString Url;
 
