Integration of AVCodecs into PixelStreaming and support for new pixel formats

diff --git a/Engine/Plugins/Media/PixelStreaming/Source/PixelStreaming/Public/PixelStreamingStatNames.h b/Engine/Plugins/Media/PixelStreaming/Source/PixelStreaming/Public/PixelStreamingStatNames.h
--- a/Engine/Plugins/Media/PixelStreaming/Source/PixelStreaming/Public/PixelStreamingStatNames.h
+++ b/Engine/Plugins/Media/PixelStreaming/Source/PixelStreaming/Public/PixelStreamingStatNames.h
@@ -9,8 +9,8 @@ namespace PixelStreamingStatNames
 
 	const FName JitterBufferDelay		  = FName(TEXT("jitterBufferDelay"));
 	const FName FramesSent				  = FName(TEXT("framesSent"));
-	const FName FramesPerSecond			  = FName(TEXT("framesPerSecond"));
 	const FName FramesReceived			  = FName(TEXT("framesReceived"));
+	const FName FramesPerSecond			  = FName(TEXT("framesPerSecond"));
 	const FName FramesDecoded			  = FName(TEXT("framesDecoded"));
 	const FName FramesDropped			  = FName(TEXT("framesDropped"));
 	const FName FramesCorrupted			  = FName(TEXT("framesCorrupted"));
@@ -35,6 +35,7 @@ namespace PixelStreamingStatNames
 	const FName FrameWidth				  = FName(TEXT("frameWidth"));
 	const FName FrameHeight				  = FName(TEXT("frameHeight"));
 	const FName BytesSent				  = FName(TEXT("bytesSent"));
+	const FName BytesReceived			  = FName(TEXT("bytesReceived"));
 	const FName QPSum					  = FName(TEXT("qpSum"));
 	const FName TotalEncodeTime			  = FName(TEXT("totalEncodeTime"));
 	const FName TotalPacketSendDelay	  = FName(TEXT("totalPacketSendDelay"));
@@ -44,12 +45,15 @@ namespace PixelStreamingStatNames
 	const FName AvgSendDelay			  = FName(TEXT("packetSendDelay"));
 
 	// Calculated stats
-	const FName FramesSentPerSecond	   = FName(TEXT("transmitFps"));
-	const FName Bitrate				   = FName(TEXT("bitrate"));
-	const FName MeanQPPerSecond		   = FName(TEXT("qp"));
-	const FName MeanEncodeTime		   = FName(TEXT("encodeTime"));
-	const FName EncodedFramesPerSecond = FName(TEXT("encodeFps"));
-	const FName MeanSendDelay		   = FName(TEXT("captureToSend"));
-	const FName SourceFps			   = FName(TEXT("captureFps"));
+	const FName FramesSentPerSecond		= FName(TEXT("transmitFps"));
+	const FName FramesReceivedPerSecond	= FName(TEXT("receiveFps"));
+	const FName Bitrate					= FName(TEXT("bitrate"));
+	const FName MeanQPPerSecond			= FName(TEXT("qp"));
+	const FName MeanEncodeTime			= FName(TEXT("encodeTime"));
+	const FName EncodedFramesPerSecond	= FName(TEXT("encodeFps"));
+	const FName DecodedFramesPerSecond	= FName(TEXT("decodeFps"));
+	const FName MeanSendDelay			= FName(TEXT("captureToSend"));
+	const FName SourceFps				= FName(TEXT("captureFps"));
+	const FName Fps						= FName(TEXT("captureFps"));
 
 } // namespace PixelStreamingStatNames
\ No newline at end of file
diff --git a/Engine/Plugins/Media/PixelStreaming/Source/PixelStreaming/Public/PixelStreamingVideoInput.h b/Engine/Plugins/Media/PixelStreaming/Source/PixelStreaming/Public/PixelStreamingVideoInput.h
--- a/Engine/Plugins/Media/PixelStreaming/Source/PixelStreaming/Public/PixelStreamingVideoInput.h
+++ b/Engine/Plugins/Media/PixelStreaming/Source/PixelStreaming/Public/PixelStreamingVideoInput.h
@@ -6,6 +6,7 @@
 #include "PixelCaptureCapturerMultiFormat.h"
 #include "Delegates/IDelegateInstance.h"
 #include "PixelStreamingWebRTCIncludes.h"
+#include "PixelCaptureBufferFormat.h"
 
 /**
  * The input of the Pixel Streaming system. Frames enter the system when OnFrame is called.
diff --git a/Engine/Plugins/Media/PixelStreaming/Source/PixelStreaming/Public/PixelStreamingWebRTCIncludes.h b/Engine/Plugins/Media/PixelStreaming/Source/PixelStreaming/Public/PixelStreamingWebRTCIncludes.h
--- a/Engine/Plugins/Media/PixelStreaming/Source/PixelStreaming/Public/PixelStreamingWebRTCIncludes.h
+++ b/Engine/Plugins/Media/PixelStreaming/Source/PixelStreaming/Public/PixelStreamingWebRTCIncludes.h
@@ -17,6 +17,8 @@ THIRD_PARTY_INCLUDES_START
 	#pragma warning(disable : 6319 6323)
 #endif // PLATFORM_WINDOWS
 
+#include "api/rtc_error.h"
+
 #include "api/data_channel_interface.h"
 #include "api/media_stream_interface.h"
 #include "api/peer_connection_interface.h"
@@ -35,6 +37,6 @@ THIRD_PARTY_INCLUDES_START
 	#ifdef PF_MAX
 		#undef PF_MAX
 	#endif
-#endif //PLATFORM_WINDOWS
+#endif // PLATFORM_WINDOWS
 
 THIRD_PARTY_INCLUDES_END
diff --git a/Engine/Plugins/Media/PixelStreaming/Source/PixelStreamingBlueprint/Private/PixelStreamingStreamerComponent.cpp b/Engine/Plugins/Media/PixelStreaming/Source/PixelStreamingBlueprint/Private/PixelStreamingStreamerComponent.cpp
--- a/Engine/Plugins/Media/PixelStreaming/Source/PixelStreamingBlueprint/Private/PixelStreamingStreamerComponent.cpp
+++ b/Engine/Plugins/Media/PixelStreaming/Source/PixelStreamingBlueprint/Private/PixelStreamingStreamerComponent.cpp
@@ -31,9 +31,9 @@ void UPixelStreamingStreamerComponent::StartStreaming()
 		CreateStreamer();
 	}
 
-	if (StreamerInput && StreamerInput->GetVideoInput())
+	if (VideoInput && VideoInput->GetVideoInput())
 	{
-		Streamer->SetVideoInput(StreamerInput->GetVideoInput());
+		Streamer->SetVideoInput(VideoInput->GetVideoInput());
 	}
 
 	Streamer->SetCoupleFramerate(CoupleFramerate);
diff --git a/Engine/Plugins/Media/PixelStreaming/Source/PixelStreamingBlueprint/Public/PixelStreamingStreamerComponent.h b/Engine/Plugins/Media/PixelStreaming/Source/PixelStreamingBlueprint/Public/PixelStreamingStreamerComponent.h
--- a/Engine/Plugins/Media/PixelStreaming/Source/PixelStreamingBlueprint/Public/PixelStreamingStreamerComponent.h
+++ b/Engine/Plugins/Media/PixelStreaming/Source/PixelStreamingBlueprint/Public/PixelStreamingStreamerComponent.h
@@ -5,7 +5,7 @@
 #include "IPixelStreamingStreamer.h"
 #include "PixelStreamingPlayerId.h"
 #include "Components/ActorComponent.h"
-#include "PixelStreamingStreamerInput.h"
+#include "PixelStreamingStreamerVideoInput.h"
 #include "PixelStreamingStreamerComponent.generated.h"
 
 UCLASS(BlueprintType, Blueprintable, Category = "PixelStreaming", META = (DisplayName = "Streamer Component", BlueprintSpawnableComponent))
@@ -63,7 +63,7 @@ public:
 	bool CoupleFramerate = false;
 
 	UPROPERTY(EditAnywhere, Category = "PixelStreaming")
-	TObjectPtr<UPixelStreamingStreamerInput> StreamerInput = nullptr;
+	TObjectPtr<UPixelStreamingStreamerVideoInput> VideoInput = nullptr;
 
 private:
 	TSharedPtr<IPixelStreamingStreamer> Streamer;
diff --git a/Engine/Plugins/Media/PixelStreaming/Source/PixelStreamingBlueprintEditor/PixelStreamingBlueprintEditor.Build.cs b/Engine/Plugins/Media/PixelStreaming/Source/PixelStreamingBlueprintEditor/PixelStreamingBlueprintEditor.Build.cs
--- a/Engine/Plugins/Media/PixelStreaming/Source/PixelStreamingBlueprintEditor/PixelStreamingBlueprintEditor.Build.cs
+++ b/Engine/Plugins/Media/PixelStreaming/Source/PixelStreamingBlueprintEditor/PixelStreamingBlueprintEditor.Build.cs
@@ -29,8 +29,10 @@ namespace UnrealBuildTool.Rules
 				new string[] {
 					"Core",
 					"CoreUObject",
+					"SlateCore",
 					"PixelStreaming",
 					"PixelStreamingBlueprint",
+					"PixelStreamingEditor",
 					"UnrealEd"
 				});
 		}
diff --git a/Engine/Plugins/Media/PixelStreaming/Source/PixelStreamingBlueprintEditor/Private/AssetTypeActions_StreamerVideoInput.h b/Engine/Plugins/Media/PixelStreaming/Source/PixelStreamingBlueprintEditor/Private/AssetTypeActions_StreamerVideoInput.h
new file mode 100644
--- /dev/null
+++ b/Engine/Plugins/Media/PixelStreaming/Source/PixelStreamingBlueprintEditor/Private/AssetTypeActions_StreamerVideoInput.h
@@ -0,0 +1,25 @@
+// Copyright Epic Games, Inc. All Rights Reserved.
+
+#pragma once
+
+#include "CoreMinimal.h"
+#include "PixelStreamingStreamerVideoInput.h"
+#include "Toolkits/IToolkitHost.h"
+#include "AssetTypeActions_Base.h"
+
+class FAssetTypeActions_StreamerVideoInput : public FAssetTypeActions_Base
+{
+public:
+	// IAssetTypeActions Implementation
+	virtual FText GetName() const override { return NSLOCTEXT("PixelStreaming", "AssetTypeActions_StreamerVideoInput", "Streamer Video Input Actions"); }
+	virtual FColor GetTypeColor() const override { return FColor(192,64,64); }
+	virtual UClass* GetSupportedClass() const override { return UPixelStreamingStreamerVideoInput::StaticClass(); }
+	virtual bool CanFilter() override { return false; }
+	virtual uint32 GetCategories() override
+	{
+		IAssetTools& AssetTools = FModuleManager::LoadModuleChecked<FAssetToolsModule>("AssetTools").Get();
+		return AssetTools.RegisterAdvancedAssetCategory("PixelStreaming", NSLOCTEXT("PixelStreaming", "AssetCategoryDisplayName", "PixelStreaming"));
+		//return EAssetTypeCategories::Misc;
+	}
+	virtual bool IsImportedAsset() const override { return false; }
+};
diff --git a/Engine/Plugins/Media/PixelStreaming/Source/PixelStreamingBlueprintEditor/Private/PixelStreamingBlueprintEditorModule.cpp b/Engine/Plugins/Media/PixelStreaming/Source/PixelStreamingBlueprintEditor/Private/PixelStreamingBlueprintEditorModule.cpp
--- a/Engine/Plugins/Media/PixelStreaming/Source/PixelStreamingBlueprintEditor/Private/PixelStreamingBlueprintEditorModule.cpp
+++ b/Engine/Plugins/Media/PixelStreaming/Source/PixelStreamingBlueprintEditor/Private/PixelStreamingBlueprintEditorModule.cpp
@@ -1,7 +1,10 @@
 // Copyright Epic Games, Inc. All Rights Reserved.
 
 #include "Modules/ModuleManager.h"
-#include "AssetTypeActions_StreamerInput.h"
+#include "AssetTypeActions_StreamerVideoInput.h"
+#include "PixelStreamingStyle.h"
+
+#define IMAGE_BRUSH_SVG(Style, RelativePath, ...) FSlateVectorImageBrush(Style.RootToContentDir(RelativePath, TEXT(".svg")), __VA_ARGS__)
 
 class FPixelStreamingBlueprintEditorModule : public IModuleInterface
 {
@@ -11,8 +14,19 @@ private:
 	void StartupModule() override
 	{
 		IAssetTools& AssetTools = FModuleManager::LoadModuleChecked<FAssetToolsModule>("AssetTools").Get();
-		AssetTools.RegisterAssetTypeActions(MakeShared<FAssetTypeActions_StreamerInput>());
+		AssetTools.RegisterAssetTypeActions(MakeShared<FAssetTypeActions_StreamerVideoInput>());
+
+		FSlateStyleSet& StyleInstance = UE::EditorPixelStreaming::FPixelStreamingStyle::Get();
+
+		const FVector2D Icon20x20(20.0f, 20.0f);
+		const FVector2D Icon64x64(64.0f, 64.0f);
+
+		StyleInstance.Set("ClassThumbnail.PixelStreamingStreamerVideoInputBackBuffer", new IMAGE_BRUSH_SVG(StyleInstance, "PixelStreaming_64", Icon64x64));
+		StyleInstance.Set("ClassIcon.PixelStreamingStreamerVideoInputBackBuffer", new IMAGE_BRUSH_SVG(StyleInstance, "PixelStreaming_20", Icon20x20));
+		StyleInstance.Set("ClassThumbnail.PixelStreamingStreamerVideoInputRenderTarget",new IMAGE_BRUSH_SVG(StyleInstance, "PixelStreaming_64", Icon64x64));
+		StyleInstance.Set("ClassIcon.PixelStreamingStreamerVideoInputRenderTarget", new IMAGE_BRUSH_SVG(StyleInstance, "PixelStreaming_20", Icon20x20));
 	}
+
 	void ShutdownModule() override {}
 };
 
diff --git a/Engine/Plugins/Media/PixelStreaming/Source/PixelStreamingEditor/PixelStreamingEditor.build.cs b/Engine/Plugins/Media/PixelStreaming/Source/PixelStreamingEditor/PixelStreamingEditor.build.cs
--- a/Engine/Plugins/Media/PixelStreaming/Source/PixelStreamingEditor/PixelStreamingEditor.build.cs
+++ b/Engine/Plugins/Media/PixelStreaming/Source/PixelStreamingEditor/PixelStreamingEditor.build.cs
@@ -27,9 +27,9 @@ namespace UnrealBuildTool.Rules
 				"Engine",
 				"Projects",
 				"RenderCore",
+				"Renderer",
 				"RHI",
 				"PixelStreaming",
-				"AVEncoder",
 				"Slate",
 				"SlateCore",
 				"EngineSettings",
@@ -38,7 +38,9 @@ namespace UnrealBuildTool.Rules
 				"PixelCapture",
 				"PixelStreamingServers",
 				"HTTP",
-				"Sockets"
+				"Sockets",
+				"ApplicationCore",
+				"PixelStreamingInput"
 			});
 
 			if(Target.bBuildEditor)
diff --git a/Engine/Plugins/Media/PixelStreaming/Source/PixelStreamingEditor/Private/PixelStreamingCommands.cpp b/Engine/Plugins/Media/PixelStreaming/Source/PixelStreamingEditor/Private/PixelStreamingCommands.cpp
--- a/Engine/Plugins/Media/PixelStreaming/Source/PixelStreamingEditor/Private/PixelStreamingCommands.cpp
+++ b/Engine/Plugins/Media/PixelStreaming/Source/PixelStreamingEditor/Private/PixelStreamingCommands.cpp
@@ -16,6 +16,7 @@ namespace UE::EditorPixelStreaming
         UI_COMMAND(VP8, "VP8", "VP8", EUserInterfaceActionType::RadioButton, FInputChord());
         UI_COMMAND(VP9, "VP9", "VP9", EUserInterfaceActionType::RadioButton, FInputChord());
         UI_COMMAND(H264, "H264", "H264", EUserInterfaceActionType::RadioButton, FInputChord());
+		UI_COMMAND(H265, "H265 (experimental)", "H265", EUserInterfaceActionType::RadioButton, FInputChord());
     }
 #undef LOCTEXT_NAMESPACE
 }
\ No newline at end of file
diff --git a/Engine/Plugins/Media/PixelStreaming/Source/PixelStreamingEditor/Private/PixelStreamingCommands.h b/Engine/Plugins/Media/PixelStreaming/Source/PixelStreamingEditor/Private/PixelStreamingCommands.h
--- a/Engine/Plugins/Media/PixelStreaming/Source/PixelStreamingEditor/Private/PixelStreamingCommands.h
+++ b/Engine/Plugins/Media/PixelStreaming/Source/PixelStreamingEditor/Private/PixelStreamingCommands.h
@@ -26,6 +26,7 @@ namespace UE::EditorPixelStreaming
         TSharedPtr<FUICommandInfo> VP8;
         TSharedPtr<FUICommandInfo> VP9;
         TSharedPtr<FUICommandInfo> H264;
+		TSharedPtr<FUICommandInfo> H265;
         TSharedPtr<FUICommandInfo> StartSignalling;
         TSharedPtr<FUICommandInfo> StopSignalling;
         TSharedPtr<FUICommandInfo> StreamLevelEditor;
diff --git a/Engine/Plugins/Media/PixelStreaming/Source/PixelStreamingEditor/Private/PixelStreamingEditorModule.cpp b/Engine/Plugins/Media/PixelStreaming/Source/PixelStreamingEditor/Private/PixelStreamingEditorModule.cpp
--- a/Engine/Plugins/Media/PixelStreaming/Source/PixelStreamingEditor/Private/PixelStreamingEditorModule.cpp
+++ b/Engine/Plugins/Media/PixelStreaming/Source/PixelStreamingEditor/Private/PixelStreamingEditorModule.cpp
@@ -12,12 +12,14 @@
 #include "Settings.h"
 #include "UnrealEngine.h"
 #include "Interfaces/IMainFrameModule.h"
-#include "PixelStreamingProtocol.h"
+#include "PixelStreamingInputProtocol.h"
 #include "PixelStreamingUtils.h"
 #include "GenericPlatform/GenericPlatformHttp.h"
 #include "LevelEditorViewport.h"
 #include "Slate/SceneViewport.h"
 #include "SLevelViewport.h"
+#include "PixelStreamingInputEnums.h"
+#include "IPixelStreamingInputModule.h"
 
 DECLARE_LOG_CATEGORY_EXTERN(LogPixelStreamingEditor, Log, All);
 DEFINE_LOG_CATEGORY(LogPixelStreamingEditor);
@@ -41,7 +43,8 @@ void FPixelStreamingEditorModule::StartupModule()
 	Settings->PostEditChange();
 
 	Settings::InitialiseSettings();
-		
+	bUseExternalSignallingServer = Settings::CVarEditorPixelStreamingUseRemoteSignallingServer.GetValueOnAnyThread();
+
 	IPixelStreamingModule& Module = IPixelStreamingModule::Get();
 	Module.OnReady().AddRaw(this, &FPixelStreamingEditorModule::InitEditorStreaming);
 }
@@ -59,27 +62,39 @@ void FPixelStreamingEditorModule::InitEditorStreaming(IPixelStreamingModule& Mod
 		return;
 	}
 
+	IPixelStreamingInputModule& InputModule = IPixelStreamingInputModule::Get();
+
+	// TODO: This code diverged and needs attention from subject matter exports to appropriately reconcile
 	// Add custom handle for { type: "Command", Resolution.Width: "1920", Resolution.Height: "1080" } when doing Editor streaming
 	// because we cannot resize the game viewport, but instead want to resize the parent window.
-	TSharedPtr<IPixelStreamingInputHandler> InputHandler = Streamer->GetInputHandler().Pin();
-	if(InputHandler) 
-	{
-		InputHandler->SetCommandHandler("Resolution.Width", [](FString Descriptor, FString WidthString){
-			bool bSuccess;
-			FString HeightString;
-			UE::PixelStreaming::ExtractJsonFromDescriptor(Descriptor, TEXT("Resolution.Height"), HeightString, bSuccess);
-			int Width = FCString::Atoi(*WidthString);
-			int Height = FCString::Atoi(*HeightString);
-			if (Width < 1 || Height < 1)
-			{
-				return;
-			}
-
-			TSharedPtr<SWindow> ParentWindow = IMainFrameModule::Get().GetParentWindow();
-			ParentWindow->Resize(FVector2D(Width, Height));
-			FSlateApplication::Get().OnSizeChanged(ParentWindow->GetNativeWindow().ToSharedRef(), Width, Height);
-		});
-	}
+	// TSharedPtr<IPixelStreamingInputHandler> InputHandler = Streamer->GetInputHandler().Pin();
+	// if(InputHandler) 
+	// {
+	// 	InputHandler->SetCommandHandler("Resolution.Width", [](FString Descriptor, FString WidthString){
+	// 		bool bSuccess;
+	// 		FString HeightString;
+	// 		UE::PixelStreaming::ExtractJsonFromDescriptor(Descriptor, TEXT("Resolution.Height"), HeightString, bSuccess);
+	// 		int Width = FCString::Atoi(*WidthString);
+	// 		int Height = FCString::Atoi(*HeightString);
+	// 		if (Width < 1 || Height < 1)
+	// 		{
+	// 			return;
+	// 		}
+
+	// 		TSharedPtr<SWindow> ParentWindow = IMainFrameModule::Get().GetParentWindow();
+	// 		ParentWindow->Resize(FVector2D(Width, Height));
+	// 		FSlateApplication::Get().OnSizeChanged(ParentWindow->GetNativeWindow().ToSharedRef(), Width, Height);
+	// 	});
+	// }
+	
+	// InputModule.RegisterMessage(
+	// 	EPixelStreamingMessageDirection::ToStreamer,
+	// 	"Command",
+	// 	FPixelStreamingInputMessage(51),
+	// 	[ExtendedHandler](FMemoryReader Ar) {
+	// 		ExtendedHandler(Ar);
+	// 	});
+	// }
 
 	// Give the editor streamer the default url if the user hasn't specified one when launching the editor
 	if (Streamer->GetSignallingServerURL().IsEmpty())
@@ -89,15 +104,15 @@ void FPixelStreamingEditorModule::InitEditorStreaming(IPixelStreamingModule& Mod
 		SignallingDomain = TEXT("ws://127.0.0.1");
 
 		Streamer->SetSignallingServerURL(FString::Printf(TEXT("%s:%d"), *SignallingDomain, StreamerPort));
-	} 
+	}
 	else
 	{
 		FString SpecifiedSignallingURL = Streamer->GetSignallingServerURL();
 		TOptional<uint16> ExtractedStreamerPort = FGenericPlatformHttp::GetUrlPort(SpecifiedSignallingURL);
-		StreamerPort = (int32) ExtractedStreamerPort.Get(8888);
+		StreamerPort = (int32)ExtractedStreamerPort.Get(8888);
 
 		FString ExtractedSignallingDomain = FGenericPlatformHttp::GetUrlDomain(SpecifiedSignallingURL);
-		if(FGenericPlatformHttp::IsSecureProtocol(SpecifiedSignallingURL).Get(false))
+		if (FGenericPlatformHttp::IsSecureProtocol(SpecifiedSignallingURL).Get(false))
 		{
 			SignallingDomain = FString::Printf(TEXT("wss://%s"), *ExtractedSignallingDomain);
 		}
@@ -107,15 +122,14 @@ void FPixelStreamingEditorModule::InitEditorStreaming(IPixelStreamingModule& Mod
 		}
 	}
 
-	IMainFrameModule::Get().OnMainFrameCreationFinished().AddLambda([&](TSharedPtr<SWindow> RootWindow, bool bIsRunningStartupDialog)
-	{
+	IMainFrameModule::Get().OnMainFrameCreationFinished().AddLambda([&](TSharedPtr<SWindow> RootWindow, bool bIsRunningStartupDialog) {
 		MaybeResizeEditor(RootWindow);
 
 		// We don't want to show tooltips in render off screen as they're currently broken
 		bool bIsRenderingOffScreen = FParse::Param(FCommandLine::Get(), TEXT("RenderOffScreen"));
 		FSlateApplication::Get().SetAllowTooltips(!bIsRenderingOffScreen);
 
-		if(UE::EditorPixelStreaming::Settings::CVarEditorPixelStreamingStartOnLaunch.GetValueOnAnyThread())
+		if (UE::EditorPixelStreaming::Settings::CVarEditorPixelStreamingStartOnLaunch.GetValueOnAnyThread())
 		{
 			StartStreaming(UE::EditorPixelStreaming::EStreamTypes::Editor);
 		}
@@ -133,7 +147,7 @@ void FPixelStreamingEditorModule::StartStreaming(UE::EditorPixelStreaming::EStre
 	}
 
 	StreamType = InStreamType;
-	switch(InStreamType)
+	switch (InStreamType)
 	{
 		case UE::EditorPixelStreaming::EStreamTypes::LevelEditorViewport:
 		{
@@ -150,7 +164,7 @@ void FPixelStreamingEditorModule::StartStreaming(UE::EditorPixelStreaming::EStre
 			Streamer->SetTargetWindow(SceneViewport->FindWindow());
 			Streamer->SetInputHandlerType(EPixelStreamingInputType::RouteToWindow);
 			Streamer->SetVideoInput(FPixelStreamingVideoInputViewport::Create());
-		}	
+		}
 		break;
 		case UE::EditorPixelStreaming::EStreamTypes::Editor:
 		{
@@ -164,11 +178,11 @@ void FPixelStreamingEditorModule::StartStreaming(UE::EditorPixelStreaming::EStre
 		{
 			UE_LOG(LogPixelStreamingEditor, Warning, TEXT("Specified Stream Type doesn't have an associated FPixelStreamingVideoInput"));
 		}
-		// Return here as we don't want to start streaming if we didn't set a viewport
-		return;
+			// Return here as we don't want to start streaming if we didn't set a viewport
+			return;
 	}
 
-	if(!bUseExternalSignallingServer)
+	if (!bUseExternalSignallingServer)
 	{
 		Streamer->SetSignallingServerURL(FString::Printf(TEXT("%s:%d"), *SignallingDomain, StreamerPort));
 		StartSignalling();
@@ -187,7 +201,7 @@ void FPixelStreamingEditorModule::StopStreaming()
 		return;
 	}
 
-	if(!bUseExternalSignallingServer)
+	if (!bUseExternalSignallingServer)
 	{
 		StopSignalling();
 	}
@@ -206,19 +220,19 @@ UE::EditorPixelStreaming::EStreamTypes FPixelStreamingEditorModule::GetStreamTyp
 void FPixelStreamingEditorModule::StartSignalling()
 {
 	bool bAlreadyLaunched = SignallingServer.IsValid() && SignallingServer->HasLaunched();
-	if(bAlreadyLaunched)
+	if (bAlreadyLaunched)
 	{
 		return;
 	}
 
 	// Download Pixel Streaming servers/frontend if we want to use a browser to view Pixel Streaming output
 	// but only attempt this is we haven't already started a download before.
-	if(!DownloadProcess.IsValid())
+	if (!DownloadProcess.IsValid())
 	{
 		DownloadProcess = UE::PixelStreamingServers::DownloadPixelStreamingServers(/*bSkipIfPresent*/ true);
-		if(DownloadProcess.IsValid())
+		if (DownloadProcess.IsValid())
 		{
-			DownloadProcess->OnCompleted().BindLambda([this](int ExitCode){
+			DownloadProcess->OnCompleted().BindLambda([this](int ExitCode) {
 				StopSignalling();
 				StartSignalling();
 			});
@@ -241,16 +255,16 @@ void FPixelStreamingEditorModule::StartSignalling()
 
 void FPixelStreamingEditorModule::StopSignalling()
 {
-	if(SignallingServer.IsValid())
+	if (SignallingServer.IsValid())
 	{
 		SignallingServer->Stop();
 		SignallingServer.Reset();
 	}
 }
 
-TSharedPtr<UE::PixelStreamingServers::IServer> FPixelStreamingEditorModule::GetSignallingServer() 
+TSharedPtr<UE::PixelStreamingServers::IServer> FPixelStreamingEditorModule::GetSignallingServer()
 {
-	if(SignallingServer.IsValid())
+	if (SignallingServer.IsValid())
 	{
 		return SignallingServer;
 	}
@@ -288,7 +302,7 @@ FPixelStreamingEditorModule* FPixelStreamingEditorModule::GetModule()
 
 bool FPixelStreamingEditorModule::ParseResolution(const TCHAR* InResolution, uint32& OutX, uint32& OutY)
 {
-	if(*InResolution)
+	if (*InResolution)
 	{
 		FString CmdString(InResolution);
 		CmdString = CmdString.TrimStartAndEnd().ToLower();
@@ -301,9 +315,9 @@ bool FPixelStreamingEditorModule::ParseResolution(const TCHAR* InResolution, uin
 
 		// Find separator between values (Example of expected format: 1280x768)
 		const TCHAR* YValue = NULL;
-		if(FCString::Strchr(*CmdString,'x'))
+		if (FCString::Strchr(*CmdString, 'x'))
 		{
-			YValue = const_cast<TCHAR*> (FCString::Strchr(*CmdString,'x')+1);
+			YValue = const_cast<TCHAR*>(FCString::Strchr(*CmdString, 'x') + 1);
 			YString = YValue;
 			// Remove any whitespace from the end of the string
 			YString = YString.TrimStartAndEnd();
@@ -311,7 +325,7 @@ bool FPixelStreamingEditorModule::ParseResolution(const TCHAR* InResolution, uin
 
 		// If the Y dimensional value exists then setup to use the specified resolution.
 		uint32 Y = 0;
-		if ( YValue && YString.Len() > 0 )
+		if (YValue && YString.Len() > 0)
 		{
 			if (YString.IsNumeric())
 			{
@@ -351,7 +365,7 @@ void FPixelStreamingEditorModule::MaybeResizeEditor(TSharedPtr<SWindow> RootWind
 		}
 	}
 
-	if(bSuccess)
+	if (bSuccess)
 	{
 		RootWindow->Resize(FVector2D(ResolutionX, ResolutionY));
 		FSlateApplication::Get().OnSizeChanged(RootWindow->GetNativeWindow().ToSharedRef(), ResolutionX, ResolutionY);
diff --git a/Engine/Plugins/Media/PixelStreaming/Source/PixelStreamingEditor/Private/PixelStreamingStyle.cpp b/Engine/Plugins/Media/PixelStreaming/Source/PixelStreamingEditor/Private/PixelStreamingStyle.cpp
--- a/Engine/Plugins/Media/PixelStreaming/Source/PixelStreamingEditor/Private/PixelStreamingStyle.cpp
+++ b/Engine/Plugins/Media/PixelStreaming/Source/PixelStreamingEditor/Private/PixelStreamingStyle.cpp
@@ -65,7 +65,7 @@ namespace UE::EditorPixelStreaming
         }
     }
     
-    const ISlateStyle& FPixelStreamingStyle::Get()
+    FSlateStyleSet& FPixelStreamingStyle::Get()
     {
         return *StyleInstance;
     }
diff --git a/Engine/Plugins/Media/PixelStreaming/Source/PixelStreamingEditor/Private/PixelStreamingToolbar.cpp b/Engine/Plugins/Media/PixelStreaming/Source/PixelStreamingEditor/Private/PixelStreamingToolbar.cpp
--- a/Engine/Plugins/Media/PixelStreaming/Source/PixelStreamingEditor/Private/PixelStreamingToolbar.cpp
+++ b/Engine/Plugins/Media/PixelStreaming/Source/PixelStreamingEditor/Private/PixelStreamingToolbar.cpp
@@ -200,6 +200,25 @@ namespace UE::EditorPixelStreaming
 			})
 		);
 
+		PluginCommands->MapAction(
+			FPixelStreamingCommands::Get().H265,
+			FExecuteAction::CreateLambda([Module = &PixelStreamingModule]()
+			{
+				Module->SetCodec(EPixelStreamingCodec::H265);
+			}),
+			FCanExecuteAction::CreateLambda([Module = &PixelStreamingModule] { 
+				if(TSharedPtr<IPixelStreamingStreamer> Streamer = Module->GetStreamer(Module->GetDefaultStreamerID()))	
+				{
+					return !Streamer->IsStreaming();
+				}
+				return false;
+			}),
+			FIsActionChecked::CreateLambda([Module = &PixelStreamingModule]()
+			{ 
+				return Module->GetCodec() == EPixelStreamingCodec::H265;
+			})
+		);
+
 		UToolMenus::RegisterStartupCallback(FSimpleMulticastDelegate::FDelegate::CreateRaw(this, &FPixelStreamingToolbar::RegisterMenus));
 	}
 	
@@ -503,6 +522,7 @@ namespace UE::EditorPixelStreaming
 		MenuBuilder.BeginSection("Codec", LOCTEXT("PixelStreamingCodecSettings", "Codec"));
 		{
 			MenuBuilder.AddMenuEntry(FPixelStreamingCommands::Get().H264);
+			MenuBuilder.AddMenuEntry(FPixelStreamingCommands::Get().H265);
 			MenuBuilder.AddMenuEntry(FPixelStreamingCommands::Get().VP8);
 			MenuBuilder.AddMenuEntry(FPixelStreamingCommands::Get().VP9);
 		}
diff --git a/Engine/Plugins/Media/PixelStreaming/Source/PixelStreamingEditor/Private/PixelStreamingVideoInputBackBufferComposited.cpp b/Engine/Plugins/Media/PixelStreaming/Source/PixelStreamingEditor/Private/PixelStreamingVideoInputBackBufferComposited.cpp
--- a/Engine/Plugins/Media/PixelStreaming/Source/PixelStreamingEditor/Private/PixelStreamingVideoInputBackBufferComposited.cpp
+++ b/Engine/Plugins/Media/PixelStreaming/Source/PixelStreamingEditor/Private/PixelStreamingVideoInputBackBufferComposited.cpp
@@ -2,6 +2,7 @@
 
 #include "PixelStreamingVideoInputBackBufferComposited.h"
 #include "Settings.h"
+#include "PixelStreamingEditorUtils.h"
 #include "Utils.h"
 #include "UtilsRender.h"
 #include "ToStringExtensions.h"
@@ -17,6 +18,8 @@
 #include "Application/SlateApplicationBase.h"
 #include "Widgets/SWindow.h"
 
+DECLARE_LOG_CATEGORY_EXTERN(LogPixelStreamingBackBufferComposited, Log, VeryVerbose);
+DEFINE_LOG_CATEGORY(LogPixelStreamingBackBufferComposited);
 
 TSharedPtr<FPixelStreamingVideoInputBackBufferComposited> FPixelStreamingVideoInputBackBufferComposited::Create()
 {
@@ -36,9 +39,7 @@ TSharedPtr<FPixelStreamingVideoInputBackBufferComposited> FPixelStreamingVideoIn
 
 FPixelStreamingVideoInputBackBufferComposited::FPixelStreamingVideoInputBackBufferComposited()
 {
-	// Initialize our composited frame texture to a default size. It'll be enlarged / shrunk as needed
-	CompositedFrame = UE::PixelStreaming::CreateRHITexture(DefaultSize.X, DefaultSize.Y);
-	CompositedFrameSize = MakeShared<FIntPoint>(DefaultSize);
+	CompositedFrameSize = MakeShared<FIntPoint>();
 	UE::PixelStreaming::DoOnGameThread([this]() {
 		FSlateApplication::Get().OnPreTick().AddLambda([this](float DeltaTime) {
 			TopLevelWindowsCriticalSection.Lock();
@@ -69,10 +70,24 @@ void FPixelStreamingVideoInputBackBufferComposited::OnBackBufferReady(SWindow& S
 		}
 	}
 
+	UE_LOG(LogPixelStreamingBackBufferComposited, Verbose, TEXT("=== Window Rendered ==="));
+	UE_LOG(LogPixelStreamingBackBufferComposited, Verbose, TEXT("Type: %s"), UE::EditorPixelStreaming::ToString(SlateWindow.GetType()));
+
+	uint32 FrameX = FrameBuffer->GetSizeXY().X;
+	uint32 FrameY = FrameBuffer->GetSizeXY().Y;
+	FString Hash = UE::EditorPixelStreaming::HashWindow(SlateWindow, FrameBuffer);
+	FTextureRHIRef StagingTexture = StagingTextures.FindRef(Hash);
+	if(!StagingTexture.IsValid()) 
+	{
+		UE_LOG(LogPixelStreamingBackBufferComposited, Verbose, TEXT("Creating new staging texture: %dx%d"), FrameX, FrameY);
+		StagingTexture = UE::PixelStreaming::CreateRHITexture(FrameX, FrameY);
+		StagingTextures.Add(Hash, StagingTexture);
+	}
+
+
 	FRHICommandListImmediate& RHICmdList = FRHICommandListExecutor::GetImmediateCommandList();
-	FTextureRHIRef StagingTexture = UE::PixelStreaming::CreateRHITexture(FrameBuffer->GetSizeXY().X, FrameBuffer->GetSizeXY().Y);
 	// Re-render FrameBuffer to StagingTexture (ensure format match)
-	UE::PixelStreaming::CopyTexture(RHICmdList, FrameBuffer, StagingTexture, nullptr);
+	UE::PixelStreaming::CopyTextureRDG(RHICmdList, FrameBuffer, StagingTexture);
 	
 	TopLevelWindowTextures.Add(&SlateWindow, StagingTexture);
 	TopLevelWindowsCriticalSection.Lock();
@@ -82,14 +97,16 @@ void FPixelStreamingVideoInputBackBufferComposited::OnBackBufferReady(SWindow& S
 	{
 		if (TopLevelWindowTextures.FindRef(&Window.Get()))
 		{
-			WindowsRendered += 1;
+			++WindowsRendered;
 		}
 	}
 
+	
 	if (WindowsRendered == TopLevelWindows.Num())
 	{
 		// We have all the textures needed, let's composite them
 		// Need to iterate over the TopLevelWindows array as this array contains the ordered windows
+		FTextureRHIRef CompositedFrame;
 		for (TSharedRef<SWindow> Window : TopLevelWindows)
 		{
 			FTextureRHIRef Texture = TopLevelWindowTextures.FindRef(&Window.Get());
@@ -98,7 +115,14 @@ void FPixelStreamingVideoInputBackBufferComposited::OnBackBufferReady(SWindow& S
 				Window->GetSizeInScreen() == FVector2D(0, 0))
 			{
 				continue;
-			}			
+			}	
+
+			if(!CompositedFrame.IsValid()) 
+			{
+				// This is the first texture being processed, make it our base for the composited frame
+				CompositedFrame = Texture;
+				continue;
+			}
 
 			uint32 CompositedX = CompositedFrame->GetSizeXY().X;
 			uint32 CompositedY = CompositedFrame->GetSizeXY().Y;
@@ -155,13 +179,6 @@ void FPixelStreamingVideoInputBackBufferComposited::OnBackBufferReady(SWindow& S
 		*CompositedFrameSize.Get() = CompositedFrame->GetSizeXY();
 		IPixelStreamingModule& Module = IPixelStreamingModule::Get();
 		Module.GetStreamer(Module.GetDefaultStreamerID())->SetTargetScreenSize(CompositedFrameSize);
-		bRecreateTexture = true;
 	}
 	TopLevelWindowsCriticalSection.Unlock();
-
-	if (bRecreateTexture)
-	{
-		CompositedFrame = UE::PixelStreaming::CreateRHITexture(DefaultSize.X, DefaultSize.Y);
-		bRecreateTexture = false;
-	}
 }
diff --git a/Engine/Plugins/Media/PixelStreaming/Source/PixelStreamingEditor/Private/Settings.cpp b/Engine/Plugins/Media/PixelStreaming/Source/PixelStreamingEditor/Private/Settings.cpp
--- a/Engine/Plugins/Media/PixelStreaming/Source/PixelStreamingEditor/Private/Settings.cpp
+++ b/Engine/Plugins/Media/PixelStreaming/Source/PixelStreamingEditor/Private/Settings.cpp
@@ -11,10 +11,17 @@ namespace UE::EditorPixelStreaming::Settings
 		TEXT("Start streaming the Editor as soon as it launches. Default: false"),
 		ECVF_Default);
 
+	TAutoConsoleVariable<bool> CVarEditorPixelStreamingUseRemoteSignallingServer(
+		TEXT("PixelStreaming.Editor.UseRemoteSignallingServer"),
+		false,
+		TEXT("Use a remote signalling server. Default: false"),
+		ECVF_Default);
+
     void InitialiseSettings()
     {
 		using namespace UE::PixelStreaming;
         // Options parse (if these exist they are set to true)
 		CommandLineParseOption(TEXT("EditorPixelStreamingStartOnLaunch"), CVarEditorPixelStreamingStartOnLaunch);
+		CommandLineParseOption(TEXT("EditorPixelStreamingUseRemoteSignallingServer"), CVarEditorPixelStreamingUseRemoteSignallingServer);
     }
 }
\ No newline at end of file
diff --git a/Engine/Plugins/Media/PixelStreaming/Source/PixelStreamingEditor/Private/Settings.h b/Engine/Plugins/Media/PixelStreaming/Source/PixelStreamingEditor/Private/Settings.h
--- a/Engine/Plugins/Media/PixelStreaming/Source/PixelStreamingEditor/Private/Settings.h
+++ b/Engine/Plugins/Media/PixelStreaming/Source/PixelStreamingEditor/Private/Settings.h
@@ -12,5 +12,6 @@ namespace UE::EditorPixelStreaming::Settings
 
 	// Begin Editor CVars
 	extern TAutoConsoleVariable<bool> CVarEditorPixelStreamingStartOnLaunch;
+	extern TAutoConsoleVariable<bool> CVarEditorPixelStreamingUseRemoteSignallingServer;
     // End Editor CVars
 }
\ No newline at end of file
diff --git a/Engine/Plugins/Media/PixelStreaming/Source/PixelStreamingEditor/Public/PixelStreamingEditorUtils.h b/Engine/Plugins/Media/PixelStreaming/Source/PixelStreamingEditor/Public/PixelStreamingEditorUtils.h
--- a/Engine/Plugins/Media/PixelStreaming/Source/PixelStreamingEditor/Public/PixelStreamingEditorUtils.h
+++ b/Engine/Plugins/Media/PixelStreaming/Source/PixelStreamingEditor/Public/PixelStreamingEditorUtils.h
@@ -3,6 +3,10 @@
 #pragma once
 
 #include "Containers/UnrealString.h"
+#include "GenericPlatform/GenericWindowDefinition.h"
+#include "RHIResources.h"
+
+class SWindow;
 
 namespace UE::EditorPixelStreaming
 {
@@ -13,18 +17,7 @@ namespace UE::EditorPixelStreaming
 		VCam = 2
 	};
 
-	inline FString ToString(EStreamTypes StreamType)
-	{
-		switch(StreamType)
-		{
-			case EStreamTypes::LevelEditorViewport:
-				return TEXT("the Level Editor");
-			case EStreamTypes::Editor:
-				return TEXT("the Full Editor");
-			case EStreamTypes::VCam:
-				return TEXT("a Virtual Camera");
-			default:
-				return TEXT("Unknown stream type!");
-		}
-	}
+	FString ToString(EStreamTypes StreamType);
+	const TCHAR* ToString(EWindowType Type);
+	const FString HashWindow(SWindow& SlateWindow, const FTexture2DRHIRef& FrameBuffer);
 }
\ No newline at end of file
diff --git a/Engine/Plugins/Media/PixelStreaming/Source/PixelStreamingEditor/Public/PixelStreamingVideoInputBackBufferComposited.h b/Engine/Plugins/Media/PixelStreaming/Source/PixelStreamingEditor/Public/PixelStreamingVideoInputBackBufferComposited.h
--- a/Engine/Plugins/Media/PixelStreaming/Source/PixelStreamingEditor/Public/PixelStreamingVideoInputBackBufferComposited.h
+++ b/Engine/Plugins/Media/PixelStreaming/Source/PixelStreamingEditor/Public/PixelStreamingVideoInputBackBufferComposited.h
@@ -28,10 +28,8 @@ private:
 
 	TArray<TSharedRef<SWindow>> TopLevelWindows;
 	TMap<SWindow*, FTextureRHIRef> TopLevelWindowTextures;
+	TMap<FString, FTextureRHIRef> StagingTextures;
 
 	FCriticalSection TopLevelWindowsCriticalSection;
-	FTexture2DRHIRef CompositedFrame;
 	TSharedPtr<FIntPoint> CompositedFrameSize;
-	bool bRecreateTexture = false;
-	FIntPoint DefaultSize = FIntPoint(1, 1);
 };
\ No newline at end of file
diff --git a/Engine/Plugins/Media/PixelStreaming/Source/PixelStreamingHMD/Private/PixelStreamingHMD.cpp b/Engine/Plugins/Media/PixelStreaming/Source/PixelStreamingHMD/Private/PixelStreamingHMD.cpp
new file mode 100644
--- /dev/null
+++ b/Engine/Plugins/Media/PixelStreaming/Source/PixelStreamingHMD/Private/PixelStreamingHMD.cpp
@@ -0,0 +1,227 @@
+// Copyright Epic Games, Inc. All Rights Reserved.
+
+#include "PixelStreamingHMD.h"
+#include "Misc/App.h"
+#include "Modules/ModuleManager.h"
+#include "EngineGlobals.h"
+#include "Engine/Engine.h"
+#include "GameFramework/PlayerController.h"
+#include "Engine/LocalPlayer.h"
+#include "SceneRendering.h"
+#include "PostProcess/PostProcessHMD.h"
+#include "GameFramework/WorldSettings.h"
+#include "Settings.h"
+
+FPixelStreamingHMD::FPixelStreamingHMD(const FAutoRegister& AutoRegister)
+	: FHeadMountedDisplayBase(nullptr)
+	, FHMDSceneViewExtension(AutoRegister)
+	, CurHmdTransform(FTransform::Identity)
+	, WorldToMeters(100.0f)
+	, InterpupillaryDistance(0.064f)
+	, bStereoEnabled(true)
+{
+}
+
+FPixelStreamingHMD::~FPixelStreamingHMD()
+{
+}
+
+float FPixelStreamingHMD::GetWorldToMetersScale() const
+{
+	return WorldToMeters;
+}
+
+bool FPixelStreamingHMD::IsHMDEnabled() const
+{
+	return UE::PixelStreamingHMD::Settings::CVarPixelStreamingEnableHMD.GetValueOnAnyThread();
+}
+
+void FPixelStreamingHMD::EnableHMD(bool enable)
+{
+	UE::PixelStreamingHMD::Settings::CVarPixelStreamingEnableHMD->Set(enable, ECVF_SetByCode);
+}
+
+bool FPixelStreamingHMD::GetHMDMonitorInfo(MonitorInfo& MonitorDesc)
+{
+	MonitorDesc.MonitorName = "";
+	MonitorDesc.MonitorId = 0;
+	MonitorDesc.DesktopX = MonitorDesc.DesktopY = MonitorDesc.ResolutionX = MonitorDesc.ResolutionY = 0;
+	return false;
+}
+
+void FPixelStreamingHMD::GetFieldOfView(float& OutHFOVInDegrees, float& OutVFOVInDegrees) const
+{
+	OutHFOVInDegrees = 0.0f;
+	OutVFOVInDegrees = 0.0f;
+}
+
+bool FPixelStreamingHMD::EnumerateTrackedDevices(TArray<int32>& OutDevices, EXRTrackedDeviceType Type)
+{
+	if (Type == EXRTrackedDeviceType::Any || Type == EXRTrackedDeviceType::HeadMountedDisplay)
+	{
+		OutDevices.Add(IXRTrackingSystem::HMDDeviceId);
+		return true;
+	}
+	return false;
+}
+
+void FPixelStreamingHMD::SetInterpupillaryDistance(float NewInterpupillaryDistance)
+{
+	InterpupillaryDistance = NewInterpupillaryDistance;
+}
+
+float FPixelStreamingHMD::GetInterpupillaryDistance() const
+{
+	return InterpupillaryDistance;
+}
+
+bool FPixelStreamingHMD::GetCurrentPose(int32 DeviceId, FQuat& CurrentOrientation, FVector& CurrentPosition)
+{
+	if (DeviceId != IXRTrackingSystem::HMDDeviceId)
+	{
+		return false;
+	}
+	CurrentOrientation = CurHmdTransform.GetRotation();
+	CurrentPosition = CurHmdTransform.GetTranslation();
+	return true;
+}
+
+bool FPixelStreamingHMD::IsChromaAbCorrectionEnabled() const
+{
+	return false;
+}
+
+void FPixelStreamingHMD::ResetOrientationAndPosition(float yaw)
+{
+	ResetOrientation(yaw);
+	ResetPosition();
+}
+
+void FPixelStreamingHMD::DrawDistortionMesh_RenderThread(struct FHeadMountedDisplayPassContext& Context, const FIntPoint& TextureSize)
+{
+	float ClipSpaceQuadZ = 0.0f;
+	FMatrix QuadTexTransform = FMatrix::Identity;
+	FMatrix QuadPosTransform = FMatrix::Identity;
+	const FViewInfo& View = Context.View;
+	const FIntRect SrcRect = View.UnscaledViewRect;
+
+	FRHICommandListImmediate& RHICmdList = Context.RHICmdList;
+	const FSceneViewFamily& ViewFamily = *(View.Family);
+	FIntPoint ViewportSize = ViewFamily.RenderTarget->GetSizeXY();
+	RHICmdList.SetViewport(0, 0, 0.0f, ViewportSize.X, ViewportSize.Y, 1.0f);
+
+	static const uint32 NumVerts = 4;
+	static const uint32 NumTris = 2;
+
+	static const FDistortionVertex MeshVerts[][NumVerts] = {
+		{
+			// left eye
+			{ FVector2f(-0.9f, -0.9f), FVector2f(0.0f, 1.0f), FVector2f(0.0f, 1.0f), FVector2f(0.0f, 1.0f), 1.0f, 0.0f },
+			{ FVector2f(-0.1f, -0.9f), FVector2f(0.5f, 1.0f), FVector2f(0.5f, 1.0f), FVector2f(0.5f, 1.0f), 1.0f, 0.0f },
+			{ FVector2f(-0.1f, 0.9f), FVector2f(0.5f, 0.0f), FVector2f(0.5f, 0.0f), FVector2f(0.5f, 0.0f), 1.0f, 0.0f },
+			{ FVector2f(-0.9f, 0.9f), FVector2f(0.0f, 0.0f), FVector2f(0.0f, 0.0f), FVector2f(0.0f, 0.0f), 1.0f, 0.0f },
+		},
+		{
+			// right eye
+			{ FVector2f(0.1f, -0.9f), FVector2f(0.5f, 1.0f), FVector2f(0.5f, 1.0f), FVector2f(0.5f, 1.0f), 1.0f, 0.0f },
+			{ FVector2f(0.9f, -0.9f), FVector2f(1.0f, 1.0f), FVector2f(1.0f, 1.0f), FVector2f(1.0f, 1.0f), 1.0f, 0.0f },
+			{ FVector2f(0.9f, 0.9f), FVector2f(1.0f, 0.0f), FVector2f(1.0f, 0.0f), FVector2f(1.0f, 0.0f), 1.0f, 0.0f },
+			{ FVector2f(0.1f, 0.9f), FVector2f(0.5f, 0.0f), FVector2f(0.5f, 0.0f), FVector2f(0.5f, 0.0f), 1.0f, 0.0f },
+		}
+	};
+
+	FRHIResourceCreateInfo CreateInfo(TEXT("FPixelStreamingHMD"));
+	FBufferRHIRef VertexBufferRHI = RHICreateVertexBuffer(sizeof(FDistortionVertex) * NumVerts, BUF_Volatile, CreateInfo);
+	void* VoidPtr = RHILockBuffer(VertexBufferRHI, 0, sizeof(FDistortionVertex) * NumVerts, RLM_WriteOnly);
+	FPlatformMemory::Memcpy(VoidPtr, MeshVerts[View.StereoViewIndex], sizeof(FDistortionVertex) * NumVerts);
+	RHIUnlockBuffer(VertexBufferRHI);
+
+	static const uint16 Indices[] = { 0, 1, 2, 0, 2, 3 };
+
+	FBufferRHIRef IndexBufferRHI = RHICreateIndexBuffer(sizeof(uint16), sizeof(uint16) * 6, BUF_Volatile, CreateInfo);
+	void* VoidPtr2 = RHILockBuffer(IndexBufferRHI, 0, sizeof(uint16) * 6, RLM_WriteOnly);
+	FPlatformMemory::Memcpy(VoidPtr2, Indices, sizeof(uint16) * 6);
+	RHIUnlockBuffer(IndexBufferRHI);
+
+	RHICmdList.SetStreamSource(0, VertexBufferRHI, 0);
+	RHICmdList.DrawIndexedPrimitive(IndexBufferRHI, 0, 0, NumVerts, 0, NumTris, 1);
+
+	IndexBufferRHI.SafeRelease();
+	VertexBufferRHI.SafeRelease();
+}
+
+bool FPixelStreamingHMD::IsStereoEnabled() const
+{
+	return bStereoEnabled;
+}
+
+bool FPixelStreamingHMD::EnableStereo(bool stereo)
+{
+	bStereoEnabled = stereo;
+	return bStereoEnabled;
+}
+
+void FPixelStreamingHMD::AdjustViewRect(int32 ViewIndex, int32& X, int32& Y, uint32& SizeX, uint32& SizeY) const
+{
+	SizeX = SizeX / 2;
+	X += SizeX * ViewIndex;
+}
+
+void FPixelStreamingHMD::CalculateStereoViewOffset(const int32 ViewIndex, FRotator& ViewRotation, const float InWorldToMeters, FVector& ViewLocation)
+{
+	if (ViewIndex != INDEX_NONE)
+	{
+		float EyeOffset = 3.20000005f;
+		const float PassOffset = (ViewIndex == EStereoscopicEye::eSSE_LEFT_EYE) ? EyeOffset : -EyeOffset;
+		ViewLocation += ViewRotation.Quaternion().RotateVector(FVector(0, PassOffset, 0));
+	}
+}
+
+FMatrix FPixelStreamingHMD::GetStereoProjectionMatrix(const int32 ViewIndex) const
+{
+	const float ProjectionCenterOffset = 0.151976421f;
+	const float PassProjectionOffset = (ViewIndex == EStereoscopicEye::eSSE_LEFT_EYE) ? ProjectionCenterOffset : -ProjectionCenterOffset;
+
+	const float HalfFov = 2.19686294f / 2.f;
+	const float InWidth = 640.f;
+	const float InHeight = 480.f;
+	const float XS = 1.0f / tan(HalfFov);
+	const float YS = InWidth / tan(HalfFov) / InHeight;
+
+	const float InNearZ = GNearClippingPlane;
+	return FMatrix(
+			   FPlane(XS, 0.0f, 0.0f, 0.0f),
+			   FPlane(0.0f, YS, 0.0f, 0.0f),
+			   FPlane(0.0f, 0.0f, 0.0f, 1.0f),
+			   FPlane(0.0f, 0.0f, InNearZ, 0.0f))
+
+		* FTranslationMatrix(FVector(PassProjectionOffset, 0, 0));
+}
+
+void FPixelStreamingHMD::GetEyeRenderParams_RenderThread(const FHeadMountedDisplayPassContext& Context, FVector2D& EyeToSrcUVScaleValue, FVector2D& EyeToSrcUVOffsetValue) const
+{
+	EyeToSrcUVOffsetValue = FVector2D::ZeroVector;
+	EyeToSrcUVScaleValue = FVector2D(1.0f, 1.0f);
+}
+
+void FPixelStreamingHMD::SetupViewFamily(FSceneViewFamily& InViewFamily)
+{
+	InViewFamily.EngineShowFlags.MotionBlur = 0;
+	InViewFamily.EngineShowFlags.HMDDistortion = true;
+	InViewFamily.EngineShowFlags.StereoRendering = IsStereoEnabled();
+
+	if (UWorld* World = GWorld)
+	{
+		WorldToMeters = World->GetWorldSettings()->WorldToMeters;
+	}
+}
+
+void FPixelStreamingHMD::PreRenderView_RenderThread(FRDGBuilder& GraphBuilder, FSceneView& InView)
+{
+	check(IsInRenderingThread());
+}
+
+void FPixelStreamingHMD::PreRenderViewFamily_RenderThread(FRDGBuilder& GraphBuilder, FSceneViewFamily& ViewFamily)
+{
+	check(IsInRenderingThread());
+}
\ No newline at end of file
diff --git a/Engine/Plugins/Media/PixelStreaming/Source/PixelStreamingHMD/Public/PixelStreamingHMDEnums.h b/Engine/Plugins/Media/PixelStreaming/Source/PixelStreamingHMD/Public/PixelStreamingHMDEnums.h
new file mode 100644
--- /dev/null
+++ b/Engine/Plugins/Media/PixelStreaming/Source/PixelStreamingHMD/Public/PixelStreamingHMDEnums.h
@@ -0,0 +1,10 @@
+// Copyright Epic Games, Inc. All Rights Reserved.
+
+#pragma once
+
+enum class PIXELSTREAMINGHMD_API EPixelStreamingXRSystem : uint8
+{
+	Unknown,
+	HTCVive,
+	Quest,
+};
\ No newline at end of file
diff --git a/Engine/Plugins/Media/PixelStreaming/Source/PixelStreamingInput/Private/PixelStreamingInputMessage.cpp b/Engine/Plugins/Media/PixelStreaming/Source/PixelStreamingInput/Private/PixelStreamingInputMessage.cpp
new file mode 100644
--- /dev/null
+++ b/Engine/Plugins/Media/PixelStreaming/Source/PixelStreamingInput/Private/PixelStreamingInputMessage.cpp
@@ -0,0 +1,78 @@
+// Copyright Epic Games, Inc. All Rights Reserved.
+
+#include "PixelStreamingInputMessage.h"
+#include "PixelStreamingInputEnums.h"
+
+uint8 FPixelStreamingInputMessage::CurrentId = 0;
+
+FPixelStreamingInputMessage::FPixelStreamingInputMessage()
+	: Id(++CurrentId)
+	, Structure({})
+{
+	ByteLength = 0;
+}
+
+FPixelStreamingInputMessage::FPixelStreamingInputMessage(uint8 InId)
+	: Id(InId)
+	, Structure({})
+{
+	// We set id to 255 for the Protocol message type.
+	// We don't want to set the CurrentId to this value otherwise we will have no room for custom messages
+	if (InId > CurrentId && InId != 255)
+	{
+		CurrentId = InId; // Ensure that any new message id starts after the highest specified Id
+	}
+
+	ByteLength = 0;
+}
+
+FPixelStreamingInputMessage::FPixelStreamingInputMessage(TArray<EPixelStreamingMessageTypes> InStructure)
+	: Id(++CurrentId)
+	, Structure(InStructure)
+{
+	ByteLength = CalculateByteLength();
+}
+
+FPixelStreamingInputMessage::FPixelStreamingInputMessage(uint8 InId, TArray<EPixelStreamingMessageTypes> InStructure)
+	: Id(InId)
+	, Structure(InStructure)
+{
+	// We set id to 255 for the Protocol message type.
+	// We don't want to set the CurrentId to this value otherwise we will have no room for custom messages
+	if (InId > CurrentId && InId != 255)
+	{
+		CurrentId = InId; // Ensure that any new message id starts after the highest specified Id
+	}
+
+	ByteLength = CalculateByteLength();
+}
+
+uint8 FPixelStreamingInputMessage::CalculateByteLength()
+{
+	uint8 Length = 0;
+	for (EPixelStreamingMessageTypes Type : Structure)
+	{
+		switch (Type)
+		{
+			case EPixelStreamingMessageTypes::Uint8:
+				Length += 1;
+				break;
+			case EPixelStreamingMessageTypes::Uint16:
+			case EPixelStreamingMessageTypes::Int16:
+				Length += 2;
+				break;
+			case EPixelStreamingMessageTypes::Float:
+				Length += 4;
+				break;
+			case EPixelStreamingMessageTypes::Double:
+				Length += 8;
+				break;
+			case EPixelStreamingMessageTypes::Variable:
+			case EPixelStreamingMessageTypes::Undefined:
+			default:
+				// These types of have no defined length
+				break;
+		}
+	}
+	return Length;
+}
\ No newline at end of file
diff --git a/Engine/Plugins/Media/PixelStreaming/Source/PixelStreamingInput/Private/PixelStreamingInputModule.cpp b/Engine/Plugins/Media/PixelStreaming/Source/PixelStreamingInput/Private/PixelStreamingInputModule.cpp
new file mode 100644
--- /dev/null
+++ b/Engine/Plugins/Media/PixelStreaming/Source/PixelStreamingInput/Private/PixelStreamingInputModule.cpp
@@ -0,0 +1,53 @@
+// Copyright Epic Games, Inc. All Rights Reserved.
+
+#include "PixelStreamingInputModule.h"
+#include "IPixelStreamingHMDModule.h"
+#include "Framework/Application/SlateApplication.h"
+#include "ApplicationWrapper.h"
+#include "Settings.h"
+#include "PixelStreamingInputHandler.h"
+
+namespace UE::PixelStreamingInput
+{
+	void FPixelStreamingInputModule::StartupModule()
+	{
+		TSharedPtr<FPixelStreamingApplicationWrapper> PixelStreamerApplicationWrapper = MakeShareable(new FPixelStreamingApplicationWrapper(FSlateApplication::Get().GetPlatformApplication()));
+		TSharedPtr<FGenericApplicationMessageHandler> BaseHandler = FSlateApplication::Get().GetPlatformApplication()->GetMessageHandler();
+		InputHandler = MakeShared<FPixelStreamingInputHandler>(PixelStreamerApplicationWrapper, BaseHandler);
+
+		Settings::InitialiseSettings();
+
+		IModularFeatures::Get().RegisterModularFeature(GetModularFeatureName(), this);
+	}
+
+	void FPixelStreamingInputModule::ShutdownModule()
+	{
+		IModularFeatures::Get().UnregisterModularFeature(GetModularFeatureName(), this);
+	}
+
+	void FPixelStreamingInputModule::RegisterMessage(EPixelStreamingMessageDirection MessageDirection, const FString& MessageType, FPixelStreamingInputMessage Message, const TFunction<void(FMemoryReader)>& Handler)
+	{
+		if (MessageDirection == EPixelStreamingMessageDirection::ToStreamer)
+		{
+			FPixelStreamingInputProtocol::ToStreamerProtocol.Add(MessageType, Message);
+			InputHandler->RegisterMessageHandler(MessageType, Handler);
+		}
+		else if (MessageDirection == EPixelStreamingMessageDirection::FromStreamer)
+		{
+			FPixelStreamingInputProtocol::FromStreamerProtocol.Add(MessageType, Message);
+		}
+		OnProtocolUpdated.Broadcast();
+	}
+
+	TFunction<void(FMemoryReader)> FPixelStreamingInputModule::FindMessageHandler(const FString& MessageType)
+	{
+		return InputHandler->FindMessageHandler(MessageType);
+	}
+
+	TSharedPtr<IInputDevice> FPixelStreamingInputModule::CreateInputDevice(const TSharedRef<FGenericApplicationMessageHandler>& InMessageHandler)
+	{
+		return InputHandler;
+	}
+} // namespace UE::PixelStreamingInput
+
+IMPLEMENT_MODULE(UE::PixelStreamingInput::FPixelStreamingInputModule, PixelStreamingInput)
\ No newline at end of file
diff --git a/Engine/Plugins/Media/PixelStreaming/Source/PixelStreamingInput/Private/PixelStreamingInputModule.h b/Engine/Plugins/Media/PixelStreaming/Source/PixelStreamingInput/Private/PixelStreamingInputModule.h
new file mode 100644
--- /dev/null
+++ b/Engine/Plugins/Media/PixelStreaming/Source/PixelStreamingInput/Private/PixelStreamingInputModule.h
@@ -0,0 +1,30 @@
+// Copyright Epic Games, Inc. All Rights Reserved.
+
+#pragma once
+
+#include "IPixelStreamingInputModule.h"
+
+namespace UE::PixelStreamingInput
+{
+	class FPixelStreamingInputModule : public IPixelStreamingInputModule
+	{
+	public:
+		virtual void RegisterMessage(EPixelStreamingMessageDirection MessageDirection, const FString& MessageType, FPixelStreamingInputMessage Message, const TFunction<void(FMemoryReader)>& Handler) override;
+		virtual TFunction<void(FMemoryReader)> FindMessageHandler(const FString& MessageType) override;
+		virtual TSharedPtr<IPixelStreamingInputHandler> GetInputHandler() override { return InputHandler; }
+
+	private:
+		/** IModuleInterface implementation */
+		void StartupModule() override;
+		void ShutdownModule() override;
+		/** End IModuleInterface implementation */
+
+		/** IInputDeviceModule implementation */
+		virtual TSharedPtr<IInputDevice> CreateInputDevice(const TSharedRef<FGenericApplicationMessageHandler>& InMessageHandler) override;
+		/** End IInputDeviceModule implementation */
+
+		// NOTE: There is only ever a single input handler which all of the streamers share. This provides
+		// a central point for the multiple streamers -> single UE instance
+		TSharedPtr<IPixelStreamingInputHandler> InputHandler;
+	};
+} // namespace UE::PixelStreamingInput
\ No newline at end of file
diff --git a/Engine/Plugins/Media/PixelStreaming/Source/PixelStreamingInput/Private/PixelStreamingInputProtocol.cpp b/Engine/Plugins/Media/PixelStreaming/Source/PixelStreamingInput/Private/PixelStreamingInputProtocol.cpp
new file mode 100644
--- /dev/null
+++ b/Engine/Plugins/Media/PixelStreaming/Source/PixelStreamingInput/Private/PixelStreamingInputProtocol.cpp
@@ -0,0 +1,65 @@
+// Copyright Epic Games, Inc. All Rights Reserved.
+
+#include "PixelStreamingInputProtocol.h"
+#include "PixelStreamingInputMessage.h"
+#include "Dom/JsonValue.h"
+
+TMap<FString, FPixelStreamingInputMessage> FPixelStreamingInputProtocol::ToStreamerProtocol;
+TMap<FString, FPixelStreamingInputMessage> FPixelStreamingInputProtocol::FromStreamerProtocol;
+
+TSharedPtr<FJsonObject> FPixelStreamingInputProtocol::ToJson(EPixelStreamingMessageDirection Direction)
+{
+	TSharedPtr<FJsonObject> ProtocolJson = MakeShareable(new FJsonObject());
+	TMap<FString, FPixelStreamingInputMessage> MessageProtocol =
+		(Direction == EPixelStreamingMessageDirection::ToStreamer)
+		? FPixelStreamingInputProtocol::ToStreamerProtocol
+		: FPixelStreamingInputProtocol::FromStreamerProtocol;
+
+	ProtocolJson->SetField("Direction", MakeShared<FJsonValueNumber>(static_cast<uint8>(Direction)));
+	for (TMap<FString, FPixelStreamingInputMessage>::TIterator Iter = MessageProtocol.CreateIterator(); Iter; ++Iter)
+	{
+		TSharedPtr<FJsonObject> MessageJson = MakeShareable(new FJsonObject());
+		FString MessageType = Iter.Key();
+		FPixelStreamingInputMessage Message = Iter.Value();
+
+		MessageJson->SetField("id", MakeShared<FJsonValueNumber>(Message.GetID()));
+		MessageJson->SetField("byteLength", MakeShared<FJsonValueNumber>(Message.GetByteLength()));
+
+		if (Message.GetByteLength() > 0)
+		{
+			TArray<TSharedPtr<FJsonValue>> StructureJson;
+			TArray<EPixelStreamingMessageTypes> Structure = Message.GetStructure();
+			for (auto It = Structure.CreateIterator(); It; ++It)
+			{
+				FString Text;
+				switch (*It)
+				{
+					case EPixelStreamingMessageTypes::Uint8:
+						Text = "uint8";
+						break;
+					case EPixelStreamingMessageTypes::Uint16:
+						Text = "uint16";
+						break;
+					case EPixelStreamingMessageTypes::Int16:
+						Text = "int16";
+						break;
+					case EPixelStreamingMessageTypes::Float:
+						Text = "float";
+						break;
+					case EPixelStreamingMessageTypes::Double:
+						Text = "double";
+						break;
+					default:
+						Text = "";
+				}
+				TSharedRef<FJsonValueString> JsonValue = MakeShareable(new FJsonValueString(FString::Printf(TEXT("%s"), *Text)));
+				StructureJson.Add(JsonValue);
+			}
+			MessageJson->SetArrayField("structure", StructureJson);
+		}
+
+		ProtocolJson->SetField(*MessageType, MakeShared<FJsonValueObject>(MessageJson));
+	}
+
+	return ProtocolJson;
+}
\ No newline at end of file
diff --git a/Engine/Plugins/Media/PixelStreaming/Source/PixelStreamingInput/Public/IPixelStreamingInputModule.h b/Engine/Plugins/Media/PixelStreaming/Source/PixelStreamingInput/Public/IPixelStreamingInputModule.h
new file mode 100644
--- /dev/null
+++ b/Engine/Plugins/Media/PixelStreaming/Source/PixelStreamingInput/Public/IPixelStreamingInputModule.h
@@ -0,0 +1,70 @@
+// Copyright Epic Games, Inc. All Rights Reserved.
+
+#pragma once
+
+#include "Modules/ModuleManager.h"
+#include "IInputDeviceModule.h"
+#include "PixelStreamingInputProtocol.h"
+#include "IPixelStreamingInputHandler.h"
+
+/**
+ * The public interface of the Pixel Streaming Input module.
+ */
+class PIXELSTREAMINGINPUT_API IPixelStreamingInputModule : public IInputDeviceModule
+{
+public:
+	/**
+	 * Singleton-like access to this module's interface.
+	 * Beware of calling this during the shutdown phase, though.  Your module might have been unloaded already.
+	 *
+	 * @return Returns singleton instance, loading the module on demand if needed
+	 */
+	static inline IPixelStreamingInputModule& Get()
+	{
+		return FModuleManager::LoadModuleChecked<IPixelStreamingInputModule>("PixelStreamingInput");
+	}
+
+	/**
+	 * Checks to see if this module is loaded.
+	 *
+	 * @return True if the module is loaded.
+	 */
+	static inline bool IsAvailable() { return FModuleManager::Get().IsModuleLoaded("PixelStreamingInput"); }
+
+	/**
+	 * Register a new message type that peers can send and pixel streaming can receive
+	 * @param MessageDirection The direction the message will travel. eg Streamer->Player or Player->Streamer
+	 * @param MessageType The human readable identifier (eg "TouchStarted")
+	 * @param Message The object used to define the structure of the message
+	 * @param Handler The handler for this message type. This function will be executed whenever the corresponding message type is received
+	 */
+	virtual void RegisterMessage(EPixelStreamingMessageDirection MessageDirection, const FString& MessageType, FPixelStreamingInputMessage Message, const TFunction<void(FMemoryReader)>& Handler) = 0;
+
+	/**
+	 * @brief Find the function to be called whenever the specified message type is received.
+	 *
+	 * @param MessageType The human readable identifier for the message
+	 * @return TFunction<void(FMemoryReader)> The function called when this message type is received.
+	 */
+	virtual TFunction<void(FMemoryReader)> FindMessageHandler(const FString& MessageType) = 0;
+
+	/**
+	 * @brief Get the Input Handler for this module
+	 *
+	 * @return TSharedPtr<IPixelStreamingInputHandler> the input handler
+	 */
+	virtual TSharedPtr<IPixelStreamingInputHandler> GetInputHandler() = 0;
+
+	/**
+	 * Attempts to create a new input device interface
+	 *
+	 * @return	Interface to the new input device, if we were able to successfully create one
+	 */
+	virtual TSharedPtr<class IInputDevice> CreateInputDevice(const TSharedRef<FGenericApplicationMessageHandler>& InMessageHandler) override = 0;
+
+	DECLARE_MULTICAST_DELEGATE(FOnProtocolUpdated);
+	FOnProtocolUpdated OnProtocolUpdated;
+
+	DECLARE_MULTICAST_DELEGATE_OneParam(FOnSendMessage, FMemoryReader);
+	FOnSendMessage OnSendMessage;
+};
diff --git a/Engine/Plugins/Media/PixelStreaming/Source/PixelStreamingInput/Public/PixelStreamingInputEnums.h b/Engine/Plugins/Media/PixelStreaming/Source/PixelStreamingInput/Public/PixelStreamingInputEnums.h
new file mode 100644
--- /dev/null
+++ b/Engine/Plugins/Media/PixelStreaming/Source/PixelStreamingInput/Public/PixelStreamingInputEnums.h
@@ -0,0 +1,54 @@
+// Copyright Epic Games, Inc. All Rights Reserved.
+
+#pragma once
+
+#include "CoreTypes.h"
+
+/**
+ * @brief The data types supported when sending messages across the data channel to/from peers
+ *
+ */
+enum class PIXELSTREAMINGINPUT_API EPixelStreamingMessageTypes
+{
+	Uint8,
+	Uint16,
+	Int16,
+	Float,
+	Double,
+	Variable,
+	Undefined
+};
+
+/**
+ * @brief The message directions
+ *
+ */
+enum class PIXELSTREAMINGINPUT_API EPixelStreamingMessageDirection : uint8
+{
+	ToStreamer = 0,
+	FromStreamer = 1
+};
+
+/**
+ * @brief The possible actions pixel streaming supports. These actions help differentiate input received from the browser
+ *
+ */
+enum class PIXELSTREAMINGINPUT_API EPixelStreamingInputAction : uint8
+{
+	X = 0,
+	Y = 1,
+	Axis,
+	Click,
+	Touch,
+	None,
+};
+
+/**
+ * @brief The possible control schemes pixel streaming supports. RouteToWindow routes input at an application level. RouteToWidget routes input to a specific widget, ignoring the rest of the application.
+ *
+ */
+enum class PIXELSTREAMINGINPUT_API EPixelStreamingInputType : uint8
+{
+	RouteToWindow,
+	RouteToWidget
+};
\ No newline at end of file
diff --git a/Engine/Plugins/Media/PixelStreaming/Source/PixelStreamingInput/Public/PixelStreamingInputMessage.h b/Engine/Plugins/Media/PixelStreaming/Source/PixelStreamingInput/Public/PixelStreamingInputMessage.h
new file mode 100644
--- /dev/null
+++ b/Engine/Plugins/Media/PixelStreaming/Source/PixelStreamingInput/Public/PixelStreamingInputMessage.h
@@ -0,0 +1,37 @@
+// Copyright Epic Games, Inc. All Rights Reserved.
+
+#pragma once
+
+#include "PixelStreamingInputEnums.h"
+
+#include "Containers/Array.h"
+#include "HAL/Platform.h"
+
+class PIXELSTREAMINGINPUT_API FPixelStreamingInputMessage
+{
+public:
+	// Emtpy constructor. Sets ID to next available ID, byteLength to 0, and an empty Structure
+	FPixelStreamingInputMessage();
+
+	// Constructor taking an ID. Sets byteLength to 0, and an empty Structure
+	FPixelStreamingInputMessage(uint8 InId);
+
+	// Constructor taking a Structure. Sets ID to next available ID and calculates bytelength from the structure
+	FPixelStreamingInputMessage(TArray<EPixelStreamingMessageTypes> InStructure);
+
+	// Constructor taking an ID and a Structure. Calculates bytelength from the structure
+	FPixelStreamingInputMessage(uint8 InId, TArray<EPixelStreamingMessageTypes> InStructure);
+
+	uint8 GetByteLength() const { return ByteLength; }
+	uint8 GetID() const { return Id; }
+	TArray<EPixelStreamingMessageTypes> GetStructure() const { return Structure; }
+
+private:
+	uint8 CalculateByteLength();
+
+	uint8 Id;
+	TArray<EPixelStreamingMessageTypes> Structure;
+	uint8 ByteLength;
+
+	static uint8 CurrentId;
+};
\ No newline at end of file
diff --git a/Engine/Plugins/Media/PixelStreaming/Source/PixelStreamingInput/Public/PixelStreamingInputProtocol.h b/Engine/Plugins/Media/PixelStreaming/Source/PixelStreamingInput/Public/PixelStreamingInputProtocol.h
new file mode 100644
--- /dev/null
+++ b/Engine/Plugins/Media/PixelStreaming/Source/PixelStreamingInput/Public/PixelStreamingInputProtocol.h
@@ -0,0 +1,17 @@
+// Copyright Epic Games, Inc. All Rights Reserved.
+
+#pragma once
+
+#include "HAL/Platform.h"
+#include "Containers/Map.h"
+#include "PixelStreamingInputMessage.h"
+#include "Dom/JsonObject.h"
+
+struct PIXELSTREAMINGINPUT_API FPixelStreamingInputProtocol
+{
+public:
+	static TMap<FString, FPixelStreamingInputMessage> ToStreamerProtocol;
+	static TMap<FString, FPixelStreamingInputMessage> FromStreamerProtocol;
+
+	static TSharedPtr<FJsonObject> ToJson(EPixelStreamingMessageDirection Direction);
+};
\ No newline at end of file
diff --git a/Engine/Plugins/Media/WebMMedia/Source/WebMMedia/WebMMedia.Build.cs b/Engine/Plugins/Media/WebMMedia/Source/WebMMedia/WebMMedia.Build.cs
--- a/Engine/Plugins/Media/WebMMedia/Source/WebMMedia/WebMMedia.Build.cs
+++ b/Engine/Plugins/Media/WebMMedia/Source/WebMMedia/WebMMedia.Build.cs
@@ -52,7 +52,7 @@ namespace UnrealBuildTool.Rules
 			{
 				// libVPX is linked inside WebRTC so just use those headers and binaries where avaliable
 				// In order to get rid of this need to rebuild webrtc without libvpx included (win64, macos, linux where editor runs)
-				PrivateIncludePaths.Add(Path.Join(EngineDirectory, "Source", "ThirdParty", WebRTC.WebRtcVersionPath, "Include", "third_party", "libvpx", "source", "libvpx"));
+				PrivateIncludePaths.Add(Path.Join(EngineDirectory, "Source", "ThirdParty", "WebRTC", "4664", "Include", "third_party", "libvpx", "source", "libvpx"));
 
 				PublicDependencyModuleNames.AddRange(
 					new string[] {
