Add missing copyright boilerplate

diff --git a/Engine/Plugins/Media/PixelStreaming/Source/PixelStreaming/Private/PlayerContext.h b/Engine/Plugins/Media/PixelStreaming/Source/PixelStreaming/Private/PlayerContext.h
--- a/Engine/Plugins/Media/PixelStreaming/Source/PixelStreaming/Private/PlayerContext.h
+++ b/Engine/Plugins/Media/PixelStreaming/Source/PixelStreaming/Private/PlayerContext.h
@@ -1,4 +1,4 @@
-//
+// Copyright Epic Games, Inc. All Rights Reserved.
 
 #pragma once
 
diff --git a/Engine/Plugins/Media/PixelStreaming/Source/PixelStreamingInput/Private/Utils.h b/Engine/Plugins/Media/PixelStreaming/Source/PixelStreamingInput/Private/Utils.h
--- a/Engine/Plugins/Media/PixelStreaming/Source/PixelStreamingInput/Private/Utils.h
+++ b/Engine/Plugins/Media/PixelStreaming/Source/PixelStreamingInput/Private/Utils.h
@@ -1,3 +1,5 @@
+// Copyright Epic Games, Inc. All Rights Reserved.
+
 #pragma once
 
 #include "Dom/JsonObject.h"
