Dependency Hack to allow PixelStreaming Blueprints to work on Server targets

diff --git a/Engine/Plugins/Media/PixelStreaming/PixelStreaming.uplugin b/Engine/Plugins/Media/PixelStreaming/PixelStreaming.uplugin
--- a/Engine/Plugins/Media/PixelStreaming/PixelStreaming.uplugin
+++ b/Engine/Plugins/Media/PixelStreaming/PixelStreaming.uplugin
@@ -20,7 +20,7 @@
 			"Type": "Runtime",
 			"LoadingPhase": "PreDefault",
 			"PlatformAllowList": [ "Win64", "Linux" ],
-			"TargetDenyList": [ "Server" ]
+			"TargetDenyList": [ ]
 		},
 		{
 			"Name": "PixelStreamingServers",
diff --git a/Engine/Plugins/Media/PixelStreaming/Source/PixelStreaming/Private/PixelStreamingModule.cpp b/Engine/Plugins/Media/PixelStreaming/Source/PixelStreaming/Private/PixelStreamingModule.cpp
--- a/Engine/Plugins/Media/PixelStreaming/Source/PixelStreaming/Private/PixelStreamingModule.cpp
+++ b/Engine/Plugins/Media/PixelStreaming/Source/PixelStreaming/Private/PixelStreamingModule.cpp
@@ -73,6 +73,11 @@ namespace UE::PixelStreaming
 	 */
 	void FPixelStreamingModule::StartupModule()
 	{
+#if UE_SERVER
+		// Hack to no-op the rest of the module so Blueprints can still work
+		return;
+#endif
+
 		// Initialise all settings from command line args etc
 		Settings::InitialiseSettings();
 
