Fix: Re-enable parsing of resolutions from the browser when streaming the editor

diff --git a/Engine/Plugins/Media/PixelStreaming/Source/PixelStreaming/Private/PixelStreamingModule.cpp b/Engine/Plugins/Media/PixelStreaming/Source/PixelStreaming/Private/PixelStreamingModule.cpp
--- a/Engine/Plugins/Media/PixelStreaming/Source/PixelStreaming/Private/PixelStreamingModule.cpp
+++ b/Engine/Plugins/Media/PixelStreaming/Source/PixelStreaming/Private/PixelStreamingModule.cpp
@@ -507,14 +507,14 @@ namespace UE::PixelStreaming
 																														EType::Uint8 
 																														}));
 		
-		// clang-format on
-		FPixelStreamingInputProtocol::ToStreamerProtocol.Add("XRButtonPressed", FPixelStreamingInputMessage(112, // Handedness,   ButtonIdx,      IsRepeat
+		FPixelStreamingInputProtocol::ToStreamerProtocol.Add("XRButtonPressed", FPixelStreamingInputMessage(112,// Handedness,   ButtonIdx,      IsRepeat
+																												{ EType::Uint8, EType::Uint8, EType::Uint8 }));
+		FPixelStreamingInputProtocol::ToStreamerProtocol.Add("XRButtonTouched", FPixelStreamingInputMessage(113,// Handedness,   ButtonIdx,      IsRepeat
 																												{ EType::Uint8, EType::Uint8, EType::Uint8 }));
-		FPixelStreamingInputProtocol::ToStreamerProtocol.Add("XRButtonTouched", FPixelStreamingInputMessage(113, // Handedness,   ButtonIdx,      IsRepeat
+		FPixelStreamingInputProtocol::ToStreamerProtocol.Add("XRButtonReleased", FPixelStreamingInputMessage(114,// Handedness,   ButtonIdx,     IsRepeat
 																												{ EType::Uint8, EType::Uint8, EType::Uint8 }));
-		FPixelStreamingInputProtocol::ToStreamerProtocol.Add("XRButtonReleased", FPixelStreamingInputMessage(114, // Handedness,   ButtonIdx,     IsRepeat
-																					 							{ EType::Uint8, EType::Uint8, EType::Uint8 }));
-																												
+
+		// clang-format on
 		FPixelStreamingInputProtocol::ToStreamerProtocol.Add("XRAnalog", FPixelStreamingInputMessage(115, { EType::Uint8, EType::Uint8, EType::Double }));
 
 		FPixelStreamingInputProtocol::ToStreamerProtocol.Add("XRSystem", FPixelStreamingInputMessage(116, { EType::Uint8 }));
@@ -544,83 +544,39 @@ namespace UE::PixelStreaming
 			FPixelStreamingInputMessage(50),
 			[this](FMemoryReader Ar) { HandleUIInteraction(Ar); });
 
-		// The current handler is the function that will currently be executed if a message with type "Command" is received
-		TFunction<void(FMemoryReader)> BaseOnCommandHandler = InputModule.FindMessageHandler("Command");
-		// We then create our new handler which will execute the "base" handler
-		TFunction<void(FMemoryReader)> ExtendedOnCommandHandler = [this, BaseOnCommandHandler](FMemoryReader Ar) {
-			// and then perform out extended functionality after.
-			// equivalent to the super::DoSomeFunc pattern
-			BaseOnCommandHandler(Ar);
-
-			// In this case, the base handler handles raw console commands. Our extended functionality parses the commands
-			// for PS specific parameters
-			HandleOnCommand(Ar);
-		};
-
-		// Handle receiving commands from peers
-		InputModule.RegisterMessage(EPixelStreamingMessageDirection::ToStreamer, "Command", FPixelStreamingInputMessage(51), ExtendedOnCommandHandler);
-		// Handle sending commands to peers
-		InputModule.OnSendMessage.AddRaw(this, &UE::PixelStreaming::FPixelStreamingModule::HandleSendCommand);
-	}
-
-	void FPixelStreamingModule::HandleOnCommand(FMemoryReader Ar)
-	{
-		FString Res;
-		Res.GetCharArray().SetNumUninitialized(Ar.TotalSize() / 2 + 1);
-		Ar.Serialize(Res.GetCharArray().GetData(), Ar.TotalSize());
-		FString Descriptor = Res.Mid(1);
-		bool bSuccess = false;
-
-		/**
-		 * Encoder Settings
-		 */
-		FString MinQPString;
-		UE::PixelStreaming::ExtractJsonFromDescriptor(Descriptor, TEXT("Encoder.MinQP"), MinQPString, bSuccess);
-		if (bSuccess)
-		{
+		TSharedPtr<IPixelStreamingInputHandler> InputHandler = InputModule.GetInputHandler();
+		// Set Encoder.MinQP CVar
+		InputHandler->SetCommandHandler(TEXT("Encoder.MinQP"), [](FString Descriptor, FString MinQPString) {
 			int MinQP = FCString::Atoi(*MinQPString);
 			UE::PixelStreaming::Settings::CVarPixelStreamingEncoderMinQP->Set(MinQP, ECVF_SetByCommandline);
-			return;
-		}
+		});
 
-		FString MaxQPString;
-		UE::PixelStreaming::ExtractJsonFromDescriptor(Descriptor, TEXT("Encoder.MaxQP"), MaxQPString, bSuccess);
-		if (bSuccess)
-		{
+		// Set Encoder.MaxQP CVar
+		InputHandler->SetCommandHandler(TEXT("Encoder.MaxQP"), [](FString Descriptor, FString MaxQPString) {
 			int MaxQP = FCString::Atoi(*MaxQPString);
 			UE::PixelStreaming::Settings::CVarPixelStreamingEncoderMaxQP->Set(MaxQP, ECVF_SetByCommandline);
-			return;
-		}
+		});
 
-		/**
-		 * WebRTC Settings
-		 */
-		FString FPSString;
-		UE::PixelStreaming::ExtractJsonFromDescriptor(Descriptor, TEXT("WebRTC.Fps"), FPSString, bSuccess);
-		if (bSuccess)
-		{
+		// Set WebRTC max FPS
+		InputHandler->SetCommandHandler(TEXT("WebRTC.Fps"), [](FString Descriptor, FString FPSString) {
 			int FPS = FCString::Atoi(*FPSString);
 			UE::PixelStreaming::Settings::CVarPixelStreamingWebRTCFps->Set(FPS, ECVF_SetByCommandline);
-			return;
-		}
+		});
 
-		FString MinBitrateString;
-		UE::PixelStreaming::ExtractJsonFromDescriptor(Descriptor, TEXT("WebRTC.MinBitrate"), MinBitrateString, bSuccess);
-		if (bSuccess)
-		{
+		// Set MinBitrate
+		InputHandler->SetCommandHandler(TEXT("WebRTC.MinBitrate"), [](FString Descriptor, FString MinBitrateString) {
 			int MinBitrate = FCString::Atoi(*MinBitrateString);
 			UE::PixelStreaming::Settings::CVarPixelStreamingWebRTCMinBitrate->Set(MinBitrate, ECVF_SetByCommandline);
-			return;
-		}
+		});
 
-		FString MaxBitrateString;
-		UE::PixelStreaming::ExtractJsonFromDescriptor(Descriptor, TEXT("WebRTC.MaxBitrate"), MaxBitrateString, bSuccess);
-		if (bSuccess)
-		{
+		// Set MaxBitrate
+		InputHandler->SetCommandHandler(TEXT("WebRTC.MaxBitrate"), [](FString Descriptor, FString MaxBitrateString) {
 			int MaxBitrate = FCString::Atoi(*MaxBitrateString);
 			UE::PixelStreaming::Settings::CVarPixelStreamingWebRTCMaxBitrate->Set(MaxBitrate, ECVF_SetByCommandline);
-			return;
-		}
+		});
+
+		// Handle sending commands to peers
+		InputModule.OnSendMessage.AddRaw(this, &UE::PixelStreaming::FPixelStreamingModule::HandleSendCommand);
 	}
 
 	void FPixelStreamingModule::HandleSendCommand(FMemoryReader Ar)
diff --git a/Engine/Plugins/Media/PixelStreaming/Source/PixelStreaming/Private/PixelStreamingModule.h b/Engine/Plugins/Media/PixelStreaming/Source/PixelStreaming/Private/PixelStreamingModule.h
--- a/Engine/Plugins/Media/PixelStreaming/Source/PixelStreaming/Private/PixelStreamingModule.h
+++ b/Engine/Plugins/Media/PixelStreaming/Source/PixelStreaming/Private/PixelStreamingModule.h
@@ -67,7 +67,6 @@ namespace UE::PixelStreaming
 		bool IsPlatformCompatible() const;
 		void PopulateProtocol();
 		void RegisterCustomHandlers();
-		void HandleOnCommand(FMemoryReader Ar);
 		void HandleSendCommand(FMemoryReader Ar);
 		void HandleUIInteraction(FMemoryReader Ar);
 
diff --git a/Engine/Plugins/Media/PixelStreaming/Source/PixelStreamingEditor/Private/PixelStreamingEditorModule.cpp b/Engine/Plugins/Media/PixelStreaming/Source/PixelStreamingEditor/Private/PixelStreamingEditorModule.cpp
--- a/Engine/Plugins/Media/PixelStreaming/Source/PixelStreamingEditor/Private/PixelStreamingEditorModule.cpp
+++ b/Engine/Plugins/Media/PixelStreaming/Source/PixelStreamingEditor/Private/PixelStreamingEditorModule.cpp
@@ -20,6 +20,7 @@
 #include "SLevelViewport.h"
 #include "PixelStreamingInputEnums.h"
 #include "IPixelStreamingInputModule.h"
+#include "IPixelStreamingInputHandler.h"
 
 DECLARE_LOG_CATEGORY_EXTERN(LogPixelStreamingEditor, Log, All);
 DEFINE_LOG_CATEGORY(LogPixelStreamingEditor);
@@ -62,39 +63,29 @@ void FPixelStreamingEditorModule::InitEditorStreaming(IPixelStreamingModule& Mod
 		return;
 	}
 
-	IPixelStreamingInputModule& InputModule = IPixelStreamingInputModule::Get();
-
-	// TODO: This code diverged and needs attention from subject matter exports to appropriately reconcile
 	// Add custom handle for { type: "Command", Resolution.Width: "1920", Resolution.Height: "1080" } when doing Editor streaming
 	// because we cannot resize the game viewport, but instead want to resize the parent window.
-	// TSharedPtr<IPixelStreamingInputHandler> InputHandler = Streamer->GetInputHandler().Pin();
-	// if(InputHandler) 
-	// {
-	// 	InputHandler->SetCommandHandler("Resolution.Width", [](FString Descriptor, FString WidthString){
-	// 		bool bSuccess;
-	// 		FString HeightString;
-	// 		UE::PixelStreaming::ExtractJsonFromDescriptor(Descriptor, TEXT("Resolution.Height"), HeightString, bSuccess);
-	// 		int Width = FCString::Atoi(*WidthString);
-	// 		int Height = FCString::Atoi(*HeightString);
-	// 		if (Width < 1 || Height < 1)
-	// 		{
-	// 			return;
-	// 		}
-
-	// 		TSharedPtr<SWindow> ParentWindow = IMainFrameModule::Get().GetParentWindow();
-	// 		ParentWindow->Resize(FVector2D(Width, Height));
-	// 		FSlateApplication::Get().OnSizeChanged(ParentWindow->GetNativeWindow().ToSharedRef(), Width, Height);
-	// 	});
-	// }
-	
-	// InputModule.RegisterMessage(
-	// 	EPixelStreamingMessageDirection::ToStreamer,
-	// 	"Command",
-	// 	FPixelStreamingInputMessage(51),
-	// 	[ExtendedHandler](FMemoryReader Ar) {
-	// 		ExtendedHandler(Ar);
-	// 	});
-	// }
+	IPixelStreamingInputModule& InputModule = IPixelStreamingInputModule::Get();
+	TSharedPtr<IPixelStreamingInputHandler> InputHandler = InputModule.GetInputHandler();
+	if (InputHandler)
+	{
+		InputHandler->SetCommandHandler("Resolution.Width",
+			[](FString Descriptor, FString WidthString) {
+				bool bSuccess;
+				FString HeightString;
+				UE::PixelStreaming::ExtractJsonFromDescriptor(Descriptor, TEXT("Resolution.Height"), HeightString, bSuccess);
+				int Width = FCString::Atoi(*WidthString);
+				int Height = FCString::Atoi(*HeightString);
+				if (Width < 1 || Height < 1)
+				{
+					return;
+				}
+
+				TSharedPtr<SWindow> ParentWindow = IMainFrameModule::Get().GetParentWindow();
+				ParentWindow->Resize(FVector2D(Width, Height));
+				FSlateApplication::Get().OnSizeChanged(ParentWindow->GetNativeWindow().ToSharedRef(), Width, Height);
+			});
+	}
 
 	// Give the editor streamer the default url if the user hasn't specified one when launching the editor
 	if (Streamer->GetSignallingServerURL().IsEmpty())
diff --git a/Engine/Plugins/Media/PixelStreaming/Source/PixelStreamingInput/Private/PixelStreamingInputHandler.cpp b/Engine/Plugins/Media/PixelStreaming/Source/PixelStreamingInput/Private/PixelStreamingInputHandler.cpp
--- a/Engine/Plugins/Media/PixelStreaming/Source/PixelStreamingInput/Private/PixelStreamingInputHandler.cpp
+++ b/Engine/Plugins/Media/PixelStreaming/Source/PixelStreamingInput/Private/PixelStreamingInputHandler.cpp
@@ -166,6 +166,8 @@ namespace UE::PixelStreamingInput
 		FPixelStreamingInputConverter::GamepadInputToFKey.Add(MakeTuple(13, Action::Click), EKeys::Gamepad_DPad_Down);
 		FPixelStreamingInputConverter::GamepadInputToFKey.Add(MakeTuple(14, Action::Click), EKeys::Gamepad_DPad_Left);
 		FPixelStreamingInputConverter::GamepadInputToFKey.Add(MakeTuple(15, Action::Click), EKeys::Gamepad_DPad_Right);
+
+		PopulateDefaultCommandHandlers();
 	}
 
 	FPixelStreamingInputHandler::~FPixelStreamingInputHandler()
@@ -999,38 +1001,27 @@ namespace UE::PixelStreamingInput
 		IPixelStreamingHMDModule::Get().SetActiveXRSystem(static_cast<XRSystem>(ActiveSystem));
 	}
 
-	/**
-	 * Command receive handling
-	 */
-	void FPixelStreamingInputHandler::HandleOnCommand(FMemoryReader Ar)
+	void FPixelStreamingInputHandler::SetCommandHandler(const FString& CommandName, const TFunction<void(FString, FString)>& Handler)
 	{
-		FString Res;
-		Res.GetCharArray().SetNumUninitialized(Ar.TotalSize() / 2 + 1);
-		Ar.Serialize(Res.GetCharArray().GetData(), Ar.TotalSize());
-
-		FString Descriptor = Res.Mid(MessageHeaderOffset);
-		UE_LOG(LogPixelStreamingInputHandler, Verbose, TEXT("Command: %s"), *Descriptor);
-
-		bool bSuccess = false;
+		CommandHandlers.Add(CommandName, Handler);
+	}
 
-		FString ConsoleCommand;
-		ExtractJsonFromDescriptor(Descriptor, TEXT("ConsoleCommand"), ConsoleCommand, bSuccess);
-		if (bSuccess && Settings::CVarPixelStreamingInputAllowConsoleCommands.GetValueOnAnyThread())
-		{
-			GEngine->Exec(GEngine->GetWorld(), *ConsoleCommand);
-			return;
-		}
+	void FPixelStreamingInputHandler::PopulateDefaultCommandHandlers()
+	{
 
-		/**
-		 * Allowed Console commands
-		 */
-		FString WidthString;
-		FString HeightString;
-		ExtractJsonFromDescriptor(Descriptor, TEXT("Resolution.Width"), WidthString, bSuccess);
-		if (bSuccess)
-		{
-			ExtractJsonFromDescriptor(Descriptor, TEXT("Resolution.Height"), HeightString, bSuccess);
+		// Execute console commands if passed "ConsoleCommand" and -PixelStreamingAllowConsoleCommands is on.
+		CommandHandlers.Add(TEXT("ConsoleCommand"), [](FString Descriptor, FString ConsoleCommand) {
+			if (!UE::PixelStreamingInput::Settings::CVarPixelStreamingInputAllowConsoleCommands.GetValueOnAnyThread())
+			{
+				return;
+			}
+			GEngine->Exec(GEngine->GetWorld(), *ConsoleCommand); });
 
+		// Change width/height if sent { "Resolution.Width": 1920, "Resolution.Height": 1080 }
+		CommandHandlers.Add(TEXT("Resolution.Width"), [](FString Descriptor, FString WidthString) {
+			bool bSuccess = false;
+			FString HeightString;
+			UE::PixelStreamingInput::ExtractJsonFromDescriptor(Descriptor, TEXT("Resolution.Height"), HeightString, bSuccess);
 			if (bSuccess)
 			{
 				int Width = FCString::Atoi(*WidthString);
@@ -1042,20 +1033,39 @@ namespace UE::PixelStreamingInput
 
 				FString ChangeResCommand = FString::Printf(TEXT("r.SetRes %dx%d"), Width, Height);
 				GEngine->Exec(GEngine->GetWorld(), *ChangeResCommand);
-				return;
-			}
-		}
+			} });
 
-		FString StatFPSString;
-		UE::PixelStreamingInput::ExtractJsonFromDescriptor(Descriptor, TEXT("Stat.FPS"), StatFPSString, bSuccess);
-		if (bSuccess)
-		{
+		// Response to "Stat.FPS" by calling "stat fps"
+		CommandHandlers.Add(TEXT("Stat.FPS"), [](FString Descriptor, FString FPSCommand) {
 			FString StatFPSCommand = FString::Printf(TEXT("stat fps"));
-			GEngine->Exec(GEngine->GetWorld(), *StatFPSCommand);
-			return;
-		}
+			GEngine->Exec(GEngine->GetWorld(), *StatFPSCommand); });
+	}
 
-		// Remaining settings are handled in the Pixel Streaming module
+	/**
+	 * Command handling
+	 */
+	void FPixelStreamingInputHandler::HandleOnCommand(FMemoryReader Ar)
+	{
+		FString Res;
+		Res.GetCharArray().SetNumUninitialized(Ar.TotalSize() / 2 + 1);
+		Ar.Serialize(Res.GetCharArray().GetData(), Ar.TotalSize());
+
+		FString Descriptor = Res.Mid(MessageHeaderOffset);
+		UE_LOG(LogPixelStreamingInputHandler, Verbose, TEXT("Command: %s"), *Descriptor);
+
+		// Iterate each command handler and see if the command we got matches any of the bound command names.
+		for (auto& CommandHandlersPair : CommandHandlers)
+		{
+			FString CommandValue;
+			bool bSuccess = false;
+			UE::PixelStreamingInput::ExtractJsonFromDescriptor(Descriptor, CommandHandlersPair.Key, CommandValue, bSuccess);
+			if (bSuccess)
+			{
+				// Execute bound command handler with descriptor and parsed command value
+				CommandHandlersPair.Value(Descriptor, CommandValue);
+				return;
+			}
+		}
 	}
 
 	/**
diff --git a/Engine/Plugins/Media/PixelStreaming/Source/PixelStreamingInput/Private/PixelStreamingInputHandler.h b/Engine/Plugins/Media/PixelStreaming/Source/PixelStreamingInput/Private/PixelStreamingInputHandler.h
--- a/Engine/Plugins/Media/PixelStreaming/Source/PixelStreamingInput/Private/PixelStreamingInputHandler.h
+++ b/Engine/Plugins/Media/PixelStreaming/Source/PixelStreamingInput/Private/PixelStreamingInputHandler.h
@@ -30,6 +30,9 @@ namespace UE::PixelStreamingInput
 		/** Set which MessageHandler will route input  */
 		virtual void SetMessageHandler(const TSharedRef<FGenericApplicationMessageHandler>& InTargetHandler) override;
 
+		/** Register a custom function to execute when command JSON is received. */
+		virtual void SetCommandHandler(const FString& CommandName, const TFunction<void(FString, FString)>& Handler) override;
+
 		/** Exec handler to allow console commands to be passed through for debugging */
 		virtual bool Exec(UWorld* InWorld, const TCHAR* Cmd, FOutputDevice& Ar) override;
 		/**
@@ -104,6 +107,11 @@ namespace UE::PixelStreamingInput
 		 */
 		virtual void HandleUIInteraction(FMemoryReader Ar);
 
+		/**
+		 * Populate default command handlers for data channel messages sent with "{ type: "Command" }".
+		 */
+		void PopulateDefaultCommandHandlers();
+
 		FIntPoint ConvertFromNormalizedScreenLocation(const FVector2D& ScreenLocation, bool bIncludeOffset = true);
 		FWidgetPath FindRoutingMessageWidget(const FVector2D& Location) const;
 		FKey TranslateMouseButtonToKey(const EMouseButtons::Type Button);
@@ -184,6 +192,15 @@ namespace UE::PixelStreamingInput
 
 		TMap<EControllerHand, FPixelStreamingXRController> XRControllers;
 
+		/**
+		 * A map of named commands we respond to when we receive a datachannel message of type "command".
+		 * Key = command name (e.g "Encoder.MaxQP")
+		 * Value = The command handler lambda function whose parameters are as follows:
+		 * 	FString - the descriptor (e.g. the full json payload of the command message)
+		 * 	FString - the parsed value of the command, e.g. if key was "Encoder.MaxQP" and descriptor was { type: "Command", "Encoder.MaxQP": 51 }, then parsed value is "51".
+		 */
+		TMap<FString, TFunction<void(FString, FString)>> CommandHandlers;
+
 	private:
 		float uint16_MAX = (float)UINT16_MAX;
 		float int16_MAX = (float)SHRT_MAX;
diff --git a/Engine/Plugins/Media/PixelStreaming/Source/PixelStreamingInput/Public/IPixelStreamingInputHandler.h b/Engine/Plugins/Media/PixelStreaming/Source/PixelStreamingInput/Public/IPixelStreamingInputHandler.h
--- a/Engine/Plugins/Media/PixelStreaming/Source/PixelStreamingInput/Public/IPixelStreamingInputHandler.h
+++ b/Engine/Plugins/Media/PixelStreaming/Source/PixelStreamingInput/Public/IPixelStreamingInputHandler.h
@@ -76,6 +76,12 @@ public:
 	 */
 	virtual void RegisterMessageHandler(const FString& MessageType, const TFunction<void(FMemoryReader)>& Handler) = 0;
 
+	/**
+	 * @brief Register a custom function to execute when command JSON is received over the data channel: "{ type: "Command", YourCommand: YourCommandValue }".
+	 * Note: You can also override the default Pixel Streaming command handlers by setting handlers with the same name as those already used, e.g. "Stat.FPS".
+	 */
+	virtual void SetCommandHandler(const FString& CommandName, const TFunction<void(FString, FString)>& Handler) = 0;
+
 	/**
 	 * @brief Find the function to be called whenever the specified message type is received.
 	 *
