QOL: Enable text box contents to be sent as a whole string instead of needing to be sent character by character

diff --git a/Engine/Plugins/Media/PixelStreaming/Source/PixelStreaming/Private/PixelStreamingModule.cpp b/Engine/Plugins/Media/PixelStreaming/Source/PixelStreaming/Private/PixelStreamingModule.cpp
--- a/Engine/Plugins/Media/PixelStreaming/Source/PixelStreaming/Private/PixelStreamingModule.cpp
+++ b/Engine/Plugins/Media/PixelStreaming/Source/PixelStreaming/Private/PixelStreamingModule.cpp
@@ -463,6 +463,7 @@ namespace UE::PixelStreaming
 		// Generic Input Messages.
 		FPixelStreamingInputProtocol::ToStreamerProtocol.Add("UIInteraction", FPixelStreamingInputMessage(50));
 		FPixelStreamingInputProtocol::ToStreamerProtocol.Add("Command", FPixelStreamingInputMessage(51));
+		FPixelStreamingInputProtocol::ToStreamerProtocol.Add("TextboxEntry", FPixelStreamingInputMessage(52));
 
 		// Keyboard Input Message.
 		// Complex command with payload, therefore we specify the length of the payload (bytes) as well as the structure of the payload
diff --git a/Engine/Plugins/Media/PixelStreaming/Source/PixelStreamingInput/Private/PixelStreamingInputHandler.cpp b/Engine/Plugins/Media/PixelStreaming/Source/PixelStreamingInput/Private/PixelStreamingInputHandler.cpp
--- a/Engine/Plugins/Media/PixelStreaming/Source/PixelStreamingInput/Private/PixelStreamingInputHandler.cpp
+++ b/Engine/Plugins/Media/PixelStreaming/Source/PixelStreamingInput/Private/PixelStreamingInputHandler.cpp
@@ -22,6 +22,8 @@
 #include "IPixelStreamingHMDModule.h"
 #include "PixelStreamingInputEnums.h"
 #include "PixelStreamingInputConversion.h"
+#include "Widgets/Input/SEditableText.h"
+#include "Widgets/Input/SMultiLineEditableTextBox.h"
 
 DECLARE_LOG_CATEGORY_EXTERN(LogPixelStreamingInputHandler, Log, VeryVerbose);
 DEFINE_LOG_CATEGORY(LogPixelStreamingInputHandler);
@@ -76,6 +78,7 @@ namespace UE::PixelStreamingInput
 
 		RegisterMessageHandler("Command", [this](FMemoryReader Ar) { HandleOnCommand(Ar); });
 		RegisterMessageHandler("UIInteraction", [this](FMemoryReader Ar) { HandleUIInteraction(Ar); });
+		RegisterMessageHandler("TextboxEntry", [this](FMemoryReader Ar) { HandleOnTextboxEntry(Ar); });
 
 		// Populate map
 		// Button indices found in: [REDACTED_PROFILE_URL]/webxr-input-profiles/tree/master/packages/registry/profiles
@@ -1076,6 +1079,39 @@ namespace UE::PixelStreamingInput
 		// Actual implementation is in the Pixel Streaming module
 	}
 
+	/**
+	 * Textbox Entry handling
+	 */
+	void FPixelStreamingInputHandler::HandleOnTextboxEntry(FMemoryReader Ar)
+	{
+		FString Res;
+		Res.GetCharArray().SetNumUninitialized(Ar.TotalSize() / 2 + 1);
+		Ar.Serialize(Res.GetCharArray().GetData(), Ar.TotalSize());
+		FString Text = Res.Mid(1);
+
+		FSlateApplication::Get().ForEachUser([this, Text](FSlateUser& User) {
+			TSharedPtr<SWidget> FocusedWidget = User.GetFocusedWidget();
+
+			bool bIsEditableTextType = FocusedWidget->GetType() == TEXT("SEditableText");
+			bool bIsMultiLineEditableTextType = FocusedWidget->GetType() == TEXT("SMultiLineEditableText");
+			bool bEditable = FocusedWidget && (bIsEditableTextType || bIsMultiLineEditableTextType);
+
+			if (bEditable)
+			{
+				if (bIsEditableTextType)
+				{
+					SEditableText* TextBox = static_cast<SEditableText*>(FocusedWidget.Get());
+					TextBox->SetText(FText::FromString(Text));
+				}
+				else if (bIsMultiLineEditableTextType)
+				{
+					SMultiLineEditableTextBox* TextBox = static_cast<SMultiLineEditableTextBox*>(FocusedWidget.Get());
+					TextBox->SetText(FText::FromString(Text));
+				}
+			}
+		});
+	}
+
 	FIntPoint FPixelStreamingInputHandler::ConvertFromNormalizedScreenLocation(const FVector2D& ScreenLocation, bool bIncludeOffset)
 	{
 		FIntPoint OutVector((int32)ScreenLocation.X, (int32)ScreenLocation.Y);
@@ -1200,77 +1236,85 @@ namespace UE::PixelStreamingInput
 			static FName SMultiLineEditableTextType(TEXT("SMultiLineEditableText"));
 			bool bEditable = FocusedWidget && (FocusedWidget->GetType() == SEditableTextType || FocusedWidget->GetType() == SMultiLineEditableTextType);
 
-			// Check to see if the focus has changed.
-			FVector2D Pos = bEditable ? FocusedWidget->GetCachedGeometry().GetAbsolutePosition() : UnfocusedPos;
-			if (Pos != FocusedPos)
-			{
-				FocusedPos = Pos;
-
-				// Tell the browser that the focus has changed.
-				TSharedPtr<FJsonObject> JsonObject = MakeShareable(new FJsonObject);
-				JsonObject->SetStringField(TEXT("command"), TEXT("onScreenKeyboard"));
-				JsonObject->SetBoolField(TEXT("showOnScreenKeyboard"), bEditable);
+			// Tell the browser that the focus has changed.
+			TSharedPtr<FJsonObject> JsonObject = MakeShareable(new FJsonObject);
+			JsonObject->SetStringField(TEXT("command"), TEXT("onScreenKeyboard"));
+			JsonObject->SetBoolField(TEXT("showOnScreenKeyboard"), bEditable);
 
-				if (bEditable)
+			if (bEditable)
+			{
+				FVector2D NormalizedLocation;
+				TSharedPtr<SWindow> ApplicationWindow = TargetWindow.Pin();
+				if (ApplicationWindow.IsValid())
 				{
-					FVector2D NormalizedLocation;
-					TSharedPtr<SWindow> ApplicationWindow = TargetWindow.Pin();
-					if (ApplicationWindow.IsValid())
+					FVector2D WindowOrigin = ApplicationWindow->GetPositionInScreen();
+					FVector2D Pos = FocusedWidget->GetCachedGeometry().GetAbsolutePosition();
+					if (TargetViewport.IsValid())
 					{
-						FVector2D WindowOrigin = ApplicationWindow->GetPositionInScreen();
-						if (TargetViewport.IsValid())
-						{
-							TSharedPtr<SViewport> ViewportWidget = TargetViewport.Pin();
+						TSharedPtr<SViewport> ViewportWidget = TargetViewport.Pin();
 
-							if (ViewportWidget.IsValid())
-							{
-								FGeometry InnerWindowGeometry = ApplicationWindow->GetWindowGeometryInWindow();
+						if (ViewportWidget.IsValid())
+						{
+							FGeometry InnerWindowGeometry = ApplicationWindow->GetWindowGeometryInWindow();
 
-								// Find the widget path relative to the window
-								FArrangedChildren JustWindow(EVisibility::Visible);
-								JustWindow.AddWidget(FArrangedWidget(ApplicationWindow.ToSharedRef(), InnerWindowGeometry));
+							// Find the widget path relative to the window
+							FArrangedChildren JustWindow(EVisibility::Visible);
+							JustWindow.AddWidget(FArrangedWidget(ApplicationWindow.ToSharedRef(), InnerWindowGeometry));
 
-								FWidgetPath PathToWidget(ApplicationWindow.ToSharedRef(), JustWindow);
-								if (PathToWidget.ExtendPathTo(FWidgetMatcher(ViewportWidget.ToSharedRef()), EVisibility::Visible))
-								{
-									FArrangedWidget ArrangedWidget = PathToWidget.FindArrangedWidget(ViewportWidget.ToSharedRef()).Get(FArrangedWidget::GetNullWidget());
+							FWidgetPath PathToWidget(ApplicationWindow.ToSharedRef(), JustWindow);
+							if (PathToWidget.ExtendPathTo(FWidgetMatcher(ViewportWidget.ToSharedRef()), EVisibility::Visible))
+							{
+								FArrangedWidget ArrangedWidget = PathToWidget.FindArrangedWidget(ViewportWidget.ToSharedRef()).Get(FArrangedWidget::GetNullWidget());
 
-									FVector2D WindowClientOffset = ArrangedWidget.Geometry.GetAbsolutePosition();
-									FVector2D WindowClientSize = ArrangedWidget.Geometry.GetAbsoluteSize();
+								FVector2D WindowClientOffset = ArrangedWidget.Geometry.GetAbsolutePosition();
+								FVector2D WindowClientSize = ArrangedWidget.Geometry.GetAbsoluteSize();
 
-									Pos = Pos - WindowClientOffset;
-									NormalizedLocation = FVector2D(Pos / WindowClientSize);
-								}
+								Pos = Pos - WindowClientOffset;
+								NormalizedLocation = FVector2D(Pos / WindowClientSize);
 							}
 						}
-						else
-						{
-							FVector2D SizeInScreen = ApplicationWindow->GetSizeInScreen();
-							NormalizedLocation = FVector2D(Pos / SizeInScreen);
-						}
 					}
-					else if (TSharedPtr<FIntPoint> ScreenSize = TargetScreenSize.Pin())
+					else
 					{
-						FIntPoint SizeInScreen = *ScreenSize;
-						NormalizedLocation = FocusedPos / SizeInScreen;
+						FVector2D SizeInScreen = ApplicationWindow->GetSizeInScreen();
+						NormalizedLocation = FVector2D(Pos / SizeInScreen);
 					}
-
-					NormalizedLocation *= uint16_MAX;
-					// ConvertToNormalizedScreenLocation(Pos, NormalizedLocation);
-					JsonObject->SetNumberField(TEXT("x"), static_cast<uint16>(NormalizedLocation.X));
-					JsonObject->SetNumberField(TEXT("y"), static_cast<uint16>(NormalizedLocation.Y));
+				}
+				else if (TSharedPtr<FIntPoint> ScreenSize = TargetScreenSize.Pin())
+				{
+					FIntPoint SizeInScreen = *ScreenSize;
+					NormalizedLocation = FocusedPos / SizeInScreen;
 				}
 
-				FString Descriptor;
-				TSharedRef<TJsonWriter<TCHAR, TCondensedJsonPrintPolicy<TCHAR>>> JsonWriter = TJsonWriterFactory<TCHAR, TCondensedJsonPrintPolicy<TCHAR>>::Create(&Descriptor);
-				FJsonSerializer::Serialize(JsonObject.ToSharedRef(), JsonWriter);
+				NormalizedLocation *= uint16_MAX;
+				// ConvertToNormalizedScreenLocation(Pos, NormalizedLocation);
+				JsonObject->SetNumberField(TEXT("x"), static_cast<uint16>(NormalizedLocation.X));
+				JsonObject->SetNumberField(TEXT("y"), static_cast<uint16>(NormalizedLocation.Y));
 
-				FBufferArchive Buffer;
-				Buffer << Descriptor;
-				TArray<uint8> Data(Buffer.GetData(), Buffer.Num());
-				// Specific implementation for this method is handled in the pixel streaming module
-				IPixelStreamingInputModule::Get().OnSendMessage.Broadcast(FMemoryReader(Data));
+				FText TextboxContents;
+				if (FocusedWidget->GetType() == TEXT("SEditableText"))
+				{
+					SEditableText* TextBox = static_cast<SEditableText*>(FocusedWidget.Get());
+					TextboxContents = TextBox->GetText();
+				}
+				else if (FocusedWidget->GetType() == TEXT("SMultiLineEditableText"))
+				{
+					SMultiLineEditableTextBox* TextBox = static_cast<SMultiLineEditableTextBox*>(FocusedWidget.Get());
+					TextboxContents = TextBox->GetText();
+				}
+
+				JsonObject->SetStringField(TEXT("contents"), TextboxContents.ToString());
 			}
+
+			FString Descriptor;
+			TSharedRef<TJsonWriter<TCHAR, TCondensedJsonPrintPolicy<TCHAR>>> JsonWriter = TJsonWriterFactory<TCHAR, TCondensedJsonPrintPolicy<TCHAR>>::Create(&Descriptor);
+			FJsonSerializer::Serialize(JsonObject.ToSharedRef(), JsonWriter);
+
+			FBufferArchive Buffer;
+			Buffer << Descriptor;
+			TArray<uint8> Data(Buffer.GetData(), Buffer.Num());
+			// Specific implementation for this method is handled in the pixel streaming module
+			IPixelStreamingInputModule::Get().OnSendMessage.Broadcast(FMemoryReader(Data));
 		});
 	}
 
diff --git a/Engine/Plugins/Media/PixelStreaming/Source/PixelStreamingInput/Private/PixelStreamingInputHandler.h b/Engine/Plugins/Media/PixelStreaming/Source/PixelStreamingInput/Private/PixelStreamingInputHandler.h
--- a/Engine/Plugins/Media/PixelStreaming/Source/PixelStreamingInput/Private/PixelStreamingInputHandler.h
+++ b/Engine/Plugins/Media/PixelStreaming/Source/PixelStreamingInput/Private/PixelStreamingInputHandler.h
@@ -106,6 +106,10 @@ namespace UE::PixelStreamingInput
 		 * UI Interaction handling
 		 */
 		virtual void HandleUIInteraction(FMemoryReader Ar);
+		/**
+		 * Textbox Entry handling
+		 */
+		virtual void HandleOnTextboxEntry(FMemoryReader Ar);
 
 		/**
 		 * Populate default command handlers for data channel messages sent with "{ type: "Command" }".
