Failing automation tests after merge, fixed problem in test cleanup, resolving most or all tests

diff --git a/Engine/Plugins/Media/PixelStreaming/Source/PixelStreaming/Private/Tests/MockPlayer.h b/Engine/Plugins/Media/PixelStreaming/Source/PixelStreaming/Private/Tests/MockPlayer.h
--- a/Engine/Plugins/Media/PixelStreaming/Source/PixelStreaming/Private/Tests/MockPlayer.h
+++ b/Engine/Plugins/Media/PixelStreaming/Source/PixelStreaming/Private/Tests/MockPlayer.h
@@ -70,15 +70,6 @@ namespace UE::PixelStreaming
 	class FMockPlayer : public IPixelStreamingSignallingConnectionObserver
 	{
 	public:
-		FMockPlayer()
-		{
-			SignallingServerConnection = MakeUnique<FPixelStreamingSignallingConnection>(MakeShareable(this), TEXT("FMockPlayer"));
-
-			UE::PixelStreaming::DoOnGameThreadAndWait(MAX_uint32, []() {
-				Settings::CVarPixelStreamingSuppressICECandidateErrors->Set(true, ECVF_SetByCode);
-			});
-		}
-
 		virtual ~FMockPlayer()
 		{
 			Disconnect();
@@ -225,4 +216,4 @@ namespace UE::PixelStreaming
 
 } // namespace UE::PixelStreaming
 
-#endif // WITH_DEV_AUTOMATION_TESTS
\ No newline at end of file
+#endif // WITH_DEV_AUTOMATION_TESTS
diff --git a/Engine/Plugins/Media/PixelStreaming/Source/PixelStreaming/Private/Tests/TestUtils.cpp b/Engine/Plugins/Media/PixelStreaming/Source/PixelStreaming/Private/Tests/TestUtils.cpp
--- a/Engine/Plugins/Media/PixelStreaming/Source/PixelStreaming/Private/Tests/TestUtils.cpp
+++ b/Engine/Plugins/Media/PixelStreaming/Source/PixelStreaming/Private/Tests/TestUtils.cpp
@@ -261,6 +261,13 @@ namespace UE::PixelStreaming
 	TSharedPtr<FMockPlayer> CreatePlayer(FMockPlayer::EMode OfferMode)
 	{
 		TSharedPtr<FMockPlayer> OutPlayer = MakeShared<FMockPlayer>();
+
+		// We have to pass an existing shared ptr to the mockplayer or else we will end up with competing reference controllers
+		OutPlayer->SignallingServerConnection = MakeUnique<FPixelStreamingSignallingConnection>(OutPlayer, TEXT("FMockPlayer"));
+		UE::PixelStreaming::DoOnGameThreadAndWait(MAX_uint32, []() {
+			Settings::CVarPixelStreamingSuppressICECandidateErrors->Set(true, ECVF_SetByCode);
+		});
+
 		OutPlayer->SetMode(OfferMode);
 		return OutPlayer;
 	}
@@ -293,4 +300,4 @@ namespace UE::PixelStreaming
 
 } // namespace UE::PixelStreaming
 
-#endif // WITH_DEV_AUTOMATION_TESTS
\ No newline at end of file
+#endif // WITH_DEV_AUTOMATION_TESTS
