Fix: Move default protocol population to stop failing tests

diff --git a/Engine/Plugins/Media/PixelStreaming/Source/PixelStreaming/Private/PixelStreamingModule.cpp b/Engine/Plugins/Media/PixelStreaming/Source/PixelStreaming/Private/PixelStreamingModule.cpp
--- a/Engine/Plugins/Media/PixelStreaming/Source/PixelStreaming/Private/PixelStreamingModule.cpp
+++ b/Engine/Plugins/Media/PixelStreaming/Source/PixelStreaming/Private/PixelStreamingModule.cpp
@@ -67,7 +67,6 @@ IPixelStreamingModule* UE::PixelStreaming::FPixelStreamingModule::PixelStreaming
 
 namespace UE::PixelStreaming
 {
-	typedef EPixelStreamingMessageTypes EType;
 	/**
 	 * IModuleInterface implementation
 	 */
@@ -102,7 +101,6 @@ namespace UE::PixelStreaming
 			return;
 		}
 
-		PopulateProtocol();
 		RegisterCustomHandlers();
 
 		// By calling InitDefaultStreamer post engine init we can use pixel streaming in standalone editor mode
@@ -438,105 +436,6 @@ namespace UE::PixelStreaming
 		return bCompatible;
 	}
 
-	void FPixelStreamingModule::PopulateProtocol()
-	{
-		// Old EToStreamerMsg Commands
-		/*
-		 * Control Messages.
-		 */
-		// Simple command with no payload
-		// Note, we only specify the ID when creating these messages to preserve backwards compatability
-		// when adding your own message type, you can simply do FPixelStreamingInputProtocol.Direction.Add("XXX");
-		FPixelStreamingInputProtocol::ToStreamerProtocol.Add("IFrameRequest", FPixelStreamingInputMessage(0));
-		FPixelStreamingInputProtocol::ToStreamerProtocol.Add("RequestQualityControl", FPixelStreamingInputMessage(1));
-		FPixelStreamingInputProtocol::ToStreamerProtocol.Add("FpsRequest", FPixelStreamingInputMessage(2));
-		FPixelStreamingInputProtocol::ToStreamerProtocol.Add("AverageBitrateRequest", FPixelStreamingInputMessage(3));
-		FPixelStreamingInputProtocol::ToStreamerProtocol.Add("StartStreaming", FPixelStreamingInputMessage(4));
-		FPixelStreamingInputProtocol::ToStreamerProtocol.Add("StopStreaming", FPixelStreamingInputMessage(5));
-		FPixelStreamingInputProtocol::ToStreamerProtocol.Add("LatencyTest", FPixelStreamingInputMessage(6));
-		FPixelStreamingInputProtocol::ToStreamerProtocol.Add("RequestInitialSettings", FPixelStreamingInputMessage(7));
-		FPixelStreamingInputProtocol::ToStreamerProtocol.Add("TestEcho", FPixelStreamingInputMessage(8));
-
-		/*
-		 * Input Messages.
-		 */
-		// Generic Input Messages.
-		FPixelStreamingInputProtocol::ToStreamerProtocol.Add("UIInteraction", FPixelStreamingInputMessage(50));
-		FPixelStreamingInputProtocol::ToStreamerProtocol.Add("Command", FPixelStreamingInputMessage(51));
-		FPixelStreamingInputProtocol::ToStreamerProtocol.Add("TextboxEntry", FPixelStreamingInputMessage(52));
-
-		// Keyboard Input Message.
-		// Complex command with payload, therefore we specify the length of the payload (bytes) as well as the structure of the payload
-		FPixelStreamingInputProtocol::ToStreamerProtocol.Add("KeyDown", FPixelStreamingInputMessage(60, { EType::Uint8, EType::Uint8 }));
-		FPixelStreamingInputProtocol::ToStreamerProtocol.Add("KeyUp", FPixelStreamingInputMessage(61, { EType::Uint8 }));
-		FPixelStreamingInputProtocol::ToStreamerProtocol.Add("KeyPress", FPixelStreamingInputMessage(62, { EType::Uint16 }));
-
-		// Mouse Input Messages.
-		FPixelStreamingInputProtocol::ToStreamerProtocol.Add("MouseEnter", FPixelStreamingInputMessage(70));
-		FPixelStreamingInputProtocol::ToStreamerProtocol.Add("MouseLeave", FPixelStreamingInputMessage(71));
-		FPixelStreamingInputProtocol::ToStreamerProtocol.Add("MouseDown", FPixelStreamingInputMessage(72, { EType::Uint8, EType::Uint16, EType::Uint16 }));
-		FPixelStreamingInputProtocol::ToStreamerProtocol.Add("MouseUp", FPixelStreamingInputMessage(73, { EType::Uint8, EType::Uint16, EType::Uint16 }));
-		FPixelStreamingInputProtocol::ToStreamerProtocol.Add("MouseMove", FPixelStreamingInputMessage(74, { EType::Uint16, EType::Uint16, EType::Uint16, EType::Uint16 }));
-		FPixelStreamingInputProtocol::ToStreamerProtocol.Add("MouseWheel", FPixelStreamingInputMessage(75, { EType::Int16, EType::Uint16, EType::Uint16 }));
-		FPixelStreamingInputProtocol::ToStreamerProtocol.Add("MouseDouble", FPixelStreamingInputMessage(76, { EType::Uint8, EType::Uint16, EType::Uint16 }));
-
-		// Touch Input Messages.
-		FPixelStreamingInputProtocol::ToStreamerProtocol.Add("TouchStart", FPixelStreamingInputMessage(80, { EType::Uint8, EType::Uint16, EType::Uint16, EType::Uint8, EType::Uint8, EType::Uint8 }));
-		FPixelStreamingInputProtocol::ToStreamerProtocol.Add("TouchEnd", FPixelStreamingInputMessage(81, { EType::Uint8, EType::Uint16, EType::Uint16, EType::Uint8, EType::Uint8, EType::Uint8 }));
-		FPixelStreamingInputProtocol::ToStreamerProtocol.Add("TouchMove", FPixelStreamingInputMessage(82, { EType::Uint8, EType::Uint16, EType::Uint16, EType::Uint8, EType::Uint8, EType::Uint8 }));
-
-		// Gamepad Input Messages.
-		FPixelStreamingInputProtocol::ToStreamerProtocol.Add("GamepadButtonPressed", FPixelStreamingInputMessage(90, { EType::Uint8, EType::Uint8, EType::Uint8 }));
-		FPixelStreamingInputProtocol::ToStreamerProtocol.Add("GamepadButtonReleased", FPixelStreamingInputMessage(91, { EType::Uint8, EType::Uint8, EType::Uint8 }));
-		FPixelStreamingInputProtocol::ToStreamerProtocol.Add("GamepadAnalog", FPixelStreamingInputMessage(92, { EType::Uint8, EType::Uint8, EType::Double }));
-
-		// XR Input Messages.
-		// clang-format off
-		FPixelStreamingInputProtocol::ToStreamerProtocol.Add("XRHMDTransform", FPixelStreamingInputMessage(110, {	// 4x4 Transform
-																													EType::Float, EType::Float, EType::Float, EType::Float,
-																													EType::Float, EType::Float, EType::Float, EType::Float,
-																													EType::Float, EType::Float, EType::Float, EType::Float,
-																													EType::Float, EType::Float, EType::Float, EType::Float,
-																												}));
-		
-		FPixelStreamingInputProtocol::ToStreamerProtocol.Add("XRControllerTransform", FPixelStreamingInputMessage(111, {// 4x4 Transform
-																														EType::Float, EType::Float, EType::Float, EType::Float, 
-																														EType::Float, EType::Float, EType::Float, EType::Float, 
-																														EType::Float, EType::Float, EType::Float, EType::Float, 
-																														EType::Float, EType::Float, EType::Float, EType::Float,
-																														// Handedness (L, R, Any)
-																														EType::Uint8 
-																														}));
-		
-		FPixelStreamingInputProtocol::ToStreamerProtocol.Add("XRButtonPressed", FPixelStreamingInputMessage(112,// Handedness,   ButtonIdx,      IsRepeat
-																												{ EType::Uint8, EType::Uint8, EType::Uint8 }));
-		FPixelStreamingInputProtocol::ToStreamerProtocol.Add("XRButtonTouched", FPixelStreamingInputMessage(113,// Handedness,   ButtonIdx,      IsRepeat
-																												{ EType::Uint8, EType::Uint8, EType::Uint8 }));
-		FPixelStreamingInputProtocol::ToStreamerProtocol.Add("XRButtonReleased", FPixelStreamingInputMessage(114,// Handedness,   ButtonIdx,     IsRepeat
-																												{ EType::Uint8, EType::Uint8, EType::Uint8 }));
-
-		// clang-format on
-		FPixelStreamingInputProtocol::ToStreamerProtocol.Add("XRAnalog", FPixelStreamingInputMessage(115, { EType::Uint8, EType::Uint8, EType::Double }));
-
-		FPixelStreamingInputProtocol::ToStreamerProtocol.Add("XRSystem", FPixelStreamingInputMessage(116, { EType::Uint8 }));
-
-		// Old EToPlayerMsg commands
-		FPixelStreamingInputProtocol::FromStreamerProtocol.Add("QualityControlOwnership", FPixelStreamingInputMessage(0));
-		FPixelStreamingInputProtocol::FromStreamerProtocol.Add("Response", FPixelStreamingInputMessage(1));
-		FPixelStreamingInputProtocol::FromStreamerProtocol.Add("Command", FPixelStreamingInputMessage(2));
-		FPixelStreamingInputProtocol::FromStreamerProtocol.Add("FreezeFrame", FPixelStreamingInputMessage(3));
-		FPixelStreamingInputProtocol::FromStreamerProtocol.Add("UnfreezeFrame", FPixelStreamingInputMessage(4));
-		FPixelStreamingInputProtocol::FromStreamerProtocol.Add("VideoEncoderAvgQP", FPixelStreamingInputMessage(5));
-		FPixelStreamingInputProtocol::FromStreamerProtocol.Add("LatencyTest", FPixelStreamingInputMessage(6));
-		FPixelStreamingInputProtocol::FromStreamerProtocol.Add("InitialSettings", FPixelStreamingInputMessage(7));
-		FPixelStreamingInputProtocol::FromStreamerProtocol.Add("FileExtension", FPixelStreamingInputMessage(8));
-		FPixelStreamingInputProtocol::FromStreamerProtocol.Add("FileMimeType", FPixelStreamingInputMessage(9));
-		FPixelStreamingInputProtocol::FromStreamerProtocol.Add("FileContents", FPixelStreamingInputMessage(10));
-		FPixelStreamingInputProtocol::FromStreamerProtocol.Add("TestEcho", FPixelStreamingInputMessage(11));
-		FPixelStreamingInputProtocol::FromStreamerProtocol.Add("InputControlOwnership", FPixelStreamingInputMessage(12));
-		FPixelStreamingInputProtocol::FromStreamerProtocol.Add("Protocol", FPixelStreamingInputMessage(255));
-	}
-
 	void FPixelStreamingModule::RegisterCustomHandlers()
 	{
 		IPixelStreamingInputModule& InputModule = IPixelStreamingInputModule::Get();
diff --git a/Engine/Plugins/Media/PixelStreaming/Source/PixelStreaming/Private/PixelStreamingModule.h b/Engine/Plugins/Media/PixelStreaming/Source/PixelStreaming/Private/PixelStreamingModule.h
--- a/Engine/Plugins/Media/PixelStreaming/Source/PixelStreaming/Private/PixelStreamingModule.h
+++ b/Engine/Plugins/Media/PixelStreaming/Source/PixelStreaming/Private/PixelStreamingModule.h
@@ -65,7 +65,6 @@ namespace UE::PixelStreaming
 		// Own methods
 		void InitDefaultStreamer();
 		bool IsPlatformCompatible() const;
-		void PopulateProtocol();
 		void RegisterCustomHandlers();
 		void HandleSendCommand(FMemoryReader Ar);
 		void HandleUIInteraction(FMemoryReader Ar);
diff --git a/Engine/Plugins/Media/PixelStreaming/Source/PixelStreamingInput/Private/PixelStreamingInputModule.cpp b/Engine/Plugins/Media/PixelStreaming/Source/PixelStreamingInput/Private/PixelStreamingInputModule.cpp
--- a/Engine/Plugins/Media/PixelStreaming/Source/PixelStreamingInput/Private/PixelStreamingInputModule.cpp
+++ b/Engine/Plugins/Media/PixelStreaming/Source/PixelStreamingInput/Private/PixelStreamingInputModule.cpp
@@ -9,12 +9,15 @@
 
 namespace UE::PixelStreamingInput
 {
+	typedef EPixelStreamingMessageTypes EType;
+
 	void FPixelStreamingInputModule::StartupModule()
 	{
 		if (!FSlateApplication::IsInitialized())
 		{
 			return;
 		}
+		PopulateProtocol();
 
 		TSharedPtr<FPixelStreamingApplicationWrapper> PixelStreamerApplicationWrapper = MakeShareable(new FPixelStreamingApplicationWrapper(FSlateApplication::Get().GetPlatformApplication()));
 		TSharedPtr<FGenericApplicationMessageHandler> BaseHandler = FSlateApplication::Get().GetPlatformApplication()->GetMessageHandler();
@@ -53,6 +56,105 @@ namespace UE::PixelStreamingInput
 	{
 		return InputHandler;
 	}
+
+	void FPixelStreamingInputModule::PopulateProtocol()
+	{
+		// Old EToStreamerMsg Commands
+		/*
+		 * Control Messages.
+		 */
+		// Simple command with no payload
+		// Note, we only specify the ID when creating these messages to preserve backwards compatability
+		// when adding your own message type, you can simply do FPixelStreamingInputProtocol::Direction.Add("XXX");
+		FPixelStreamingInputProtocol::ToStreamerProtocol.Add("IFrameRequest", FPixelStreamingInputMessage(0));
+		FPixelStreamingInputProtocol::ToStreamerProtocol.Add("RequestQualityControl", FPixelStreamingInputMessage(1));
+		FPixelStreamingInputProtocol::ToStreamerProtocol.Add("FpsRequest", FPixelStreamingInputMessage(2));
+		FPixelStreamingInputProtocol::ToStreamerProtocol.Add("AverageBitrateRequest", FPixelStreamingInputMessage(3));
+		FPixelStreamingInputProtocol::ToStreamerProtocol.Add("StartStreaming", FPixelStreamingInputMessage(4));
+		FPixelStreamingInputProtocol::ToStreamerProtocol.Add("StopStreaming", FPixelStreamingInputMessage(5));
+		FPixelStreamingInputProtocol::ToStreamerProtocol.Add("LatencyTest", FPixelStreamingInputMessage(6));
+		FPixelStreamingInputProtocol::ToStreamerProtocol.Add("RequestInitialSettings", FPixelStreamingInputMessage(7));
+		FPixelStreamingInputProtocol::ToStreamerProtocol.Add("TestEcho", FPixelStreamingInputMessage(8));
+
+		/*
+		 * Input Messages.
+		 */
+		// Generic Input Messages.
+		FPixelStreamingInputProtocol::ToStreamerProtocol.Add("UIInteraction", FPixelStreamingInputMessage(50));
+		FPixelStreamingInputProtocol::ToStreamerProtocol.Add("Command", FPixelStreamingInputMessage(51));
+		FPixelStreamingInputProtocol::ToStreamerProtocol.Add("TextboxEntry", FPixelStreamingInputMessage(52));
+
+		// Keyboard Input Message.
+		// Complex command with payload, therefore we specify the length of the payload (bytes) as well as the structure of the payload
+		FPixelStreamingInputProtocol::ToStreamerProtocol.Add("KeyDown", FPixelStreamingInputMessage(60, { EType::Uint8, EType::Uint8 }));
+		FPixelStreamingInputProtocol::ToStreamerProtocol.Add("KeyUp", FPixelStreamingInputMessage(61, { EType::Uint8 }));
+		FPixelStreamingInputProtocol::ToStreamerProtocol.Add("KeyPress", FPixelStreamingInputMessage(62, { EType::Uint16 }));
+
+		// Mouse Input Messages.
+		FPixelStreamingInputProtocol::ToStreamerProtocol.Add("MouseEnter", FPixelStreamingInputMessage(70));
+		FPixelStreamingInputProtocol::ToStreamerProtocol.Add("MouseLeave", FPixelStreamingInputMessage(71));
+		FPixelStreamingInputProtocol::ToStreamerProtocol.Add("MouseDown", FPixelStreamingInputMessage(72, { EType::Uint8, EType::Uint16, EType::Uint16 }));
+		FPixelStreamingInputProtocol::ToStreamerProtocol.Add("MouseUp", FPixelStreamingInputMessage(73, { EType::Uint8, EType::Uint16, EType::Uint16 }));
+		FPixelStreamingInputProtocol::ToStreamerProtocol.Add("MouseMove", FPixelStreamingInputMessage(74, { EType::Uint16, EType::Uint16, EType::Uint16, EType::Uint16 }));
+		FPixelStreamingInputProtocol::ToStreamerProtocol.Add("MouseWheel", FPixelStreamingInputMessage(75, { EType::Int16, EType::Uint16, EType::Uint16 }));
+		FPixelStreamingInputProtocol::ToStreamerProtocol.Add("MouseDouble", FPixelStreamingInputMessage(76, { EType::Uint8, EType::Uint16, EType::Uint16 }));
+
+		// Touch Input Messages.
+		FPixelStreamingInputProtocol::ToStreamerProtocol.Add("TouchStart", FPixelStreamingInputMessage(80, { EType::Uint8, EType::Uint16, EType::Uint16, EType::Uint8, EType::Uint8, EType::Uint8 }));
+		FPixelStreamingInputProtocol::ToStreamerProtocol.Add("TouchEnd", FPixelStreamingInputMessage(81, { EType::Uint8, EType::Uint16, EType::Uint16, EType::Uint8, EType::Uint8, EType::Uint8 }));
+		FPixelStreamingInputProtocol::ToStreamerProtocol.Add("TouchMove", FPixelStreamingInputMessage(82, { EType::Uint8, EType::Uint16, EType::Uint16, EType::Uint8, EType::Uint8, EType::Uint8 }));
+
+		// Gamepad Input Messages.
+		FPixelStreamingInputProtocol::ToStreamerProtocol.Add("GamepadButtonPressed", FPixelStreamingInputMessage(90, { EType::Uint8, EType::Uint8, EType::Uint8 }));
+		FPixelStreamingInputProtocol::ToStreamerProtocol.Add("GamepadButtonReleased", FPixelStreamingInputMessage(91, { EType::Uint8, EType::Uint8, EType::Uint8 }));
+		FPixelStreamingInputProtocol::ToStreamerProtocol.Add("GamepadAnalog", FPixelStreamingInputMessage(92, { EType::Uint8, EType::Uint8, EType::Double }));
+
+		// XR Input Messages.
+		// clang-format off
+		FPixelStreamingInputProtocol::ToStreamerProtocol.Add("XRHMDTransform", FPixelStreamingInputMessage(110, {	// 4x4 Transform
+																													EType::Float, EType::Float, EType::Float, EType::Float,
+																													EType::Float, EType::Float, EType::Float, EType::Float,
+																													EType::Float, EType::Float, EType::Float, EType::Float,
+																													EType::Float, EType::Float, EType::Float, EType::Float,
+																												}));
+		
+		FPixelStreamingInputProtocol::ToStreamerProtocol.Add("XRControllerTransform", FPixelStreamingInputMessage(111, {// 4x4 Transform
+																														EType::Float, EType::Float, EType::Float, EType::Float, 
+																														EType::Float, EType::Float, EType::Float, EType::Float, 
+																														EType::Float, EType::Float, EType::Float, EType::Float, 
+																														EType::Float, EType::Float, EType::Float, EType::Float,
+																														// Handedness (L, R, Any)
+																														EType::Uint8 
+																														}));
+		
+		FPixelStreamingInputProtocol::ToStreamerProtocol.Add("XRButtonPressed", FPixelStreamingInputMessage(112,// Handedness,   ButtonIdx,      IsRepeat
+																												{ EType::Uint8, EType::Uint8, EType::Uint8 }));
+		FPixelStreamingInputProtocol::ToStreamerProtocol.Add("XRButtonTouched", FPixelStreamingInputMessage(113,// Handedness,   ButtonIdx,      IsRepeat
+																												{ EType::Uint8, EType::Uint8, EType::Uint8 }));
+		FPixelStreamingInputProtocol::ToStreamerProtocol.Add("XRButtonReleased", FPixelStreamingInputMessage(114,// Handedness,   ButtonIdx,     IsRepeat
+																												{ EType::Uint8, EType::Uint8, EType::Uint8 }));
+
+		// clang-format on
+		FPixelStreamingInputProtocol::ToStreamerProtocol.Add("XRAnalog", FPixelStreamingInputMessage(115, { EType::Uint8, EType::Uint8, EType::Double }));
+
+		FPixelStreamingInputProtocol::ToStreamerProtocol.Add("XRSystem", FPixelStreamingInputMessage(116, { EType::Uint8 }));
+
+		// Old EToPlayerMsg commands
+		FPixelStreamingInputProtocol::FromStreamerProtocol.Add("QualityControlOwnership", FPixelStreamingInputMessage(0));
+		FPixelStreamingInputProtocol::FromStreamerProtocol.Add("Response", FPixelStreamingInputMessage(1));
+		FPixelStreamingInputProtocol::FromStreamerProtocol.Add("Command", FPixelStreamingInputMessage(2));
+		FPixelStreamingInputProtocol::FromStreamerProtocol.Add("FreezeFrame", FPixelStreamingInputMessage(3));
+		FPixelStreamingInputProtocol::FromStreamerProtocol.Add("UnfreezeFrame", FPixelStreamingInputMessage(4));
+		FPixelStreamingInputProtocol::FromStreamerProtocol.Add("VideoEncoderAvgQP", FPixelStreamingInputMessage(5));
+		FPixelStreamingInputProtocol::FromStreamerProtocol.Add("LatencyTest", FPixelStreamingInputMessage(6));
+		FPixelStreamingInputProtocol::FromStreamerProtocol.Add("InitialSettings", FPixelStreamingInputMessage(7));
+		FPixelStreamingInputProtocol::FromStreamerProtocol.Add("FileExtension", FPixelStreamingInputMessage(8));
+		FPixelStreamingInputProtocol::FromStreamerProtocol.Add("FileMimeType", FPixelStreamingInputMessage(9));
+		FPixelStreamingInputProtocol::FromStreamerProtocol.Add("FileContents", FPixelStreamingInputMessage(10));
+		FPixelStreamingInputProtocol::FromStreamerProtocol.Add("TestEcho", FPixelStreamingInputMessage(11));
+		FPixelStreamingInputProtocol::FromStreamerProtocol.Add("InputControlOwnership", FPixelStreamingInputMessage(12));
+		FPixelStreamingInputProtocol::FromStreamerProtocol.Add("Protocol", FPixelStreamingInputMessage(255));
+	}
 } // namespace UE::PixelStreamingInput
 
 IMPLEMENT_MODULE(UE::PixelStreamingInput::FPixelStreamingInputModule, PixelStreamingInput)
\ No newline at end of file
diff --git a/Engine/Plugins/Media/PixelStreaming/Source/PixelStreamingInput/Private/PixelStreamingInputModule.h b/Engine/Plugins/Media/PixelStreaming/Source/PixelStreamingInput/Private/PixelStreamingInputModule.h
--- a/Engine/Plugins/Media/PixelStreaming/Source/PixelStreamingInput/Private/PixelStreamingInputModule.h
+++ b/Engine/Plugins/Media/PixelStreaming/Source/PixelStreamingInput/Private/PixelStreamingInputModule.h
@@ -26,5 +26,7 @@ namespace UE::PixelStreamingInput
 		// NOTE: There is only ever a single input handler which all of the streamers share. This provides
 		// a central point for the multiple streamers -> single UE instance
 		TSharedPtr<IPixelStreamingInputHandler> InputHandler;
+
+		void PopulateProtocol();
 	};
 } // namespace UE::PixelStreamingInput
\ No newline at end of file
