Fix: Remove H265 as a selectable codec as it currently doesn't stream.

diff --git a/Engine/Plugins/Media/PixelStreaming/Source/PixelStreamingEditor/Private/PixelStreamingToolbar.cpp b/Engine/Plugins/Media/PixelStreaming/Source/PixelStreamingEditor/Private/PixelStreamingToolbar.cpp
--- a/Engine/Plugins/Media/PixelStreaming/Source/PixelStreamingEditor/Private/PixelStreamingToolbar.cpp
+++ b/Engine/Plugins/Media/PixelStreaming/Source/PixelStreamingEditor/Private/PixelStreamingToolbar.cpp
@@ -200,24 +200,24 @@ namespace UE::EditorPixelStreaming
 			})
 		);
 
-		PluginCommands->MapAction(
-			FPixelStreamingCommands::Get().H265,
-			FExecuteAction::CreateLambda([Module = &PixelStreamingModule]()
-			{
-				Module->SetCodec(EPixelStreamingCodec::H265);
-			}),
-			FCanExecuteAction::CreateLambda([Module = &PixelStreamingModule] { 
-				if(TSharedPtr<IPixelStreamingStreamer> Streamer = Module->GetStreamer(Module->GetDefaultStreamerID()))	
-				{
-					return !Streamer->IsStreaming();
-				}
-				return false;
-			}),
-			FIsActionChecked::CreateLambda([Module = &PixelStreamingModule]()
-			{ 
-				return Module->GetCodec() == EPixelStreamingCodec::H265;
-			})
-		);
+		// PluginCommands->MapAction(
+		// 	FPixelStreamingCommands::Get().H265,
+		// 	FExecuteAction::CreateLambda([Module = &PixelStreamingModule]()
+		// 	{
+		// 		Module->SetCodec(EPixelStreamingCodec::H265);
+		// 	}),
+		// 	FCanExecuteAction::CreateLambda([Module = &PixelStreamingModule] { 
+		// 		if(TSharedPtr<IPixelStreamingStreamer> Streamer = Module->GetStreamer(Module->GetDefaultStreamerID()))	
+		// 		{
+		// 			return !Streamer->IsStreaming();
+		// 		}
+		// 		return false;
+		// 	}),
+		// 	FIsActionChecked::CreateLambda([Module = &PixelStreamingModule]()
+		// 	{ 
+		// 		return Module->GetCodec() == EPixelStreamingCodec::H265;
+		// 	})
+		// );
 
 		UToolMenus::RegisterStartupCallback(FSimpleMulticastDelegate::FDelegate::CreateRaw(this, &FPixelStreamingToolbar::RegisterMenus));
 	}
@@ -522,7 +522,7 @@ namespace UE::EditorPixelStreaming
 		MenuBuilder.BeginSection("Codec", LOCTEXT("PixelStreamingCodecSettings", "Codec"));
 		{
 			MenuBuilder.AddMenuEntry(FPixelStreamingCommands::Get().H264);
-			MenuBuilder.AddMenuEntry(FPixelStreamingCommands::Get().H265);
+			// MenuBuilder.AddMenuEntry(FPixelStreamingCommands::Get().H265);
 			MenuBuilder.AddMenuEntry(FPixelStreamingCommands::Get().VP8);
 			MenuBuilder.AddMenuEntry(FPixelStreamingCommands::Get().VP9);
 		}
