Allow passing arbitrary config settings on PixelStreaming startup.
Don't use WITH_EDITOR as criteria for hover default.

diff --git a/Engine/Plugins/Media/PixelStreaming/Source/PixelStreaming/Private/PixelStreamingSignallingConnection.cpp b/Engine/Plugins/Media/PixelStreaming/Source/PixelStreaming/Private/PixelStreamingSignallingConnection.cpp
--- a/Engine/Plugins/Media/PixelStreaming/Source/PixelStreaming/Private/PixelStreamingSignallingConnection.cpp
+++ b/Engine/Plugins/Media/PixelStreaming/Source/PixelStreaming/Private/PixelStreamingSignallingConnection.cpp
@@ -123,11 +123,6 @@ void FPixelStreamingSignallingConnection::SendOffer(FPixelStreamingPlayerId Play
 	FString SdpStr = UE::PixelStreaming::ToString(SdpAnsi);
 	OfferJson->SetStringField(TEXT("sdp"), SdpStr);
 
-#if WITH_EDITOR
-	// Default hover preference needs to be sent with the offer to prevent error messages on older infrastructure
-	OfferJson->SetStringField(TEXT("defaultToHover"), TEXT("true"));
-#endif // WITH_EDITOR
-
 	UE_LOG(LogPixelStreamingSS, Log, TEXT("Sending player=%s \"offer\" to SS %s"), *PlayerId, *Url);
 	UE_LOG(LogPixelStreamingSS, Verbose, TEXT("SDP offer\n%s"), *SdpStr);
 
diff --git a/Engine/Plugins/Media/PixelStreaming/Source/PixelStreaming/Private/Streamer.cpp b/Engine/Plugins/Media/PixelStreaming/Source/PixelStreaming/Private/Streamer.cpp
--- a/Engine/Plugins/Media/PixelStreaming/Source/PixelStreaming/Private/Streamer.cpp
+++ b/Engine/Plugins/Media/PixelStreaming/Source/PixelStreaming/Private/Streamer.cpp
@@ -368,6 +368,32 @@ namespace UE::PixelStreaming
 		FPixelStreamingPeerConnection::RemoveAudioInput(AudioInput);
 	}
 
+	void FStreamer::SetConfigOption(const FName& OptionName, const FString& Value)
+	{
+		if (Value.IsEmpty())
+		{
+			ConfigOptions.Remove(OptionName);
+		}
+		else
+		{
+			ConfigOptions.Add(OptionName, Value);
+		}
+	}
+
+	bool FStreamer::GetConfigOption(const FName& OptionName, FString& OutValue)
+	{
+		FString* OptionValue = ConfigOptions.Find(OptionName);
+		if (OptionValue)
+		{
+			OutValue = *OptionValue;
+			return true;
+		}
+		else
+		{
+			return false;
+		}
+	}
+
 	void FStreamer::AddPlayerConfig(TSharedRef<FJsonObject>& JsonObject)
 	{
 		checkf(InputHandler.IsValid(), TEXT("No Input Device available when populating Player Config"));
@@ -800,7 +826,20 @@ namespace UE::PixelStreaming
 			Settings::CVarPixelStreamingEnableFillerData.GetValueOnAnyThread() ? 1 : 0,
 			*Settings::CVarPixelStreamingEncoderMultipass.GetValueOnAnyThread());
 
-		const FString FullPayload = FString::Printf(TEXT("{ \"PixelStreaming\": %s, \"Encoder\": %s, \"WebRTC\": %s }"), *PixelStreamingPayload, *EncoderPayload, *WebRTCPayload);
+		FString ConfigPayload = TEXT("{ ");
+		bool bComma = false; // Simplest way to avoid complaints from pedantic JSON parsers
+		for (const TPair<FName, FString>& Option: ConfigOptions)
+		{
+			if (bComma)
+			{
+				ConfigPayload.Append(TEXT(", "));
+			}
+			ConfigPayload.Append(FString::Printf(TEXT("\"%s\": \"%s\""), *Option.Key.ToString(), *Option.Value));
+			bComma = true;
+		}
+		ConfigPayload.Append(TEXT("}"));
+
+		const FString FullPayload = FString::Printf(TEXT("{ \"PixelStreaming\": %s, \"Encoder\": %s, \"WebRTC\": %s, \"ConfigOptions\": %s }"), *PixelStreamingPayload, *EncoderPayload, *WebRTCPayload, *ConfigPayload);
 
 		if (const FPlayerContext* PlayerContext = Players.Find(PlayerId))
 		{
diff --git a/Engine/Plugins/Media/PixelStreaming/Source/PixelStreaming/Private/Streamer.h b/Engine/Plugins/Media/PixelStreaming/Source/PixelStreaming/Private/Streamer.h
--- a/Engine/Plugins/Media/PixelStreaming/Source/PixelStreaming/Private/Streamer.h
+++ b/Engine/Plugins/Media/PixelStreaming/Source/PixelStreaming/Private/Streamer.h
@@ -70,6 +70,9 @@ namespace UE::PixelStreaming
 		TSharedPtr<IPixelStreamingAudioInput> CreateAudioInput() override;
 		void RemoveAudioInput(TSharedPtr<IPixelStreamingAudioInput> AudioInput) override;
 
+		virtual void SetConfigOption(const FName& OptionName, const FString& Value) override;
+		virtual bool GetConfigOption(const FName& OptionName, FString& OutValue) override;
+
 		// TODO(Luke) hook this back up so that the Engine can change how the interface is working browser side
 		void AddPlayerConfig(TSharedRef<FJsonObject>& JsonObject);
 
@@ -138,5 +141,7 @@ namespace UE::PixelStreaming
 		FDelegateHandle AllConnectionsClosedHandle;
 
 		IPixelStreamingModule& Module;
+
+		TMap<FName, FString> ConfigOptions;
 	};
 } // namespace UE::PixelStreaming
diff --git a/Engine/Plugins/Media/PixelStreaming/Source/PixelStreaming/Public/IPixelStreamingStreamer.h b/Engine/Plugins/Media/PixelStreaming/Source/PixelStreaming/Public/IPixelStreamingStreamer.h
--- a/Engine/Plugins/Media/PixelStreaming/Source/PixelStreaming/Public/IPixelStreamingStreamer.h
+++ b/Engine/Plugins/Media/PixelStreaming/Source/PixelStreaming/Public/IPixelStreamingStreamer.h
@@ -257,4 +257,12 @@ public:
 	 * @brief Remove the audio input from the Pixel Streaming mix.
 	 */
 	virtual void RemoveAudioInput(TSharedPtr<IPixelStreamingAudioInput> AudioInput) = 0;
+
+	/**
+	 * @brief Allows sending arbitrary configuration options during initial connection
+	 *
+	 * @param Value Setting a value to an empty string clears it from the mapping and prevents it being sent
+	 */
+	virtual void SetConfigOption(const FName& OptionName, const FString& Value) = 0;
+	virtual bool GetConfigOption(const FName& OptionName, FString& OutValue) = 0;
 };
diff --git a/Engine/Plugins/Media/PixelStreaming/Source/PixelStreamingEditor/Private/PixelStreamingEditorModule.cpp b/Engine/Plugins/Media/PixelStreaming/Source/PixelStreamingEditor/Private/PixelStreamingEditorModule.cpp
--- a/Engine/Plugins/Media/PixelStreaming/Source/PixelStreamingEditor/Private/PixelStreamingEditorModule.cpp
+++ b/Engine/Plugins/Media/PixelStreaming/Source/PixelStreamingEditor/Private/PixelStreamingEditorModule.cpp
@@ -89,6 +89,8 @@ void FPixelStreamingEditorModule::InitEditorStreaming(IPixelStreamingModule& Mod
 		}
 	}
 
+	Streamer->SetConfigOption(FName(*FString(TEXT("DefaultToHover"))), TEXT("true"));
+
 	IMainFrameModule::Get().OnMainFrameCreationFinished().AddLambda([&](TSharedPtr<SWindow> RootWindow, bool bIsRunningStartupDialog) {
 		MaybeResizeEditor(RootWindow);
 
@@ -363,4 +365,4 @@ void FPixelStreamingEditorModule::MaybeResizeEditor(TSharedPtr<SWindow> RootWind
 	}
 }
 
-IMPLEMENT_MODULE(FPixelStreamingEditorModule, PixelStreamingEditor)
\ No newline at end of file
+IMPLEMENT_MODULE(FPixelStreamingEditorModule, PixelStreamingEditor)
