Change default Viewer port for PixelStreaming on Linux to 8080

diff --git a/Engine/Plugins/Media/PixelStreaming/Source/PixelStreamingEditor/Public/PixelStreamingEditorModule.h b/Engine/Plugins/Media/PixelStreaming/Source/PixelStreamingEditor/Public/PixelStreamingEditorModule.h
--- a/Engine/Plugins/Media/PixelStreaming/Source/PixelStreamingEditor/Public/PixelStreamingEditorModule.h
+++ b/Engine/Plugins/Media/PixelStreaming/Source/PixelStreamingEditor/Public/PixelStreamingEditorModule.h
@@ -53,8 +53,13 @@ private:
     FString SignallingDomain;
     // The port the streamer will connect to. eg 8888
     int32 StreamerPort;
+
     // The port the streams can be viewed at on the browser. eg 80
+#if PLATFORM_LINUX
+    int32 ViewerPort = 8080; // ports <1000 require superuser privileges on Linux
+#else
     int32 ViewerPort = 80;
+#endif
 
 public:
     bool bUseExternalSignallingServer = false;
