Fix: UE5 wouldn't stream if using the latest signalling server and opting to have the browser send the offer

diff --git a/Engine/Plugins/Experimental/PixelStreamingPlayer/Source/PixelStreamingPlayer/Private/PixelStreamingSignallingComponent.cpp b/Engine/Plugins/Experimental/PixelStreamingPlayer/Source/PixelStreamingPlayer/Private/PixelStreamingSignallingComponent.cpp
--- a/Engine/Plugins/Experimental/PixelStreamingPlayer/Source/PixelStreamingPlayer/Private/PixelStreamingSignallingComponent.cpp
+++ b/Engine/Plugins/Experimental/PixelStreamingPlayer/Source/PixelStreamingPlayer/Private/PixelStreamingSignallingComponent.cpp
@@ -116,7 +116,7 @@ void UPixelStreamingSignallingComponent::FPixelStreamingSignallingConnectionObse
 	// no-op player count is not interesting to us here
 }
 
-void UPixelStreamingSignallingComponent::FPixelStreamingSignallingConnectionObserver::OnSignallingPlayerConnected(FPixelStreamingPlayerId PlayerId, const FPixelStreamingPlayerConfig& PlayerConfig)
+void UPixelStreamingSignallingComponent::FPixelStreamingSignallingConnectionObserver::OnSignallingPlayerConnected(FPixelStreamingPlayerId PlayerId, const FPixelStreamingPlayerConfig& PlayerConfig, bool bSendOffer)
 {
 	// no-op our player receiving other players is possible but not a case we are interested in
 	UE_LOG(LogPixelStreamingPlayer, Warning, TEXT("Warning: Received message about player connected from SS, this very likely means you have misconfigured the Pixel Streaming player connection and entered the streamer port. Please check your connect node and ensure it using the player port, e.g. 80 (or no port)."));
diff --git a/Engine/Plugins/Experimental/PixelStreamingPlayer/Source/PixelStreamingPlayer/Public/PixelStreamingSignallingComponent.h b/Engine/Plugins/Experimental/PixelStreamingPlayer/Source/PixelStreamingPlayer/Public/PixelStreamingSignallingComponent.h
--- a/Engine/Plugins/Experimental/PixelStreamingPlayer/Source/PixelStreamingPlayer/Public/PixelStreamingSignallingComponent.h
+++ b/Engine/Plugins/Experimental/PixelStreamingPlayer/Source/PixelStreamingPlayer/Public/PixelStreamingSignallingComponent.h
@@ -113,7 +113,8 @@ protected:
 	class FPixelStreamingSignallingConnectionObserver : public IPixelStreamingSignallingConnectionObserver
 	{
 	public:
-		FPixelStreamingSignallingConnectionObserver(UPixelStreamingSignallingComponent* InParent) : Parent(InParent) {}
+		FPixelStreamingSignallingConnectionObserver(UPixelStreamingSignallingComponent* InParent)
+			: Parent(InParent) {}
 
 		//
 		// ISignallingServerConnectionObserver implementation.
@@ -126,12 +127,12 @@ protected:
 		virtual void OnSignallingRemoteIceCandidate(const FString& SdpMid, int SdpMLineIndex, const FString& Sdp) override;
 		virtual void OnSignallingPeerDataChannels(int32 SendStreamId, int32 RecvStreamId) override;
 		virtual void OnSignallingPlayerCount(uint32 Count) override;
-		virtual void OnSignallingPlayerConnected(FPixelStreamingPlayerId PlayerId, const FPixelStreamingPlayerConfig& PlayerConfig) override;
+		virtual void OnSignallingPlayerConnected(FPixelStreamingPlayerId PlayerId, const FPixelStreamingPlayerConfig& PlayerConfig, bool bSendOffer) override;
 		virtual void OnSignallingPlayerDisconnected(FPixelStreamingPlayerId PlayerId) override;
 		virtual void OnSignallingSFUPeerDataChannels(FPixelStreamingPlayerId SFUId, FPixelStreamingPlayerId PlayerId, int32 SendStreamId, int32 RecvStreamId) override;
 		virtual void OnSignallingSessionDescription(FPixelStreamingPlayerId PlayerId, webrtc::SdpType Type, const FString& Sdp) override {}
 		virtual void OnSignallingRemoteIceCandidate(FPixelStreamingPlayerId PlayerId, const FString& SdpMid, int SdpMLineIndex, const FString& Sdp) override {}
-	
+
 	private:
 		TObjectPtr<UPixelStreamingSignallingComponent> Parent;
 	};
diff --git a/Engine/Plugins/Media/PixelStreaming/Source/PixelStreaming/Private/PixelStreamingSignallingConnection.cpp b/Engine/Plugins/Media/PixelStreaming/Source/PixelStreaming/Private/PixelStreamingSignallingConnection.cpp
--- a/Engine/Plugins/Media/PixelStreaming/Source/PixelStreaming/Private/PixelStreamingSignallingConnection.cpp
+++ b/Engine/Plugins/Media/PixelStreaming/Source/PixelStreaming/Private/PixelStreamingSignallingConnection.cpp
@@ -563,14 +563,16 @@ void FPixelStreamingSignallingConnection::OnPlayerConnected(const FJsonObjectPtr
 	FPixelStreamingPlayerConfig PlayerConfig;
 
 	// Default to always making datachannel, unless explicitly set to false.
-	bool bMakeDataChannel = true;
-	Json->TryGetBoolField(TEXT("datachannel"), PlayerConfig.SupportsDataChannel);
+	Json->TryGetBoolField(TEXT("dataChannel"), PlayerConfig.SupportsDataChannel);
 
 	// Default peer is not an SFU, unless explictly set as SFU
-	bool bIsSFU = false;
 	Json->TryGetBoolField(TEXT("sfu"), PlayerConfig.IsSFU);
 
-	Observer->OnSignallingPlayerConnected(PlayerId, PlayerConfig);
+	// Default to always sending an offer, unless explicitly set to false
+	bool bSendOffer = true;
+	Json->TryGetBoolField(TEXT("sendOffer"), bSendOffer);
+
+	Observer->OnSignallingPlayerConnected(PlayerId, PlayerConfig, bSendOffer);
 }
 
 void FPixelStreamingSignallingConnection::OnPlayerDisconnected(const FJsonObjectPtr& Json)
diff --git a/Engine/Plugins/Media/PixelStreaming/Source/PixelStreaming/Private/SignallingConnectionObserver.cpp b/Engine/Plugins/Media/PixelStreaming/Source/PixelStreaming/Private/SignallingConnectionObserver.cpp
--- a/Engine/Plugins/Media/PixelStreaming/Source/PixelStreaming/Private/SignallingConnectionObserver.cpp
+++ b/Engine/Plugins/Media/PixelStreaming/Source/PixelStreaming/Private/SignallingConnectionObserver.cpp
@@ -62,9 +62,9 @@ namespace UE::PixelStreaming
 		}
 	}
 
-	void FPixelStreamingSignallingConnectionObserver::OnSignallingPlayerConnected(FPixelStreamingPlayerId PlayerId, const FPixelStreamingPlayerConfig& PlayerConfig)
+	void FPixelStreamingSignallingConnectionObserver::OnSignallingPlayerConnected(FPixelStreamingPlayerId PlayerId, const FPixelStreamingPlayerConfig& PlayerConfig, bool bSendOffer)
 	{
-		Streamer.OnPlayerConnected(PlayerId, PlayerConfig, true /*bSendOffer*/);
+		Streamer.OnPlayerConnected(PlayerId, PlayerConfig, bSendOffer);
 	}
 
 	void FPixelStreamingSignallingConnectionObserver::OnSignallingPlayerDisconnected(FPixelStreamingPlayerId PlayerId)
diff --git a/Engine/Plugins/Media/PixelStreaming/Source/PixelStreaming/Private/SignallingConnectionObserver.h b/Engine/Plugins/Media/PixelStreaming/Source/PixelStreaming/Private/SignallingConnectionObserver.h
--- a/Engine/Plugins/Media/PixelStreaming/Source/PixelStreaming/Private/SignallingConnectionObserver.h
+++ b/Engine/Plugins/Media/PixelStreaming/Source/PixelStreaming/Private/SignallingConnectionObserver.h
@@ -21,7 +21,7 @@ namespace UE::PixelStreaming
 		// Streamer-only
 		virtual void OnSignallingSessionDescription(FPixelStreamingPlayerId PlayerId, webrtc::SdpType Type, const FString& Sdp) override;
 		virtual void OnSignallingRemoteIceCandidate(FPixelStreamingPlayerId PlayerId, const FString& SdpMid, int SdpMLineIndex, const FString& Sdp) override;
-		virtual void OnSignallingPlayerConnected(FPixelStreamingPlayerId PlayerId, const FPixelStreamingPlayerConfig& PlayerConfig) override;
+		virtual void OnSignallingPlayerConnected(FPixelStreamingPlayerId PlayerId, const FPixelStreamingPlayerConfig& PlayerConfig, bool bSendOffer) override;
 		virtual void OnSignallingPlayerDisconnected(FPixelStreamingPlayerId PlayerId) override;
 		virtual void OnSignallingSFUPeerDataChannels(FPixelStreamingPlayerId SFUId, FPixelStreamingPlayerId PlayerId, int32 SendStreamId, int32 RecvStreamId) override;
 
diff --git a/Engine/Plugins/Media/PixelStreaming/Source/PixelStreaming/Private/Streamer.cpp b/Engine/Plugins/Media/PixelStreaming/Source/PixelStreaming/Private/Streamer.cpp
--- a/Engine/Plugins/Media/PixelStreaming/Source/PixelStreaming/Private/Streamer.cpp
+++ b/Engine/Plugins/Media/PixelStreaming/Source/PixelStreaming/Private/Streamer.cpp
@@ -488,10 +488,15 @@ namespace UE::PixelStreaming
 		}
 	}
 
-
-
 	void FStreamer::OnPlayerConnected(FPixelStreamingPlayerId PlayerId, const FPixelStreamingPlayerConfig& PlayerConfig, bool bSendOffer)
 	{
+		if (!bSendOffer)
+		{
+			// If we're not sending the offer, don't create the player session
+			// we'll wait until the offer arrives to do that
+			return;
+		}
+
 		FPlayerContext& PlayerContext = Players.GetOrAdd(PlayerId);
 		PlayerContext.Config = PlayerConfig;
 
@@ -513,19 +518,16 @@ namespace UE::PixelStreaming
 				? FPixelStreamingPeerConnection::EReceiveMediaOption::Nothing
 				: FPixelStreamingPeerConnection::EReceiveMediaOption::Audio;
 
-			if (bSendOffer)
-			{
-				PlayerContext.PeerConnection->CreateOffer(
-					ReceiveOption,
-					[this, PlayerId](const webrtc::SessionDescriptionInterface* SDP) {
-						// on success
-						SignallingServerConnection->SendOffer(PlayerId, *SDP);
-					},
-					[this, PlayerId](const FString& Error) {
-						// on error
-						DeletePlayerSession(PlayerId);
-					});
-			}
+			PlayerContext.PeerConnection->CreateOffer(
+				ReceiveOption,
+				[this, PlayerId](const webrtc::SessionDescriptionInterface* SDP) {
+					// on success
+					SignallingServerConnection->SendOffer(PlayerId, *SDP);
+				},
+				[this, PlayerId](const FString& Error) {
+					// on error
+					DeletePlayerSession(PlayerId);
+				});
 		}
 	}
 
@@ -545,25 +547,12 @@ namespace UE::PixelStreaming
 
 	void FStreamer::OnOffer(FPixelStreamingPlayerId PlayerId, const FString& Sdp)
 	{
-		FPlayerContext* PlayerContext = Players.Find(PlayerId);
-
-		// For backwards compatibility with versions before 5.0 where browser/player "offer" happens first and there is no "playerConnected" we make a player right here.
-		if (PlayerContext == nullptr)
-		{
-			UE_LOG(LogPixelStreaming, Log, TEXT("Got offer before \"playerConnected\", making peer connection for this player - this should only happen when using older versions of the signalling server pre UE 5.0."));
-
-			FPixelStreamingPlayerConfig Config;
-			Config.SupportsDataChannel = true;
-			Config.IsSFU = false;
-			// Note: We do not send an offer here, as we are responding to an offer.
-			OnPlayerConnected(PlayerId, Config, false /*bSendOffer*/);
-			PlayerContext = Players.Find(PlayerId);
-		}
-
-		verifyf(PlayerContext, TEXT("Player context should not be nullptr at this point."));
+		FPlayerContext& PlayerContext = Players.GetOrAdd(PlayerId);
 
-		// clear any existing connection
-		PlayerContext->PeerConnection = nullptr;
+		FPixelStreamingPlayerConfig Config;
+		Config.SupportsDataChannel = true;
+		Config.IsSFU = false;
+		PlayerContext.Config = Config;
 
 		if (CreateSession(PlayerId))
 		{
@@ -571,11 +560,11 @@ namespace UE::PixelStreaming
 				DeletePlayerSession(PlayerId);
 			};
 
-			PlayerContext->PeerConnection->ReceiveOffer(
+			PlayerContext.PeerConnection->ReceiveOffer(
 				Sdp,
 				[this, PlayerContext, PlayerId, OnGeneralFailure]() {
 					AddStreams(PlayerId);
-					PlayerContext->PeerConnection->CreateAnswer(
+					PlayerContext.PeerConnection->CreateAnswer(
 						FPixelStreamingPeerConnection::EReceiveMediaOption::Audio,
 						[this, PlayerId](const webrtc::SessionDescriptionInterface* LocalDescription) {
 							SignallingServerConnection->SendAnswer(PlayerId, *LocalDescription);
@@ -1062,7 +1051,7 @@ namespace UE::PixelStreaming
 			// still wrapping the base FSlateApplication after we stop streaming
 			TArray<uint8> EmptyArray;
 			TFunction<void(FMemoryReader)> MouseLeaveHandler = IPixelStreamingInputModule::Get().FindMessageHandler("MouseLeave");
-			MouseLeaveHandler(FMemoryReader(EmptyArray));
+			// MouseLeaveHandler(FMemoryReader(EmptyArray));
 		}
 	}
 } // namespace UE::PixelStreaming
diff --git a/Engine/Plugins/Media/PixelStreaming/Source/PixelStreaming/Private/Tests/HandshakeTest.cpp b/Engine/Plugins/Media/PixelStreaming/Source/PixelStreaming/Private/Tests/HandshakeTest.cpp
--- a/Engine/Plugins/Media/PixelStreaming/Source/PixelStreaming/Private/Tests/HandshakeTest.cpp
+++ b/Engine/Plugins/Media/PixelStreaming/Source/PixelStreaming/Private/Tests/HandshakeTest.cpp
@@ -43,8 +43,7 @@ namespace UE::PixelStreaming
 	IMPLEMENT_SIMPLE_AUTOMATION_TEST(FHandshakeTestPlayerOffer, "System.Plugins.PixelStreaming.HandshakePlayerOffer", EAutomationTestFlags::EditorContext | EAutomationTestFlags::ClientContext | EAutomationTestFlags::ProductFilter)
 	bool FHandshakeTestPlayerOffer::RunTest(const FString& Parameters)
 	{
-		// TODO (william.belcher): This test crashes on Windows
-		// DoHandShake(this, FMockPlayer::EMode::CreateOffers, false /*bUseLegacySignallingServer*/);
+		DoHandShake(this, FMockPlayer::EMode::CreateOffers, false /*bUseLegacySignallingServer*/);
 		return true;
 	}
 
diff --git a/Engine/Plugins/Media/PixelStreaming/Source/PixelStreaming/Private/Tests/MockPlayer.h b/Engine/Plugins/Media/PixelStreaming/Source/PixelStreaming/Private/Tests/MockPlayer.h
--- a/Engine/Plugins/Media/PixelStreaming/Source/PixelStreaming/Private/Tests/MockPlayer.h
+++ b/Engine/Plugins/Media/PixelStreaming/Source/PixelStreaming/Private/Tests/MockPlayer.h
@@ -196,7 +196,7 @@ namespace UE::PixelStreaming
 		virtual void OnSignallingPeerDataChannels(int32 SendStreamId, int32 RecvStreamId) override {}
 		virtual void OnSignallingSessionDescription(FPixelStreamingPlayerId PlayerId, webrtc::SdpType Type, const FString& Sdp) override {}
 		virtual void OnSignallingRemoteIceCandidate(FPixelStreamingPlayerId PlayerId, const FString& SdpMid, int SdpMLineIndex, const FString& Sdp) override {}
-		virtual void OnSignallingPlayerConnected(FPixelStreamingPlayerId PlayerId, const FPixelStreamingPlayerConfig& PlayerConfig) override{};
+		virtual void OnSignallingPlayerConnected(FPixelStreamingPlayerId PlayerId, const FPixelStreamingPlayerConfig& PlayerConfig, bool bSendOffer) override {}
 		virtual void OnSignallingPlayerDisconnected(FPixelStreamingPlayerId PlayerId) override {}
 		virtual void OnSignallingSFUPeerDataChannels(FPixelStreamingPlayerId SFUId, FPixelStreamingPlayerId PlayerId, int32 SendStreamId, int32 RecvStreamId) override {}
 
diff --git a/Engine/Plugins/Media/PixelStreaming/Source/PixelStreaming/Private/Tests/SignallingServerConnectionTests.cpp b/Engine/Plugins/Media/PixelStreaming/Source/PixelStreaming/Private/Tests/SignallingServerConnectionTests.cpp
--- a/Engine/Plugins/Media/PixelStreaming/Source/PixelStreaming/Private/Tests/SignallingServerConnectionTests.cpp
+++ b/Engine/Plugins/Media/PixelStreaming/Source/PixelStreaming/Private/Tests/SignallingServerConnectionTests.cpp
@@ -36,7 +36,7 @@ namespace
 		{
 			SetIceCandidate = { SdpMid, SdpMLineIndex, Sdp, PlayerId };
 		}
-		virtual void OnSignallingPlayerConnected(FPixelStreamingPlayerId, const FPixelStreamingPlayerConfig&) override {}
+		virtual void OnSignallingPlayerConnected(FPixelStreamingPlayerId, const FPixelStreamingPlayerConfig&, bool bSendOffer) override {}
 		virtual void OnSignallingPlayerDisconnected(FPixelStreamingPlayerId PlayerId) override {}
 		virtual void OnSignallingSFUPeerDataChannels(FPixelStreamingPlayerId SFUId, FPixelStreamingPlayerId PlayerId, int32 SendStreamId, int32 RecvStreamId) override {}
 
@@ -304,15 +304,17 @@ namespace
 
 				FPixelStreamingPlayerConfig PlayerConfig;
 
-				// Default to always making datachannel, unless explicitly set to false.
-				bool bMakeDataChannel = true;
-				JsonMsg->TryGetBoolField(TEXT("datachannel"), PlayerConfig.SupportsDataChannel);
+				// Default to always making datachannel, unless explicitly set to false
+				JsonMsg->TryGetBoolField(TEXT("dataChannel"), PlayerConfig.SupportsDataChannel);
 
 				// Default peer is not an SFU, unless explictly set as SFU
-				bool bIsSFU = false;
 				JsonMsg->TryGetBoolField(TEXT("sfu"), PlayerConfig.IsSFU);
 
-				Observer.OnSignallingPlayerConnected(PlayerId, PlayerConfig);
+				// Default to always sending an offer, unless explicitly set to false
+				bool bSendOffer = true;
+				JsonMsg->TryGetBoolField(TEXT("sendOffer"), bSendOffer);
+
+				Observer.OnSignallingPlayerConnected(PlayerId, PlayerConfig, bSendOffer);
 			});
 			RegisterHandler("playerDisconnected", [this](FJsonObjectPtr JsonMsg) {
 				FPixelStreamingPlayerId PlayerId;
diff --git a/Engine/Plugins/Media/PixelStreaming/Source/PixelStreaming/Public/IPixelStreamingSignallingConnectionObserver.h b/Engine/Plugins/Media/PixelStreaming/Source/PixelStreaming/Public/IPixelStreamingSignallingConnectionObserver.h
--- a/Engine/Plugins/Media/PixelStreaming/Source/PixelStreaming/Public/IPixelStreamingSignallingConnectionObserver.h
+++ b/Engine/Plugins/Media/PixelStreaming/Source/PixelStreaming/Public/IPixelStreamingSignallingConnectionObserver.h
@@ -18,7 +18,9 @@ public:
 	// Streamer-only
 	virtual void OnSignallingSessionDescription(FPixelStreamingPlayerId PlayerId, webrtc::SdpType Type, const FString& Sdp) = 0;
 	virtual void OnSignallingRemoteIceCandidate(FPixelStreamingPlayerId PlayerId, const FString& SdpMid, int SdpMLineIndex, const FString& Sdp) = 0;
-	virtual void OnSignallingPlayerConnected(FPixelStreamingPlayerId PlayerId, const FPixelStreamingPlayerConfig& PlayerConfig) = 0;
+	UE_DEPRECATED(5.2, "Please use OnSignallingPlayerConnected(FPixelStreamingPlayerId PlayerId, const FPixelStreamingPlayerConfig& PlayerConfig, bool bSendOffer). Using this method defaults bSendOffer to true! ")
+	virtual void OnSignallingPlayerConnected(FPixelStreamingPlayerId PlayerId, const FPixelStreamingPlayerConfig& PlayerConfig) { OnSignallingPlayerConnected(PlayerId, PlayerConfig, true /* bSendOffer */); };
+	virtual void OnSignallingPlayerConnected(FPixelStreamingPlayerId PlayerId, const FPixelStreamingPlayerConfig& PlayerConfig, bool bSendOffer) = 0;
 	virtual void OnSignallingPlayerDisconnected(FPixelStreamingPlayerId PlayerId) = 0;
 	virtual void OnSignallingSFUPeerDataChannels(FPixelStreamingPlayerId SFUId, FPixelStreamingPlayerId PlayerId, int32 SendStreamId, int32 RecvStreamId) = 0;
 
