Fixing bad include path.

diff --git a/Engine/Plugins/Media/PixelStreaming/Source/PixelStreaming/Private/UtilsRender.h b/Engine/Plugins/Media/PixelStreaming/Source/PixelStreaming/Private/UtilsRender.h
--- a/Engine/Plugins/Media/PixelStreaming/Source/PixelStreaming/Private/UtilsRender.h
+++ b/Engine/Plugins/Media/PixelStreaming/Source/PixelStreaming/Private/UtilsRender.h
@@ -12,7 +12,7 @@
 #include "Utils.h"
 #include "MediaShaders.h"
 #include "RHI.h"
-#include "Runtime/Renderer/Private/ScreenPass.h"
+#include "ScreenPass.h"
 #include "PostProcess/SceneFilterRendering.h"
 
 namespace UE::PixelStreaming
