Fix: Add missing button mappings for FPixelStreamingInputConverter::XRInputToFKey

diff --git a/Engine/Plugins/Media/PixelStreaming/Source/PixelStreamingInput/Private/PixelStreamingInputHandler.cpp b/Engine/Plugins/Media/PixelStreaming/Source/PixelStreamingInput/Private/PixelStreamingInputHandler.cpp
--- a/Engine/Plugins/Media/PixelStreaming/Source/PixelStreamingInput/Private/PixelStreamingInputHandler.cpp
+++ b/Engine/Plugins/Media/PixelStreaming/Source/PixelStreamingInput/Private/PixelStreamingInputHandler.cpp
@@ -85,6 +85,7 @@ namespace UE::PixelStreamingInput
 		// HTC Vive - Left Hand
 		// Buttons
 		FPixelStreamingInputConverter::XRInputToFKey.Add(MakeTuple(XRSystem::HTCVive, EControllerHand::Left, 0, Action::Click), EKeys::Vive_Left_Trigger_Click);
+		FPixelStreamingInputConverter::XRInputToFKey.Add(MakeTuple(XRSystem::HTCVive, EControllerHand::Left, 0, Action::Axis), EKeys::Vive_Left_Trigger_Axis);
 		FPixelStreamingInputConverter::XRInputToFKey.Add(MakeTuple(XRSystem::HTCVive, EControllerHand::Left, 1, Action::Click), EKeys::Vive_Left_Grip_Click);
 		FPixelStreamingInputConverter::XRInputToFKey.Add(MakeTuple(XRSystem::HTCVive, EControllerHand::Left, 2, Action::Click), EKeys::Vive_Left_Trackpad_Click);
 		FPixelStreamingInputConverter::XRInputToFKey.Add(MakeTuple(XRSystem::HTCVive, EControllerHand::Left, 2, Action::Touch), EKeys::Vive_Left_Trackpad_Touch);
@@ -94,6 +95,7 @@ namespace UE::PixelStreamingInput
 		// HTC Vive - Right Hand
 		// Buttons
 		FPixelStreamingInputConverter::XRInputToFKey.Add(MakeTuple(XRSystem::HTCVive, EControllerHand::Right, 0, Action::Click), EKeys::Vive_Right_Trigger_Click);
+		FPixelStreamingInputConverter::XRInputToFKey.Add(MakeTuple(XRSystem::HTCVive, EControllerHand::Right, 0, Action::Axis), EKeys::Vive_Right_Trigger_Axis);
 		FPixelStreamingInputConverter::XRInputToFKey.Add(MakeTuple(XRSystem::HTCVive, EControllerHand::Right, 1, Action::Click), EKeys::Vive_Right_Grip_Click);
 		FPixelStreamingInputConverter::XRInputToFKey.Add(MakeTuple(XRSystem::HTCVive, EControllerHand::Right, 2, Action::Click), EKeys::Vive_Right_Trackpad_Click);
 		FPixelStreamingInputConverter::XRInputToFKey.Add(MakeTuple(XRSystem::HTCVive, EControllerHand::Right, 2, Action::Touch), EKeys::Vive_Right_Trackpad_Touch);
