Fix flickering with multiple Pixel Streaming streamers when using H.264 hardware encoder.

diff --git a/Engine/Plugins/Media/PixelStreaming/Source/PixelStreaming/Private/FrameBufferMultiFormat.cpp b/Engine/Plugins/Media/PixelStreaming/Source/PixelStreaming/Private/FrameBufferMultiFormat.cpp
--- a/Engine/Plugins/Media/PixelStreaming/Source/PixelStreaming/Private/FrameBufferMultiFormat.cpp
+++ b/Engine/Plugins/Media/PixelStreaming/Source/PixelStreaming/Private/FrameBufferMultiFormat.cpp
@@ -4,13 +4,14 @@
 
 namespace UE::PixelStreaming
 {
-	FFrameBufferMultiFormatBase::FFrameBufferMultiFormatBase(TSharedPtr<FPixelCaptureCapturerMultiFormat> InFrameCapturer)
+	FFrameBufferMultiFormatBase::FFrameBufferMultiFormatBase(TSharedPtr<FPixelCaptureCapturerMultiFormat> InFrameCapturer, uint32 InStreamId)
 		: FrameCapturer(InFrameCapturer)
+		, StreamId(InStreamId)
 	{
 	}
 
-	FFrameBufferMultiFormatLayered::FFrameBufferMultiFormatLayered(TSharedPtr<FPixelCaptureCapturerMultiFormat> InFrameCapturer)
-		: FFrameBufferMultiFormatBase(InFrameCapturer)
+	FFrameBufferMultiFormatLayered::FFrameBufferMultiFormatLayered(TSharedPtr<FPixelCaptureCapturerMultiFormat> InFrameCapturer, uint32 InStreamId)
+		: FFrameBufferMultiFormatBase(InFrameCapturer, InStreamId)
 	{
 	}
 
@@ -31,11 +32,11 @@ namespace UE::PixelStreaming
 
 	rtc::scoped_refptr<FFrameBufferMultiFormat> FFrameBufferMultiFormatLayered::GetLayer(int LayerIndex) const
 	{
-		return new rtc::RefCountedObject<FFrameBufferMultiFormat>(FrameCapturer, LayerIndex);
+		return new rtc::RefCountedObject<FFrameBufferMultiFormat>(FrameCapturer, StreamId, LayerIndex);
 	}
 
-	FFrameBufferMultiFormat::FFrameBufferMultiFormat(TSharedPtr<FPixelCaptureCapturerMultiFormat> InFrameCapturer, int InLayerIndex)
-		: FFrameBufferMultiFormatBase(InFrameCapturer)
+	FFrameBufferMultiFormat::FFrameBufferMultiFormat(TSharedPtr<FPixelCaptureCapturerMultiFormat> InFrameCapturer, uint32 InStreamId, int InLayerIndex)
+		: FFrameBufferMultiFormatBase(InFrameCapturer, InStreamId)
 		, LayerIndex(InLayerIndex)
 	{
 	}
diff --git a/Engine/Plugins/Media/PixelStreaming/Source/PixelStreaming/Private/FrameBufferMultiFormat.h b/Engine/Plugins/Media/PixelStreaming/Source/PixelStreaming/Private/FrameBufferMultiFormat.h
--- a/Engine/Plugins/Media/PixelStreaming/Source/PixelStreaming/Private/FrameBufferMultiFormat.h
+++ b/Engine/Plugins/Media/PixelStreaming/Source/PixelStreaming/Private/FrameBufferMultiFormat.h
@@ -15,7 +15,7 @@ namespace UE::PixelStreaming
 	class FFrameBufferMultiFormatBase : public webrtc::VideoFrameBuffer
 	{
 	public:
-		FFrameBufferMultiFormatBase(TSharedPtr<FPixelCaptureCapturerMultiFormat> InFrameCapturer);
+		FFrameBufferMultiFormatBase(TSharedPtr<FPixelCaptureCapturerMultiFormat> InFrameCapturer, uint32 InStreamId);
 		virtual ~FFrameBufferMultiFormatBase() = default;
 
 		virtual webrtc::VideoFrameBuffer::Type type() const override { return webrtc::VideoFrameBuffer::Type::kNative; }
@@ -34,6 +34,7 @@ namespace UE::PixelStreaming
 
 	protected:
 		TSharedPtr<FPixelCaptureCapturerMultiFormat> FrameCapturer;
+		uint32 StreamId;
 	};
 
 	/**
@@ -42,7 +43,7 @@ namespace UE::PixelStreaming
 	class FFrameBufferMultiFormatLayered : public FFrameBufferMultiFormatBase
 	{
 	public:
-		FFrameBufferMultiFormatLayered(TSharedPtr<FPixelCaptureCapturerMultiFormat> InFrameCapturer);
+		FFrameBufferMultiFormatLayered(TSharedPtr<FPixelCaptureCapturerMultiFormat> InFrameCapturer, uint32 InStreamId);
 		virtual ~FFrameBufferMultiFormatLayered() = default;
 
 		virtual int width() const override;
@@ -58,7 +59,7 @@ namespace UE::PixelStreaming
 	class FFrameBufferMultiFormat : public FFrameBufferMultiFormatBase
 	{
 	public:
-		FFrameBufferMultiFormat(TSharedPtr<FPixelCaptureCapturerMultiFormat> InFrameCapturer, int InLayerIndex);
+		FFrameBufferMultiFormat(TSharedPtr<FPixelCaptureCapturerMultiFormat> InFrameCapturer, uint32 StreamId, int InLayerIndex);
 		virtual ~FFrameBufferMultiFormat() = default;
 
 		virtual int width() const override;
@@ -66,6 +67,8 @@ namespace UE::PixelStreaming
 
 		IPixelCaptureOutputFrame* RequestFormat(int32 Format) const;
 
+		uint32 GetSourceStreamId() const { return StreamId; }
+
 	private:
 		int32 LayerIndex;
 		// we want the frame buffer to always refer to the same frame. so the first request for a format
diff --git a/Engine/Plugins/Media/PixelStreaming/Source/PixelStreaming/Private/PixelStreamingVideoInput.cpp b/Engine/Plugins/Media/PixelStreaming/Source/PixelStreaming/Private/PixelStreamingVideoInput.cpp
--- a/Engine/Plugins/Media/PixelStreaming/Source/PixelStreaming/Private/PixelStreamingVideoInput.cpp
+++ b/Engine/Plugins/Media/PixelStreaming/Source/PixelStreaming/Private/PixelStreamingVideoInput.cpp
@@ -6,7 +6,10 @@
 #include "PixelCaptureBufferFormat.h"
 #include "PixelStreamingTrace.h"
 
+uint32 FPixelStreamingVideoInput::NextStreamId = 1;
+
 FPixelStreamingVideoInput::FPixelStreamingVideoInput()
+:StreamId(NextStreamId++)
 {
 	CreateFrameCapturer();
 }
@@ -36,7 +39,7 @@ void FPixelStreamingVideoInput::OnFrame(const IPixelCaptureInputFrame& InputFram
 
 rtc::scoped_refptr<webrtc::VideoFrameBuffer> FPixelStreamingVideoInput::GetFrameBuffer()
 {
-	return new rtc::RefCountedObject<UE::PixelStreaming::FFrameBufferMultiFormatLayered>(FrameCapturer);
+	return new rtc::RefCountedObject<UE::PixelStreaming::FFrameBufferMultiFormatLayered>(FrameCapturer, StreamId);
 }
 
 TSharedPtr<IPixelCaptureOutputFrame> FPixelStreamingVideoInput::RequestFormat(int32 Format, int32 LayerIndex)
diff --git a/Engine/Plugins/Media/PixelStreaming/Source/PixelStreaming/Private/VideoEncoderFactorySingleLayer.cpp b/Engine/Plugins/Media/PixelStreaming/Source/PixelStreaming/Private/VideoEncoderFactorySingleLayer.cpp
--- a/Engine/Plugins/Media/PixelStreaming/Source/PixelStreaming/Private/VideoEncoderFactorySingleLayer.cpp
+++ b/Engine/Plugins/Media/PixelStreaming/Source/PixelStreaming/Private/VideoEncoderFactorySingleLayer.cpp
@@ -140,7 +140,7 @@ namespace UE::PixelStreaming
 		}
 	}
 
-	void FVideoEncoderFactorySingleLayer::OnEncodedImage(const webrtc::EncodedImage& encoded_image, const webrtc::CodecSpecificInfo* codec_specific_info)
+	void FVideoEncoderFactorySingleLayer::OnEncodedImage(const webrtc::EncodedImage& encoded_image, const webrtc::CodecSpecificInfo* codec_specific_info, uint32 StreamId)
 	{
 		// Lock as we send encoded image to each encoder.
 		FScopeLock Lock(&ActiveEncodersGuard);
@@ -150,7 +150,7 @@ namespace UE::PixelStreaming
 			// Go through each encoder and send our encoded image to its callback
 			for (FVideoEncoderSingleLayerHardware* Encoder : ActiveEncoders)
 			{
-				Encoder->SendEncodedImage(encoded_image, codec_specific_info);
+				Encoder->SendEncodedImage(encoded_image, codec_specific_info, StreamId);
 			}
 		}
 	}
@@ -162,82 +162,35 @@ namespace UE::PixelStreaming
 		ActiveEncoders.Remove(Encoder);
 	}
 
-	TWeakPtr<TVideoEncoder<FVideoResourceRHI>> FVideoEncoderFactorySingleLayer::GetHardwareEncoder() const
+	void FVideoEncoderFactorySingleLayer::FreeUnusedEncoders()
 	{
-		return HardwareEncoder;
-	}
-
-	TWeakPtr<TVideoEncoder<FVideoResourceRHI>> FVideoEncoderFactorySingleLayer::GetOrCreateHardwareEncoder(const FVideoEncoderConfigH264& VideoConfig)
-	{
-		FScopeLock InitLock(&InitEncoderGuard);
-
-		if (HardwareEncoder.IsValid())
-		{
-			return HardwareEncoder;
-		}
-		else
+		// first clear unused encoders.
+		TArray<uint32> DeleteStreamIds;
+		for (auto&& [EncoderStreamId, EncoderPtr] : HardwareEncoders)
 		{
-			// Make the hardware encoder wrapper
-			HardwareEncoder = FVideoEncoder::Create<FVideoResourceRHI>(FAVDevice::GetHardwareDevice(), VideoConfig);
-
-			if (!HardwareEncoder.IsValid())
+			if (EncoderPtr.IsUnique())
 			{
-				UE_LOG(LogPixelStreaming, Error, TEXT("Could not create encoder. Check encoder config or perhaps you used up all your HW encoders."));
-				// We could not make the encoder, so indicate the id was not set successfully.
-				return nullptr;
+				DeleteStreamIds.Add(EncoderStreamId);
 			}
-
-			return HardwareEncoder;
 		}
-	}
-
-	TWeakPtr<TVideoEncoder<FVideoResourceRHI>> FVideoEncoderFactorySingleLayer::GetOrCreateHardwareEncoder(const FVideoEncoderConfigH265& VideoConfig)
-	{
-		FScopeLock InitLock(&InitEncoderGuard);
-
-		if (HardwareEncoder.IsValid())
+		for (uint32 DeleteStreamId : DeleteStreamIds)
 		{
-			return HardwareEncoder;
-		}
-		else
-		{
-			// Make the hardware encoder wrapper
-			HardwareEncoder = FVideoEncoder::Create<FVideoResourceRHI>(FAVDevice::GetHardwareDevice(), VideoConfig);
-
-			if (!HardwareEncoder.IsValid())
-			{
-				UE_LOG(LogPixelStreaming, Error, TEXT("Could not create encoder. Check encoder config or perhaps you used up all your HW encoders."));
-				// We could not make the encoder, so indicate the id was not set successfully.
-				return nullptr;
-			}
-
-			return HardwareEncoder;
+			HardwareEncoders.Remove(DeleteStreamId);
 		}
 	}
 
 	void FVideoEncoderFactorySingleLayer::ForceKeyFrame()
 	{
-		if (HardwareEncoder)
-		{
-			bForceNextKeyframe = true;
-		}
+		bForceNextKeyframe = true;
 	}
 
 	bool FVideoEncoderFactorySingleLayer::ShouldForceKeyframe() const
 	{
-		if (HardwareEncoder)
-		{
-			return bForceNextKeyframe;
-		}
-
-		return false;
+		return bForceNextKeyframe;
 	}
 
 	void FVideoEncoderFactorySingleLayer::UnforceKeyFrame()
 	{
-		if (HardwareEncoder)
-		{
-			bForceNextKeyframe = false;
-		}
+		bForceNextKeyframe = false;
 	}
 } // namespace UE::PixelStreaming
diff --git a/Engine/Plugins/Media/PixelStreaming/Source/PixelStreaming/Private/VideoEncoderFactorySingleLayer.h b/Engine/Plugins/Media/PixelStreaming/Source/PixelStreaming/Private/VideoEncoderFactorySingleLayer.h
--- a/Engine/Plugins/Media/PixelStreaming/Source/PixelStreaming/Private/VideoEncoderFactorySingleLayer.h
+++ b/Engine/Plugins/Media/PixelStreaming/Source/PixelStreaming/Private/VideoEncoderFactorySingleLayer.h
@@ -34,19 +34,44 @@ namespace UE::PixelStreaming
 
 		void ReleaseVideoEncoder(FVideoEncoderSingleLayerHardware* Encoder);
 
-		TWeakPtr<FVideoEncoderHardware> GetHardwareEncoder() const;
-		TWeakPtr<FVideoEncoderHardware> GetOrCreateHardwareEncoder(const FVideoEncoderConfigH264& VideoConfig);
-		TWeakPtr<FVideoEncoderHardware> GetOrCreateHardwareEncoder(const FVideoEncoderConfigH265& VideoConfig);
-
 		void ForceKeyFrame();
 		bool ShouldForceKeyframe() const;
 		void UnforceKeyFrame();
-		void OnEncodedImage(const webrtc::EncodedImage& Encoded_image, const webrtc::CodecSpecificInfo* CodecSpecificInfo);
+		void OnEncodedImage(const webrtc::EncodedImage& Encoded_image, const webrtc::CodecSpecificInfo* CodecSpecificInfo, uint32 StreamId);
+
+		template <typename ConfigType>
+		TWeakPtr<TVideoEncoder<FVideoResourceRHI>> GetOrCreateHardwareEncoder(uint32 StreamId, const ConfigType& VideoConfig)
+		{
+			FScopeLock InitLock(&InitEncoderGuard);
+
+			FreeUnusedEncoders();
+
+			if (auto* ExistingEncoder = HardwareEncoders.Find(StreamId))
+			{
+				return *ExistingEncoder;
+			}
+			else
+			{
+				// Make the hardware encoder wrapper
+				TSharedPtr<FVideoEncoderHardware> HardwareEncoder = FVideoEncoder::Create<FVideoResourceRHI>(FAVDevice::GetHardwareDevice(), VideoConfig);
+
+				if (!HardwareEncoder.IsValid())
+				{
+					UE_LOG(LogPixelStreaming, Error, TEXT("Could not create encoder. Check encoder config or perhaps you used up all your HW encoders."));
+					// We could not make the encoder, so indicate the id was not set successfully.
+					return nullptr;
+				}
+
+				HardwareEncoders.Add(StreamId, HardwareEncoder);
+				return HardwareEncoder;
+			}
+		}
 
 	private:
 		static webrtc::SdpVideoFormat CreateH264Format(webrtc::H264Profile Profile, webrtc::H264Level Level);
+		void FreeUnusedEncoders();
 
-		TSharedPtr<FVideoEncoderHardware> HardwareEncoder;
+		TMap<uint32, TSharedPtr<FVideoEncoderHardware>> HardwareEncoders;
 
 		uint8 bForceNextKeyframe : 1;
 
diff --git a/Engine/Plugins/Media/PixelStreaming/Source/PixelStreaming/Private/VideoEncoderSingleLayerHardware.cpp b/Engine/Plugins/Media/PixelStreaming/Source/PixelStreaming/Private/VideoEncoderSingleLayerHardware.cpp
--- a/Engine/Plugins/Media/PixelStreaming/Source/PixelStreaming/Private/VideoEncoderSingleLayerHardware.cpp
+++ b/Engine/Plugins/Media/PixelStreaming/Source/PixelStreaming/Private/VideoEncoderSingleLayerHardware.cpp
@@ -51,28 +51,26 @@ namespace UE::PixelStreaming
 		{
 			case EPixelStreamingCodec::H264:
 			{
-				FVideoEncoderConfigH264 VideoConfig;
-				SetInitialSettings(InCodecSettings, VideoConfig);
-				VideoConfig.Profile = PixelStreaming::Settings::GetH264Profile();
-				VideoConfig.RepeatSPSPPS = true;
-				VideoConfig.IntraRefreshPeriodFrames = PixelStreaming::Settings::CVarPixelStreamingEncoderIntraRefreshPeriodFrames.GetValueOnAnyThread();
-				VideoConfig.IntraRefreshCountFrames = PixelStreaming::Settings::CVarPixelStreamingEncoderIntraRefreshCountFrames.GetValueOnAnyThread();
-				VideoConfig.KeyframeInterval = PixelStreaming::Settings::CVarPixelStreamingEncoderKeyframeInterval.GetValueOnAnyThread();
-
-				HardwareEncoder = Factory.GetOrCreateHardwareEncoder(VideoConfig);
+				TUniquePtr<FVideoEncoderConfigH264> VideoConfig = MakeUnique<FVideoEncoderConfigH264>();
+				SetInitialSettings(InCodecSettings, *VideoConfig);
+				VideoConfig->Profile = PixelStreaming::Settings::GetH264Profile();
+				VideoConfig->RepeatSPSPPS = true;
+				VideoConfig->IntraRefreshPeriodFrames = PixelStreaming::Settings::CVarPixelStreamingEncoderIntraRefreshPeriodFrames.GetValueOnAnyThread();
+				VideoConfig->IntraRefreshCountFrames = PixelStreaming::Settings::CVarPixelStreamingEncoderIntraRefreshCountFrames.GetValueOnAnyThread();
+				VideoConfig->KeyframeInterval = PixelStreaming::Settings::CVarPixelStreamingEncoderKeyframeInterval.GetValueOnAnyThread();
+				InitialVideoConfig = MoveTemp(VideoConfig);
 				return WEBRTC_VIDEO_CODEC_OK;
 			}
 			case EPixelStreamingCodec::H265:
 			{
-				FVideoEncoderConfigH265 VideoConfig;
-				SetInitialSettings(InCodecSettings, VideoConfig);
-				VideoConfig.Profile = PixelStreaming::Settings::GetH265Profile();
-				VideoConfig.RepeatSPSPPS = true;
-				VideoConfig.IntraRefreshPeriodFrames = PixelStreaming::Settings::CVarPixelStreamingEncoderIntraRefreshPeriodFrames.GetValueOnAnyThread();
-				VideoConfig.IntraRefreshCountFrames = PixelStreaming::Settings::CVarPixelStreamingEncoderIntraRefreshCountFrames.GetValueOnAnyThread();
-				VideoConfig.KeyframeInterval = PixelStreaming::Settings::CVarPixelStreamingEncoderKeyframeInterval.GetValueOnAnyThread();
-
-				HardwareEncoder = Factory.GetOrCreateHardwareEncoder(VideoConfig);
+				TUniquePtr<FVideoEncoderConfigH265> VideoConfig = MakeUnique<FVideoEncoderConfigH265>();
+				SetInitialSettings(InCodecSettings, *VideoConfig);
+				VideoConfig->Profile = PixelStreaming::Settings::GetH265Profile();
+				VideoConfig->RepeatSPSPPS = true;
+				VideoConfig->IntraRefreshPeriodFrames = PixelStreaming::Settings::CVarPixelStreamingEncoderIntraRefreshPeriodFrames.GetValueOnAnyThread();
+				VideoConfig->IntraRefreshCountFrames = PixelStreaming::Settings::CVarPixelStreamingEncoderIntraRefreshCountFrames.GetValueOnAnyThread();
+				VideoConfig->KeyframeInterval = PixelStreaming::Settings::CVarPixelStreamingEncoderKeyframeInterval.GetValueOnAnyThread();
+				InitialVideoConfig = MoveTemp(VideoConfig);
 				return WEBRTC_VIDEO_CODEC_OK;
 			}
 			default:
@@ -139,12 +137,41 @@ namespace UE::PixelStreaming
 		return NewFrame;
 	}
 
+	void FVideoEncoderSingleLayerHardware::LateInitHardwareEncoder(uint32 StreamId)
+	{
+		switch (Codec)
+		{
+			case EPixelStreamingCodec::H264:
+			{
+				FVideoEncoderConfigH264& VideoConfig = *StaticCast<FVideoEncoderConfigH264*>(InitialVideoConfig.Get());
+				HardwareEncoder = Factory.GetOrCreateHardwareEncoder(StreamId, VideoConfig);
+				break;
+			}
+			case EPixelStreamingCodec::H265:
+			{
+				FVideoEncoderConfigH265& VideoConfig = *StaticCast<FVideoEncoderConfigH265*>(InitialVideoConfig.Get());
+				HardwareEncoder = Factory.GetOrCreateHardwareEncoder(StreamId, VideoConfig);
+				break;
+			}
+		}
+	}
+
 	int32 FVideoEncoderSingleLayerHardware::Encode(webrtc::VideoFrame const& frame, std::vector<webrtc::VideoFrameType> const* frame_types)
 	{
+		const FFrameBufferMultiFormat* FrameBuffer = StaticCast<FFrameBufferMultiFormat*>(frame.video_frame_buffer().get());
+
+		// We late init here so we can pull the stream ID off the incoming frames and pull the correct encoder for the stream
+		// earlier locations do not have this information.
+		if (!HardwareEncoder.IsValid())
+		{
+			EncodingStreamId = FrameBuffer->GetSourceStreamId();
+			LateInitHardwareEncoder(EncodingStreamId);
+		}
+
 		if (TSharedPtr<FVideoEncoderHardware> const& PinnedHardwareEncoder = HardwareEncoder.Pin())
 		{
 			TRACE_CPUPROFILER_EVENT_SCOPE_ON_CHANNEL_STR("PixelStreaming Hardware Encoding", PixelStreamingChannel);
-			const FFrameBufferMultiFormat* FrameBuffer = StaticCast<FFrameBufferMultiFormat*>(frame.video_frame_buffer().get());
+			
 			IPixelCaptureOutputFrame* AdaptedLayer = FrameBuffer->RequestFormat(PixelCaptureBufferFormat::FORMAT_RHI);
 
 			if (AdaptedLayer == nullptr)
@@ -235,7 +262,7 @@ namespace UE::PixelStreaming
 
 				UpdateFrameMetadataPrePacketization(*AdaptedLayer);
 
-				Factory.OnEncodedImage(Image, &CodecInfo);
+				Factory.OnEncodedImage(Image, &CodecInfo, EncodingStreamId);
 
 				UpdateFrameMetadataPostPacketization(*AdaptedLayer);
 			}
@@ -342,12 +369,15 @@ namespace UE::PixelStreaming
 		}
 	}
 
-	void FVideoEncoderSingleLayerHardware::SendEncodedImage(webrtc::EncodedImage const& encoded_image, webrtc::CodecSpecificInfo const* codec_specific_info)
+	void FVideoEncoderSingleLayerHardware::SendEncodedImage(webrtc::EncodedImage const& encoded_image, webrtc::CodecSpecificInfo const* codec_specific_info, uint32 StreamId)
 	{
-		MaybeDumpFrame(encoded_image);
-		if (OnEncodedImageCallback)
+		if (StreamId == EncodingStreamId)
 		{
-			OnEncodedImageCallback->OnEncodedImage(encoded_image, codec_specific_info);
+			MaybeDumpFrame(encoded_image);
+			if (OnEncodedImageCallback)
+			{
+				OnEncodedImageCallback->OnEncodedImage(encoded_image, codec_specific_info);
+			}
 		}
 	}
 
diff --git a/Engine/Plugins/Media/PixelStreaming/Source/PixelStreaming/Private/VideoEncoderSingleLayerHardware.h b/Engine/Plugins/Media/PixelStreaming/Source/PixelStreaming/Private/VideoEncoderSingleLayerHardware.h
--- a/Engine/Plugins/Media/PixelStreaming/Source/PixelStreaming/Private/VideoEncoderSingleLayerHardware.h
+++ b/Engine/Plugins/Media/PixelStreaming/Source/PixelStreaming/Private/VideoEncoderSingleLayerHardware.h
@@ -34,9 +34,10 @@ namespace UE::PixelStreaming
 		// virtual void OnLossNotification(const LossNotification& loss_notification) override;
 		// End WebRTC Interface.
 
-		void SendEncodedImage(webrtc::EncodedImage const& encoded_image, webrtc::CodecSpecificInfo const* codec_specific_info);
+		void SendEncodedImage(webrtc::EncodedImage const& encoded_image, webrtc::CodecSpecificInfo const* codec_specific_info, uint32 StreamId);
 
 	private:
+		void LateInitHardwareEncoder(uint32 StreamId);
 		void UpdateConfig(uint32 width, uint32 height);
 		void MaybeDumpFrame(webrtc::EncodedImage const& encoded_image);
 
@@ -49,6 +50,7 @@ namespace UE::PixelStreaming
 		FVideoEncoderFactorySingleLayer& Factory;
 		EPixelStreamingCodec Codec;
 
+		TUniquePtr<FVideoEncoderConfig> InitialVideoConfig;
 		TWeakPtr<FVideoEncoderHardware> HardwareEncoder;
 
 		// We store this so we can restore back to it if the user decides to use then stop using the PixelStreaming.Encoder.TargetBitrate CVar.
@@ -59,5 +61,8 @@ namespace UE::PixelStreaming
 		// WebRTC may request a bitrate/framerate change using SetRates(), we only respect this if this encoder is actually encoding
 		// so we use this optional object to store a rate change and act upon it when this encoder does its next call to Encode().
 		TOptional<RateControlParameters> PendingRateChange;
+
+		// used to key into active hardware encoders and pull the correct encoder for the stream.
+		uint32 EncodingStreamId;
 	};
 } // namespace UE::PixelStreaming
diff --git a/Engine/Plugins/Media/PixelStreaming/Source/PixelStreaming/Public/PixelStreamingVideoInput.h b/Engine/Plugins/Media/PixelStreaming/Source/PixelStreaming/Public/PixelStreamingVideoInput.h
--- a/Engine/Plugins/Media/PixelStreaming/Source/PixelStreaming/Public/PixelStreamingVideoInput.h
+++ b/Engine/Plugins/Media/PixelStreaming/Source/PixelStreaming/Public/PixelStreamingVideoInput.h
@@ -45,7 +45,7 @@ public:
 	 */
 	DECLARE_MULTICAST_DELEGATE(FOnFrameCaptured);
 	FOnFrameCaptured OnFrameCaptured;
-	
+
 private:
 	int32 LastFrameWidth = -1;
 	int32 LastFrameHeight = -1;
@@ -57,4 +57,8 @@ private:
 
 	void CreateFrameCapturer();
 	void OnCaptureComplete();
+
+	// video inputs need a unique stream id so we can properly divide hardware encoders between inputs
+	uint32 StreamId = 0;
+	static uint32 NextStreamId;
 };
diff --git a/Engine/Plugins/Media/PixelStreaming/Source/PixelStreamingServers/Private/SignallingServer.cpp b/Engine/Plugins/Media/PixelStreaming/Source/PixelStreamingServers/Private/SignallingServer.cpp
--- a/Engine/Plugins/Media/PixelStreaming/Source/PixelStreamingServers/Private/SignallingServer.cpp
+++ b/Engine/Plugins/Media/PixelStreaming/Source/PixelStreamingServers/Private/SignallingServer.cpp
@@ -377,7 +377,7 @@ namespace UE::PixelStreamingServers
 	void FSignallingServer::OnStreamerIdMessage(uint16 ConnectionId, TSharedPtr<FJsonObject> JSONObj)
 	{
 		FString StreamerName;
-		if (JSONObj->TryGetStringField("endpointId", StreamerName))
+		if (JSONObj->TryGetStringField("id", StreamerName))
 		{
 			StreamersWS->NameConnection(ConnectionId, StreamerName);
 			StreamersWS->RemoveName(LEGACY_NAME);
