Resolve failing Non-Unity build from missing Log category

diff --git a/Engine/Plugins/Media/PixelStreaming/Source/PixelStreaming/Private/VideoEncoderFactorySingleLayer.h b/Engine/Plugins/Media/PixelStreaming/Source/PixelStreaming/Private/VideoEncoderFactorySingleLayer.h
--- a/Engine/Plugins/Media/PixelStreaming/Source/PixelStreaming/Private/VideoEncoderFactorySingleLayer.h
+++ b/Engine/Plugins/Media/PixelStreaming/Source/PixelStreaming/Private/VideoEncoderFactorySingleLayer.h
@@ -8,6 +8,8 @@
 #include "Video/Encoders/Configs/VideoEncoderConfigH265.h"
 #include "Video/Resources/VideoResourceRHI.h"
 
+#include "PixelStreamingPrivate.h"
+
 namespace UE::PixelStreaming
 {
 	using FVideoEncoderHardware = TVideoEncoder<FVideoResourceRHI>;
