Upgrade get_ps_servers.bat|.sh to pull down a release from Github.

- Downloaded infra now creates a DOWNLOAD_VERSION file that tracks the downloaded version and if latest release version is mismatched the frontend is redownloaded.

diff --git a/Engine/Plugins/Media/PixelStreaming/Resources/WebServers/get_ps_servers.bat b/Engine/Plugins/Media/PixelStreaming/Resources/WebServers/get_ps_servers.bat
--- a/Engine/Plugins/Media/PixelStreaming/Resources/WebServers/get_ps_servers.bat
+++ b/Engine/Plugins/Media/PixelStreaming/Resources/WebServers/get_ps_servers.bat
@@ -5,6 +5,21 @@
 @Rem Set script location as working directory for commands.
 pushd "%~dp0"
 
+@Rem Turned on delayed expansion so we have variables evaluated in nested scope
+setlocal EnableDelayedExpansion
+
+@Rem Unset all our variables as these persist between cmd sessions if cmd not closed.
+SET "PSInfraOrg="
+SET "PSInfraRepo="
+SET "PSInfraTagOrBranch="
+SET "ReleaseVersion="
+SET "ReleaseUrl="
+SET "IsTag="
+SET "RefType="
+SET "Url="
+SET "DownloadVersion="
+SET "FlagPassed="
+
 :arg_loop_start
 SET ARG=%1
 if DEFINED ARG (
@@ -13,16 +28,26 @@ if DEFINED ARG (
     )
     if "%ARG%"=="/v" (
         SET UEVersion=%2
+        SET FlagPassed=1
         SHIFT
     )
     if "%ARG%"=="/b" (
         SET PSInfraTagOrBranch=%2
         SET IsTag=0
+        SET FlagPassed=1
         SHIFT
     )
     if "%ARG%"=="/t" (
         SET PSInfraTagOrBranch=%2
         SET IsTag=1
+        SET FlagPassed=1
+        SHIFT
+    )
+    if "%ARG%"=="/r" (
+        SET "ReleaseVersion=%2"
+        SET "ReleaseUrl=https://github.com/EpicGames/PixelStreamingInfrastructure/releases/download/!ReleaseVersion!/!ReleaseVersion!.zip"
+        SET IsTag=0
+        SET FlagPassed=1
         SHIFT
     )
     SHIFT
@@ -47,53 +72,112 @@ if DEFINED UEVersion (
     SET PSInfraTagOrBranch=UE5.0
     SET IsTag=0
   )
+  if "%UEVersion%"=="5.1" (
+    SET PSInfraTagOrBranch=UE5.1
+    SET IsTag=0
+  )
+  if "%UEVersion%"=="5.2" (
+    SET PSInfraTagOrBranch=UE5.2
+    SET IsTag=0
+  )
 )
 
 @Rem If no arguments select a specific version, fetch the appropriate default
 if NOT DEFINED PSInfraTagOrBranch (
-    SET PSInfraTagOrBranch=UE5.1
+    SET PSInfraTagOrBranch=UE5.2
     SET IsTag=0
 )
+echo Tag or branch: !PSInfraTagOrBranch!
 
-@Rem Whether the named reference is a tag or a branch affects the URL we fetch it on
+@Rem Whether the named reference is a tag or a branch affects the Url we fetch it on
 if %IsTag%==1 (
   SET RefType=tags
 ) else (
   SET RefType=heads
 )
 
-@Rem Look for a SignallingWebServer directory next to this script
-if exist SignallingWebServer\ (
-  echo SignallingWebServer directory found...skipping install.
-) else (
-  echo SignallingWebServer directory not found...beginning ps-infra download.
-
-  @Rem Download ps-infra and follow redirects.
-  curl -L https://github.com/%PSInfraOrg%/%PSInfraRepo%/archive/refs/%RefType%/%PSInfraTagOrBranch%.zip > ps-infra.zip
-  
-  @Rem Unarchive the .zip
-  tar -xmf ps-infra.zip || echo bad archive, contents: && type ps-infra.zip && exit 0
-
-  @Rem Rename the extracted, versioned, directory
-  for /d %%i in ("PixelStreamingInfrastructure-*") do (
-    for /d %%j in ("%%i/*") do (
-      echo "%%i\%%j"
-      move "%%i\%%j" .
-    )
-    for %%j in ("%%i/*") do (
-      echo "%%i\%%j"
-      move "%%i\%%j" .
+@Rem We have a branch, no user-specified release, then check repo for the presence of a RELEASE_VERSION file in the current branch
+if %IsTag%==0 (
+  if NOT DEFINED ReleaseUrl (
+    @Rem We don't want to auto-set the release version if the user passed an explicit flag.
+    if NOT DEFINED FlagPassed (
+      FOR /F "tokens=* USEBACKQ" %%F IN (`curl -s -f -L https://raw.githubusercontent.com/EpicGames/PixelStreamingInfrastructure/%PSInfraTagOrBranch%/RELEASE_VERSION`) DO (
+        SET "ReleaseVersion=!PSInfraTagOrBranch!-%%F"
+        SET "ReleaseUrl=https://github.com/EpicGames/PixelStreamingInfrastructure/releases/download/!ReleaseVersion!/!ReleaseVersion!.zip"
+      )
     )
+  )
+)
+
+@Rem Set our DownloadVersion here as we use this to check the contents of our DOWNLOAD_VERSION file shortly.
+SET "DownloadVersion=%PSInfraTagOrBranch%"
+if DEFINED ReleaseVersion (
+  SET "DownloadVersion=!ReleaseVersion!"
+  echo Release: !ReleaseVersion!
+)
+
+@Rem Check for the existence of a DOWNLOAD_VERSION file and if found, check its contents against our %DownloadVersion%
+if exist DOWNLOAD_VERSION (
+
+  @Rem Read DOWNLOAD_VERSION file into variable
+  FOR /F "delims=" %%F IN ( DOWNLOAD_VERSION ) DO (
+    SET "PreviousDownloadVersion=%%F"
+    @Rem Remove whitespace
+    SET "PreviousDownloadVersion=!PreviousDownloadVersion: =!"
+  )
+
+  if !DownloadVersion! == !PreviousDownloadVersion! (
+    echo Downloaded version ^(!DownloadVersion!^) of PS infra matches release version ^(!PreviousDownloadVersion!^)...skipping install.
+    goto :EOF
+  ) else (
+    echo There is a newer released version ^(!DownloadVersion!^) - had ^(!PreviousDownloadVersion!^), downloading...
+  )
+) else (
+  echo DOWNLOAD_VERSION file not found...beginning ps-infra download.
+)
 
-    echo "%%i"
-    rmdir /s /q "%%i"
+@Rem By default set the download url to the .zip of the branch
+SET Url=https://github.com/%PSInfraOrg%/%PSInfraRepo%/archive/refs/%RefType%/%PSInfraTagOrBranch%.zip
+
+@Rem If we have a ReleaseUrl then set it to our download url
+if DEFINED ReleaseUrl (
+  SET Url=!ReleaseUrl!
+)
+
+@Rem Download ps-infra and follow redirects.
+echo Attempting downloading Pixel Streaming infrastructure from: !Url!
+curl -L !Url! > ps-infra.zip
+
+@Rem Unarchive the .zip
+tar -xmf ps-infra.zip || echo bad archive, contents: && type ps-infra.zip && exit 0
+
+@Rem Remove old infra
+if exist Frontend\ ( rmdir /s /q Frontend )
+if exist Matchmaker\ ( rmdir /s /q Matchmaker )
+if exist SignallingWebserver\ ( rmdir /s /q SignallingWebserver )
+if exist SFU\ ( rmdir /s /q SFU )
+
+@Rem Rename the extracted, versioned, directory
+for /d %%i in ("PixelStreamingInfrastructure-*") do (
+  for /d %%j in ("%%i/*") do (
+    echo "%%i\%%j"
+    move "%%i\%%j" .
+  )
+  for %%j in ("%%i/*") do (
+    echo "%%i\%%j"
+    move "%%i\%%j" .
   )
 
-  @Rem Delete the downloaded zip
-  del ps-infra.zip
+  echo "%%i"
+  rmdir /s /q "%%i"
 )
 
-exit 0
+@Rem Delete the downloaded zip
+del ps-infra.zip
+
+@Rem Create a DOWNLOAD_VERSION file, which we use as a comparison file to check if we should auto upgrade when these scripts are run again
+echo %DownloadVersion%> DOWNLOAD_VERSION
+goto :EOF
 
 :print_help
 echo.
@@ -101,10 +185,11 @@ echo  Tool for fetching PixelStreaming Infrastructure. If no flags are set speci
 echo  the recommended version will be chosen as a default.
 echo.
 echo  Usage:
-echo    %~n0%~x0 [^/h] [^/v ^<UE version^>] [^/b ^<branch^>] [^/t ^<tag^>]
+echo    %~n0%~x0 [^/h] [^/v ^<UE version^>] [^/b ^<branch^>] [^/t ^<tag^>] [^/r ^<release^>]
 echo  Where:
 echo    /v      Specify a version of Unreal Engine to download the recommended release for
 echo    /b      Specify a specific branch for the tool to download from repo
 echo    /t      Specify a specific tag for the tool to download from repo
+echo    /r      Specify a specific release url path e.g. https://github.com/EpicGames/PixelStreamingInfrastructure/releases/download/${RELEASE_HERE}.zip
 echo    /h      Display this help message
-exit 1
\ No newline at end of file
+goto :EOF
\ No newline at end of file
diff --git a/Engine/Plugins/Media/PixelStreaming/Resources/WebServers/get_ps_servers.sh b/Engine/Plugins/Media/PixelStreaming/Resources/WebServers/get_ps_servers.sh
old mode 100644
new mode 100755
--- a/Engine/Plugins/Media/PixelStreaming/Resources/WebServers/get_ps_servers.sh
+++ b/Engine/Plugins/Media/PixelStreaming/Resources/WebServers/get_ps_servers.sh
@@ -17,25 +17,37 @@ print_help() {
            release for
    -b      Specify a specific branch for the tool to download from repo
    -t      Specify a specific tag for the tool to download from repo
+   -r      Specify a specific release url path e.g. https://github.com/EpicGames/PixelStreamingInfrastructure/releases/download/<RELEASE_VERSION>/<RELEASE_VERSION>.zip
    -h      Display this help message
 "
  exit 1
 }
 
+# Set all default variables (e.g. # Name and version of ps-infra that we are downloading)
+PSInfraOrg=EpicGames
+PSInfraRepo=PixelStreamingInfrastructure
+PSInfraTagOrBranch=UE5.2
+RefType=heads
+IsTag=0
+ReleaseUrlBase=https://github.com/EpicGames/PixelStreamingInfrastructure/releases/download
+# Unset any variables that don't have defaults that we use that may have persisted between bash terminals.
+unset Url
+unset DownloadVersion
+unset FlagPassed
+unset ReleaseVersion
+unset ReleaseUrl
+
 while(($#)) ; do
   case "$1" in
    -h ) print_help;;
-   -v ) UEVersion="$2"; shift 2;;
-   -b ) PSInfraTagOrBranch="$2"; IsTag=0; shift 2;;
-   -t ) PSInfraTagOrBranch="$2"; IsTag=1; shift 2;;
+   -v ) UEVersion="$2"; FlagPassed=1; shift 2;;
+   -b ) PSInfraTagOrBranch="$2"; FlagPassed=1; IsTag=0; shift 2;;
+   -t ) PSInfraTagOrBranch="$2"; FlagPassed=1; IsTag=1; shift 2;;
+   -r ) ReleaseVersion="$2"; FlagPassed=1; IsTag=0; ReleaseUrl=$ReleaseUrlBase/$ReleaseVersion/$ReleaseVersion.tar.gz; shift 2;;
    * ) echo "Unknown command: $1"; shift;;
   esac
  done
 
-# Name and version of ps-infra that we are downloading
-PSInfraOrg=EpicGames
-PSInfraRepo=PixelStreamingInfrastructure
-
 # If a UE version is supplied set the right branch or tag to fetch for that version of UE
 if [ ! -z "$UEVersion" ]
 then
@@ -54,16 +66,27 @@ then
     PSInfraTagOrBranch=UE5.0
     IsTag=0
   fi
+  if [ "$UEVersion" = "5.1" ]
+  then
+    PSInfraTagOrBranch=UE5.1
+    IsTag=0
+  fi
+  if [ "$UEVersion" = "5.2" ]
+  then
+    PSInfraTagOrBranch=UE5.2
+    IsTag=0
+  fi
 fi
 
 # If no arguments select a specific version, fetch the appropriate default
 if [ -z "$PSInfraTagOrBranch" ]
 then
-  PSInfraTagOrBranch=UE5.1
+  PSInfraTagOrBranch=UE5.2
   IsTag=0
 fi
+echo "Tag or branch: $PSInfraTagOrBranch"
 
-# Whether the named reference is a tag or a branch affects the URL we fetch it on
+# Whether the named reference is a tag or a branch affects the Url we fetch it on
 if [ "$IsTag" -eq 1 ]
 then
   RefType=tags
@@ -71,25 +94,81 @@ else
   RefType=heads
 fi
 
-# Look for a SignallingWebServer directory next to this script
-if [ -d SignallingWebServer ]
+# We have a branch, no user-specified release, then check repo for the presence of a RELEASE_VERSION file in the current branch.
+if [ "$IsTag" -eq 0 ] && [ -z "$ReleaseUrl" ] && [ -z "$FlagPassed" ]
 then
-  echo "SignallingWebServer directory found...skipping install."
+  RelUrl=https://raw.githubusercontent.com/EpicGames/PixelStreamingInfrastructure/$PSInfraTagOrBranch/RELEASE_VERSION
+  if curl --output /dev/null --silent -r 0-0 --fail "$RelUrl"; then
+    ReleaseVersion="$PSInfraTagOrBranch-$(curl -L -s $RelUrl)"
+    ReleaseUrl=https://github.com/EpicGames/PixelStreamingInfrastructure/releases/download/$ReleaseVersion/$ReleaseVersion.tar.gz
+    echo "Valid RELEASE_VERSION file found in Github repo at $RelUrl"
+  else
+    echo "RELEASE_VERSION file does not exist at: $RelUrl"
+  fi
 else
-  echo "SignallingWebServer directory not found...beginning ps-infra download."
+  echo "Skipping downloading RELEASE_VERSION file."
+fi
+
+#Set our DownloadVersion here as we use this to check the contents of our DOWNLOAD_VERSION file shortly.
+DownloadVersion="$PSInfraTagOrBranch"
+if [ ! -z "$ReleaseVersion" ] 
+then
+  DownloadVersion="$ReleaseVersion"
+  echo "Release: $ReleaseVersion"
+fi
+
+#Rem Check for the existence of a DOWNLOAD_VERSION file and if found, check its contents against our $DownloadVersion
+if test -f DOWNLOAD_VERSION;
+then
+  PREVIOUS_DOWNLOAD_VERSION=$(cat DOWNLOAD_VERSION)
+  if [ "$DownloadVersion" = "$PREVIOUS_DOWNLOAD_VERSION" ]
+  then
+    echo "Downloaded version ($DownloadVersion) of PS infra matches release version ($PREVIOUS_DOWNLOAD_VERSION)...skipping install."
+    exit 0
+  else
+    echo "There is a newer released version ($DownloadVersion) - had ($PREVIOUS_DOWNLOAD_VERSION), downloading..."
+    #Remove old infra
+    rm -rf Frontend
+    rm -rf Matchmaker
+    rm -rf SignallingWebServer
+    rm -rf SFU
+  fi
+else 
+  echo "DOWNLOAD_VERSION file not found..."
+fi
+
+# Pre-download - Set the download url to the .zip of the branch
+Url=https://github.com/$PSInfraOrg/$PSInfraRepo/archive/refs/$RefType/$PSInfraTagOrBranch.tar.gz
+
+#Check if ReleaseUrl is valid by CURLing with fail fast and if success then set it to our download url
+if [ ! -z "$ReleaseUrl" ]
+then
+  echo "Checking if release url is valid, url: $ReleaseUrl"
+  if curl --output /dev/null --silent -r 0-0 --fail "$ReleaseUrl"; then
+    echo "Valid release url at: $ReleaseUrl"
+    Url=$ReleaseUrl
+  else
+    echo "Invalid Github release url: $ReleaseUrl"
+    exit 1
+  fi 
+fi
+
+# Download - download ps-infra and follow redirects.
+echo "Beginning ps-infra download from: $Url"
+curl -L $Url > ps-infra.tar.gz
 
-  # Download ps-infra and follow redirects.
-  curl -L https://github.com/$PSInfraOrg/$PSInfraRepo/archive/refs/$RefType/$PSInfraTagOrBranch.tar.gz > ps-infra.tar.gz
+# Unarchive the .tar
+tar -xmf ps-infra.tar.gz || $(echo "bad archive, contents:" && head --lines=20 ps-infra.tar.gz && exit 0)
 
-  # Unarchive the .tar
-  tar -xmf ps-infra.tar.gz || $(echo "bad archive, contents:" && head --lines=20 ps-infra.tar.gz && exit 0)
+# Move the server folders into the current directory (WebServers) and delete the original directory
+mv PixelStreamingInfrastructure-*/* .
+# Copy any files and folders beginning with dot (ignored by * glob) and discard errors regarding to not being able to move "." and ".."
+mv PixelStreamingInfrastructure-*/.* . 2>/dev/null || :
+rm -rf PixelStreamingInfrastructure-*
 
-  # Move the server folders into the current directory (WebServers) and delete the original directory
-  mv PixelStreamingInfrastructure-*/* .
-  # Copy any files and folders beginning with dot (ignored by * glob) and discard errors regarding to not being able to move "." and ".."
-  mv PixelStreamingInfrastructure-*/.* . 2>/dev/null || :
-  rm -rf PixelStreamingInfrastructure-*
+# Delete the downloaded tar
+rm ps-infra.tar.gz
 
-  # Delete the downloaded tar
-  rm ps-infra.tar.gz
-fi
\ No newline at end of file
+#Create a DOWNLOAD_VERSION file, which we use as a comparison file to check if we should auto upgrade when these scripts are run again
+echo "$DownloadVersion" >| DOWNLOAD_VERSION
+exit 0
\ No newline at end of file
diff --git a/Engine/Plugins/Media/PixelStreaming/Source/PixelStreamingEditor/Private/PixelStreamingEditorModule.cpp b/Engine/Plugins/Media/PixelStreaming/Source/PixelStreamingEditor/Private/PixelStreamingEditorModule.cpp
--- a/Engine/Plugins/Media/PixelStreaming/Source/PixelStreamingEditor/Private/PixelStreamingEditorModule.cpp
+++ b/Engine/Plugins/Media/PixelStreaming/Source/PixelStreamingEditor/Private/PixelStreamingEditorModule.cpp
@@ -222,7 +222,8 @@ void FPixelStreamingEditorModule::StartSignalling()
 	// but only attempt this is we haven't already started a download before.
 	if (!DownloadProcess.IsValid())
 	{
-		DownloadProcess = UE::PixelStreamingServers::DownloadPixelStreamingServers(/*bSkipIfPresent*/ true);
+		// We set bSkipIfPresent to false, which means the get_ps_servers script will always be run, that script will choose whether to download or not
+		DownloadProcess = UE::PixelStreamingServers::DownloadPixelStreamingServers(/*bSkipIfPresent*/ false);
 		if (DownloadProcess.IsValid())
 		{
 			DownloadProcess->OnCompleted().BindLambda([this](int ExitCode) {
