Fix: Missing copyright notice

diff --git a/Engine/Plugins/Media/PixelStreaming/Source/PixelStreamingInput/Private/PixelStreamingInputDevice.cpp b/Engine/Plugins/Media/PixelStreaming/Source/PixelStreamingInput/Private/PixelStreamingInputDevice.cpp
--- a/Engine/Plugins/Media/PixelStreaming/Source/PixelStreamingInput/Private/PixelStreamingInputDevice.cpp
+++ b/Engine/Plugins/Media/PixelStreaming/Source/PixelStreamingInput/Private/PixelStreamingInputDevice.cpp
@@ -1,4 +1,4 @@
-//
+// Copyright Epic Games, Inc. All Rights Reserved.
 
 #include "PixelStreamingInputDevice.h"
 
