Fix: Non-Unity compilation error

diff --git a/Engine/Plugins/Media/PixelStreaming/Source/PixelStreamingInput/Public/PixelStreamingInputProtocolMap.h b/Engine/Plugins/Media/PixelStreaming/Source/PixelStreamingInput/Public/PixelStreamingInputProtocolMap.h
--- a/Engine/Plugins/Media/PixelStreaming/Source/PixelStreamingInput/Public/PixelStreamingInputProtocolMap.h
+++ b/Engine/Plugins/Media/PixelStreaming/Source/PixelStreamingInput/Public/PixelStreamingInputProtocolMap.h
@@ -4,6 +4,7 @@
 
 #include "Containers/Map.h"
 #include "Misc/ScopeLock.h"
+#include "PixelStreamingInputMessage.h"
 
 /**
  * @brief An map type that broadcasts the OnProtocolUpdated whenever
