Fix: Embedded signalling server not able to connected to frontend.

- The fix was to remove autosubscribe from the embedded signalling server C++ code as this was moved to the frontend.

diff --git a/Engine/Plugins/Media/PixelStreaming/Source/PixelStreamingServers/Private/SignallingServer.cpp b/Engine/Plugins/Media/PixelStreaming/Source/PixelStreamingServers/Private/SignallingServer.cpp
--- a/Engine/Plugins/Media/PixelStreaming/Source/PixelStreamingServers/Private/SignallingServer.cpp
+++ b/Engine/Plugins/Media/PixelStreaming/Source/PixelStreamingServers/Private/SignallingServer.cpp
@@ -334,11 +334,6 @@ namespace UE::PixelStreamingServers
 		// Send config to newly connected player, which kicks off making a new RTC connection
 		TSharedPtr<FJsonObject> ConfigJSON = CreateConfigJSON();
 		SendPlayerMessage(ConnectionId, ConfigJSON);
-
-		if (StreamersWS->Count() == 1)
-		{
-			SubscribePlayer(ConnectionId, StreamersWS->GetConnectionNames()[0]);
-		}
 	}
 
 	void FSignallingServer::OnPlayerDisconnected(uint16 ConnectionId)
