Pixel streaming: Add entry points to set the back buffer as external video source input

diff --git a/Engine/Plugins/Media/PixelStreaming/Source/PixelStreaming/Private/PixelStreamingModule.cpp b/Engine/Plugins/Media/PixelStreaming/Source/PixelStreaming/Private/PixelStreamingModule.cpp
--- a/Engine/Plugins/Media/PixelStreaming/Source/PixelStreaming/Private/PixelStreamingModule.cpp
+++ b/Engine/Plugins/Media/PixelStreaming/Source/PixelStreaming/Private/PixelStreamingModule.cpp
@@ -318,6 +318,17 @@ namespace UE::PixelStreaming
 		ExternalVideoSourceGroup->SetFPS(InFPS);
 	}
 
+	void FPixelStreamingModule::SetExternalVideoSourceCoupleFramerate(bool bShouldCoupleFPS)
+	{
+		ExternalVideoSourceGroup->SetCoupleFramerate(bShouldCoupleFPS);
+	}
+
+	void FPixelStreamingModule::SetExternalVideoSourceInput(TSharedPtr<FPixelStreamingVideoInput> InVideoInput)
+	{
+		ExternalVideoSourceGroup->SetVideoInput(InVideoInput);
+		ExternalVideoSourceGroup->Start();
+	}
+
 	rtc::scoped_refptr<webrtc::VideoTrackSourceInterface> FPixelStreamingModule::CreateExternalVideoSource()
 	{
 		return ExternalVideoSourceGroup->CreateVideoSource([]() { return true; });
diff --git a/Engine/Plugins/Media/PixelStreaming/Source/PixelStreaming/Private/PixelStreamingModule.h b/Engine/Plugins/Media/PixelStreaming/Source/PixelStreaming/Private/PixelStreamingModule.h
--- a/Engine/Plugins/Media/PixelStreaming/Source/PixelStreaming/Private/PixelStreamingModule.h
+++ b/Engine/Plugins/Media/PixelStreaming/Source/PixelStreaming/Private/PixelStreamingModule.h
@@ -45,6 +45,8 @@ namespace UE::PixelStreaming
 		virtual const TArray<UPixelStreamingInput*> GetInputComponents() override;
 		// Don't delete, is used externally to PS
 		virtual void SetExternalVideoSourceFPS(uint32 InFPS) override;
+		virtual void SetExternalVideoSourceCoupleFramerate(bool bShouldCoupleFPS) override;
+		virtual void SetExternalVideoSourceInput(TSharedPtr<FPixelStreamingVideoInput> InVideoInput) override;
 		virtual rtc::scoped_refptr<webrtc::VideoTrackSourceInterface> CreateExternalVideoSource() override;
 		virtual void ReleaseExternalVideoSource(const webrtc::VideoTrackSourceInterface* InVideoSource) override;
 		virtual TUniquePtr<webrtc::VideoEncoderFactory> CreateVideoEncoderFactory() override;
diff --git a/Engine/Plugins/Media/PixelStreaming/Source/PixelStreaming/Public/IPixelStreamingModule.h b/Engine/Plugins/Media/PixelStreaming/Source/PixelStreaming/Public/IPixelStreamingModule.h
--- a/Engine/Plugins/Media/PixelStreaming/Source/PixelStreaming/Public/IPixelStreamingModule.h
+++ b/Engine/Plugins/Media/PixelStreaming/Source/PixelStreaming/Public/IPixelStreamingModule.h
@@ -145,6 +145,17 @@ public:
 	 */
 	virtual void SetExternalVideoSourceFPS(uint32 InFPS) = 0;
 
+	/**
+	 * Control coupling the external video source framerate to the video input framerate
+	 * @param bShouldCoupleFPS Flag for coupling the video source and video input framerates
+	 */
+	virtual void SetExternalVideoSourceCoupleFramerate(bool bShouldCoupleFPS) = 0;
+
+	/**
+	 * Hook to set the input to the external video source. Will also call Start() on the external video source
+	 */
+	virtual void SetExternalVideoSourceInput(TSharedPtr<FPixelStreamingVideoInput> InVideoInput) = 0;
+
 	/**
 	 * Allows the creation of Video Tracks that are fed the backbuffer
 	 */
