Fix: VR streamed eyes are swapped

diff --git a/Engine/Plugins/Media/PixelStreaming/Source/PixelStreamingHMD/Private/PixelStreamingHMD.cpp b/Engine/Plugins/Media/PixelStreaming/Source/PixelStreamingHMD/Private/PixelStreamingHMD.cpp
--- a/Engine/Plugins/Media/PixelStreaming/Source/PixelStreamingHMD/Private/PixelStreamingHMD.cpp
+++ b/Engine/Plugins/Media/PixelStreaming/Source/PixelStreamingHMD/Private/PixelStreamingHMD.cpp
@@ -172,7 +172,7 @@ void FPixelStreamingHMD::CalculateStereoViewOffset(const int32 ViewIndex, FRotat
 	if (ViewIndex != INDEX_NONE)
 	{
 		float EyeOffset = 3.20000005f;
-		const float PassOffset = (ViewIndex == EStereoscopicEye::eSSE_LEFT_EYE) ? EyeOffset : -EyeOffset;
+		const float PassOffset = (ViewIndex == EStereoscopicEye::eSSE_LEFT_EYE) ? -EyeOffset : EyeOffset;
 		ViewLocation += ViewRotation.Quaternion().RotateVector(FVector(0, PassOffset, 0));
 	}
 }
