Removing Private include paths from a number of modules. Their Private dependencies were cleaned up in previous CLs.

diff --git a/Engine/Plugins/Animation/ControlRig/Source/ControlRigEditor/ControlRigEditor.Build.cs b/Engine/Plugins/Animation/ControlRig/Source/ControlRigEditor/ControlRigEditor.Build.cs
--- a/Engine/Plugins/Animation/ControlRig/Source/ControlRigEditor/ControlRigEditor.Build.cs
+++ b/Engine/Plugins/Animation/ControlRig/Source/ControlRigEditor/ControlRigEditor.Build.cs
@@ -8,15 +8,7 @@ namespace UnrealBuildTool.Rules
         {
 			PrivateIncludePaths.AddRange(
                 new string[] {
-					System.IO.Path.Combine(GetModuleDirectory("AssetTools"), "Private"),
 					System.IO.Path.Combine(GetModuleDirectory("ControlRig"), "Private"),
-					System.IO.Path.Combine(GetModuleDirectory("Kismet"), "Private"),
-					System.IO.Path.Combine(GetModuleDirectory("MessageLog"), "Private"), //compatibility for FBX importer
-					System.IO.Path.Combine(GetModuleDirectory("Persona"), "Private"),
-					System.IO.Path.Combine(GetModuleDirectory("PropertyEditor"), "Private"),
-					System.IO.Path.Combine(GetModuleDirectory("SceneOutliner"), "Private"),
-					System.IO.Path.Combine(GetModuleDirectory("Slate"), "Private"),
-					System.IO.Path.Combine(GetModuleDirectory("UnrealEd"), "Private"),
 				}
 			);
 
diff --git a/Engine/Plugins/Media/PixelStreaming/Source/PixelStreaming/PixelStreaming.Build.cs b/Engine/Plugins/Media/PixelStreaming/Source/PixelStreaming/PixelStreaming.Build.cs
--- a/Engine/Plugins/Media/PixelStreaming/Source/PixelStreaming/PixelStreaming.Build.cs
+++ b/Engine/Plugins/Media/PixelStreaming/Source/PixelStreaming/PixelStreaming.Build.cs
@@ -1,10 +1,7 @@
 // Copyright Epic Games, Inc. All Rights Reserved.
 
-using System;
-using System.IO;
 using System.Collections.Generic;
-using EpicGames.Core;
-
+using System.IO;
 
 namespace UnrealBuildTool.Rules
 {
@@ -57,8 +54,6 @@ namespace UnrealBuildTool.Rules
 
 		public PixelStreaming(ReadOnlyTargetRules Target) : base(Target)
 		{
-			var EngineDir = Path.GetFullPath(Target.RelativeEnginePath);
-
 			// This is so for game projects using our public headers don't have to include extra modules they might not know about.
 			PublicDependencyModuleNames.AddRange(new string[]
 			{
@@ -69,12 +64,6 @@ namespace UnrealBuildTool.Rules
 				"PixelStreamingInput"
 			});
 
-			// NOTE: General rule is not to access the private folder of another module
-			PrivateIncludePaths.AddRange(new string[]
-			{
-				Path.Combine(GetModuleDirectory("AudioMixer"), "Private"),
-			});
-
 			PrivateDependencyModuleNames.AddRange(new string[]
 			{
 				"Core",
diff --git a/Engine/Plugins/Runtime/GameplayAbilities/Source/GameplayAbilitiesEditor/GameplayAbilitiesEditor.Build.cs b/Engine/Plugins/Runtime/GameplayAbilities/Source/GameplayAbilitiesEditor/GameplayAbilitiesEditor.Build.cs
--- a/Engine/Plugins/Runtime/GameplayAbilities/Source/GameplayAbilitiesEditor/GameplayAbilitiesEditor.Build.cs
+++ b/Engine/Plugins/Runtime/GameplayAbilities/Source/GameplayAbilitiesEditor/GameplayAbilitiesEditor.Build.cs
@@ -1,21 +1,11 @@
 // Copyright Epic Games, Inc. All Rights Reserved.
 
-using UnrealBuildTool;
-using System.IO;
-
 namespace UnrealBuildTool.Rules
 {
 	public class GameplayAbilitiesEditor : ModuleRules
 	{
 		public GameplayAbilitiesEditor(ReadOnlyTargetRules Target) : base(Target)
 		{
-			PrivateIncludePaths.AddRange(
-				new string[] {
-					Path.Combine(GetModuleDirectory("AssetTools"), "Private"),
-					Path.Combine(GetModuleDirectory("GameplayTagsEditor"), "Private"),
-					Path.Combine(GetModuleDirectory("Kismet"), "Private"),
-				});
-
 			PublicDependencyModuleNames.Add("GameplayTasks");
 
 			PrivateDependencyModuleNames.AddRange(
diff --git a/Engine/Source/Developer/Android/AndroidDeviceDetection/AndroidDeviceDetection.Build.cs b/Engine/Source/Developer/Android/AndroidDeviceDetection/AndroidDeviceDetection.Build.cs
--- a/Engine/Source/Developer/Android/AndroidDeviceDetection/AndroidDeviceDetection.Build.cs
+++ b/Engine/Source/Developer/Android/AndroidDeviceDetection/AndroidDeviceDetection.Build.cs
@@ -25,13 +25,6 @@ public class AndroidDeviceDetection : ModuleRules
 			}
 		);
 
-		PublicIncludePaths.AddRange(
-			new string[]
-			{
-				"Runtime/Core/Public/Android"
-			}
-		);
-
         if (Target.bCompileAgainstEngine)
 		{
 			PrivateDependencyModuleNames.Add("Engine");
diff --git a/Engine/Source/Developer/Android/AndroidTargetPlatform/AndroidTargetPlatform.Build.cs b/Engine/Source/Developer/Android/AndroidTargetPlatform/AndroidTargetPlatform.Build.cs
--- a/Engine/Source/Developer/Android/AndroidTargetPlatform/AndroidTargetPlatform.Build.cs
+++ b/Engine/Source/Developer/Android/AndroidTargetPlatform/AndroidTargetPlatform.Build.cs
@@ -21,14 +21,6 @@ public class AndroidTargetPlatform : ModuleRules
             }
 		);
 
-		PublicIncludePaths.AddRange(
-			new string[]
-			{
-				"Runtime/Core/Public/Android"
-			}
-		);
-
-
         if (Target.bCompileAgainstEngine)
 		{
 			PrivateDependencyModuleNames.Add("Engine");
diff --git a/Engine/Source/Editor/GameplayTasksEditor/GameplayTasksEditor.Build.cs b/Engine/Source/Editor/GameplayTasksEditor/GameplayTasksEditor.Build.cs
--- a/Engine/Source/Editor/GameplayTasksEditor/GameplayTasksEditor.Build.cs
+++ b/Engine/Source/Editor/GameplayTasksEditor/GameplayTasksEditor.Build.cs
@@ -8,13 +8,6 @@ namespace UnrealBuildTool.Rules
 		{
 			OverridePackageType = PackageOverrideType.EngineDeveloper;
 
-			PrivateIncludePaths.AddRange(
-                new string[] {
-					System.IO.Path.Combine(GetModuleDirectory("AssetTools"), "Private"),
-					System.IO.Path.Combine(GetModuleDirectory("Kismet"), "Private"),
-				}
-			);
-
             PrivateDependencyModuleNames.AddRange(
                 new string[]
 				{
