Fix PixelStreaming application being unable to scale mouse input due to target window being nullptr at load. Also fix touching the debug console when editor streaming causing nullptr exception.

diff --git a/Engine/Plugins/Media/PixelStreaming/Source/PixelStreaming/Private/PixelStreamingModule.cpp b/Engine/Plugins/Media/PixelStreaming/Source/PixelStreaming/Private/PixelStreamingModule.cpp
--- a/Engine/Plugins/Media/PixelStreaming/Source/PixelStreaming/Private/PixelStreamingModule.cpp
+++ b/Engine/Plugins/Media/PixelStreaming/Source/PixelStreaming/Private/PixelStreamingModule.cpp
@@ -400,15 +400,14 @@ namespace UE::PixelStreaming
 			// default to the scene viewport if we have a game engine
 			if (UGameEngine* GameEngine = Cast<UGameEngine>(GEngine))
 			{
-				TSharedPtr<FSceneViewport> TargetViewport = GameEngine->SceneViewport;
-				if (TargetViewport.IsValid())
+				TSharedPtr<SWindow> TargetWindow = GameEngine->GameViewport->GetWindow();
+				if (TargetWindow.IsValid())
 				{
-					DefaultStreamer->SetTargetViewport(TargetViewport->GetViewportWidget());
-					DefaultStreamer->SetTargetWindow(TargetViewport->FindWindow());
+					DefaultStreamer->SetTargetWindow(TargetWindow);
 				}
 				else
 				{
-					UE_LOG(LogPixelStreaming, Error, TEXT("Cannot set target viewport/window - target viewport is not valid."));
+					UE_LOG(LogPixelStreaming, Error, TEXT("Cannot set target window - target window is not valid."));
 				}
 			}
 		}
diff --git a/Engine/Plugins/Media/PixelStreaming/Source/PixelStreamingInput/Private/PixelStreamingInputHandler.cpp b/Engine/Plugins/Media/PixelStreaming/Source/PixelStreamingInput/Private/PixelStreamingInputHandler.cpp
--- a/Engine/Plugins/Media/PixelStreaming/Source/PixelStreamingInput/Private/PixelStreamingInputHandler.cpp
+++ b/Engine/Plugins/Media/PixelStreaming/Source/PixelStreamingInput/Private/PixelStreamingInputHandler.cpp
@@ -1145,7 +1145,9 @@ namespace UE::PixelStreamingInput
 					FVector2D WindowClientOffset = ArrangedWidget.Geometry.GetAbsolutePosition();
 					FVector2D WindowClientSize = ArrangedWidget.Geometry.GetAbsoluteSize();
 
-					FVector2D OutTemp = bIncludeOffset ? WindowOrigin + WindowClientOffset + (ScreenLocation * WindowClientSize) : (ScreenLocation * WindowClientSize);
+					FVector2D OutTemp = bIncludeOffset
+						? (ScreenLocation * WindowClientSize) + WindowOrigin + WindowClientOffset
+						: (ScreenLocation * WindowClientSize);
 					UE_LOG(LogPixelStreamingInputHandler, Verbose, TEXT("%.4f, %.4f"), ScreenLocation.X, ScreenLocation.Y);
 					OutVector = FIntPoint((int32)OutTemp.X, (int32)OutTemp.Y);
 				}
@@ -1153,7 +1155,9 @@ namespace UE::PixelStreamingInput
 			else
 			{
 				FVector2D SizeInScreen = ApplicationWindow->GetSizeInScreen();
-				FVector2D OutTemp = SizeInScreen * ScreenLocation;
+				FVector2D OutTemp = bIncludeOffset
+					? (SizeInScreen * ScreenLocation) + ApplicationWindow->GetPositionInScreen()
+					: (SizeInScreen * ScreenLocation);
 				OutVector = FIntPoint((int32)OutTemp.X, (int32)OutTemp.Y);
 			}
 		}
@@ -1318,7 +1322,7 @@ namespace UE::PixelStreamingInput
 				}
 				else if (FocusedWidget->GetType() == TEXT("SMultiLineEditableText"))
 				{
-					SMultiLineEditableTextBox* TextBox = static_cast<SMultiLineEditableTextBox*>(FocusedWidget.Get());
+					SMultiLineEditableText* TextBox = static_cast<SMultiLineEditableText*>(FocusedWidget.Get());
 					TextboxContents = TextBox->GetText();
 				}
 
