Fix: HMD image distorts on head roll

diff --git a/Engine/Plugins/Media/PixelStreaming/Source/PixelStreamingHMD/PixelStreamingHMD.build.cs b/Engine/Plugins/Media/PixelStreaming/Source/PixelStreamingHMD/PixelStreamingHMD.build.cs
--- a/Engine/Plugins/Media/PixelStreaming/Source/PixelStreamingHMD/PixelStreamingHMD.build.cs
+++ b/Engine/Plugins/Media/PixelStreaming/Source/PixelStreamingHMD/PixelStreamingHMD.build.cs
@@ -23,6 +23,7 @@ namespace UnrealBuildTool.Rules
 					"RenderCore",
 					"Renderer",
 					"HeadMountedDisplay",
+					"SlateCore"
 				}
 			);
 		}
diff --git a/Engine/Plugins/Media/PixelStreaming/Source/PixelStreamingHMD/Private/PixelStreamingHMD.cpp b/Engine/Plugins/Media/PixelStreaming/Source/PixelStreamingHMD/Private/PixelStreamingHMD.cpp
--- a/Engine/Plugins/Media/PixelStreaming/Source/PixelStreamingHMD/Private/PixelStreamingHMD.cpp
+++ b/Engine/Plugins/Media/PixelStreaming/Source/PixelStreamingHMD/Private/PixelStreamingHMD.cpp
@@ -11,6 +11,7 @@
 #include "PostProcess/PostProcessHMD.h"
 #include "GameFramework/WorldSettings.h"
 #include "Settings.h"
+#include "Widgets/SWindow.h"
 
 FPixelStreamingHMD::FPixelStreamingHMD(const FAutoRegister& AutoRegister)
 	: FHeadMountedDisplayBase(nullptr)
@@ -183,8 +184,10 @@ FMatrix FPixelStreamingHMD::GetStereoProjectionMatrix(const int32 ViewIndex) con
 	const float PassProjectionOffset = (ViewIndex == EStereoscopicEye::eSSE_LEFT_EYE) ? ProjectionCenterOffset : -ProjectionCenterOffset;
 
 	const float HalfFov = 2.19686294f / 2.f;
-	const float InWidth = 640.f;
-	const float InHeight = 480.f;
+	TSharedPtr<SWindow> TargetWindow = GEngine->GameViewport->GetWindow();
+	FVector2f SizeInScreen = TargetWindow->GetSizeInScreen();
+	const float InWidth = SizeInScreen.X / 2.f;
+	const float InHeight = SizeInScreen.Y;
 	const float XS = 1.0f / tan(HalfFov);
 	const float YS = InWidth / tan(HalfFov) / InHeight;
 
