Fix non-unity
#rnx

diff --git a/Engine/Plugins/Media/PixelStreaming/Source/PixelStreamingHMD/Private/PixelStreamingHMD.cpp b/Engine/Plugins/Media/PixelStreaming/Source/PixelStreamingHMD/Private/PixelStreamingHMD.cpp
--- a/Engine/Plugins/Media/PixelStreaming/Source/PixelStreamingHMD/Private/PixelStreamingHMD.cpp
+++ b/Engine/Plugins/Media/PixelStreaming/Source/PixelStreamingHMD/Private/PixelStreamingHMD.cpp
@@ -5,6 +5,7 @@
 #include "Modules/ModuleManager.h"
 #include "EngineGlobals.h"
 #include "Engine/Engine.h"
+#include "Engine/GameViewportClient.h"
 #include "GameFramework/PlayerController.h"
 #include "Engine/LocalPlayer.h"
 #include "SceneRendering.h"
