Fixing level editor streaming with multi streamers.
- The code was previously trying to get the "default streamer" which doesn't exist in the editor mode. Other editor streaming code had been adjusted to use the streamer named "Editor" so this code was updated to match.

diff --git a/Engine/Plugins/Media/PixelStreaming/Source/PixelStreamingEditor/Private/PixelStreamingEditorModule.cpp b/Engine/Plugins/Media/PixelStreaming/Source/PixelStreamingEditor/Private/PixelStreamingEditorModule.cpp
--- a/Engine/Plugins/Media/PixelStreaming/Source/PixelStreamingEditor/Private/PixelStreamingEditorModule.cpp
+++ b/Engine/Plugins/Media/PixelStreaming/Source/PixelStreamingEditor/Private/PixelStreamingEditorModule.cpp
@@ -145,7 +145,7 @@ void FPixelStreamingEditorModule::StartStreaming(UE::EditorPixelStreaming::EStre
 			EditorStreamer->SetTargetViewport(SceneViewport->GetViewportWidget());
 			EditorStreamer->SetTargetWindow(SceneViewport->FindWindow());
 			EditorStreamer->SetInputHandlerType(EPixelStreamingInputType::RouteToWindow);
-			EditorStreamer->SetVideoInput(FPixelStreamingVideoInputViewport::Create());
+			EditorStreamer->SetVideoInput(FPixelStreamingVideoInputViewport::Create(EditorStreamer));
 		}
 		break;
 		case UE::EditorPixelStreaming::EStreamTypes::Editor:
diff --git a/Engine/Plugins/Media/PixelStreaming/Source/PixelStreamingEditor/Private/PixelStreamingVideoInputViewport.cpp b/Engine/Plugins/Media/PixelStreaming/Source/PixelStreamingEditor/Private/PixelStreamingVideoInputViewport.cpp
--- a/Engine/Plugins/Media/PixelStreaming/Source/PixelStreamingEditor/Private/PixelStreamingVideoInputViewport.cpp
+++ b/Engine/Plugins/Media/PixelStreaming/Source/PixelStreamingEditor/Private/PixelStreamingVideoInputViewport.cpp
@@ -15,9 +15,10 @@
 #include "Slate/SceneViewport.h"
 #include "IPixelStreamingModule.h"
 
-TSharedPtr<FPixelStreamingVideoInputViewport> FPixelStreamingVideoInputViewport::Create()
+TSharedPtr<FPixelStreamingVideoInputViewport> FPixelStreamingVideoInputViewport::Create(TSharedPtr<IPixelStreamingStreamer> InAssociatedStreamer)
 {
 	TSharedPtr<FPixelStreamingVideoInputViewport> NewInput = TSharedPtr<FPixelStreamingVideoInputViewport>(new FPixelStreamingVideoInputViewport());
+	NewInput->AssociatedStreamer = InAssociatedStreamer;
 	TWeakPtr<FPixelStreamingVideoInputViewport> WeakInput = NewInput;
 
 	// Set up the callback on the game thread since FSlateApplication::Get() can only be used there
@@ -41,36 +42,47 @@ FPixelStreamingVideoInputViewport::~FPixelStreamingVideoInputViewport()
 	}
 }
 
-void FPixelStreamingVideoInputViewport::OnViewportRendered(FViewport* InViewport)
+bool FPixelStreamingVideoInputViewport::FilterViewport(const FViewport* InViewport)
 {
-	IPixelStreamingModule& Module = IPixelStreamingModule::Get();
-	TSharedPtr<IPixelStreamingStreamer> Streamer = Module.GetStreamer(Module.GetDefaultStreamerID());
+	TSharedPtr<IPixelStreamingStreamer> Streamer = AssociatedStreamer.Pin();
 	if (!Streamer.IsValid() || !Streamer->IsStreaming())
 	{
-		return;
+		return false;
 	}
 
 	TSharedPtr<SViewport> TargetScene = Streamer->GetTargetViewport().Pin();
 	if (!TargetScene.IsValid())
 	{
-		return;
+		return false;
 	}
 
 	if (InViewport == nullptr)
 	{
-		return;
+		return false;
 	}
 
 	if (InViewport->GetViewportType() != TargetViewportType)
 	{
-		return;
+		return false;
 	}
 
 	// Bit dirty to do a static cast here, but we check viewport type just above so it is somewhat "safe".
-	TSharedPtr<SViewport> InScene = static_cast<FSceneViewport*>(InViewport)->GetViewportWidget().Pin();
+	TSharedPtr<SViewport> InScene = StaticCast<const FSceneViewport*>(InViewport)->GetViewportWidget().Pin();
 
 	// If the viewport we were passed is not our target viewport we are not interested in getting its texture.
 	if (TargetScene != InScene)
+	{
+		return false;
+	}
+
+
+
+	return true;
+}
+
+void FPixelStreamingVideoInputViewport::OnViewportRendered(FViewport* InViewport)
+{
+	if (!FilterViewport(InViewport))
 	{
 		return;
 	}
diff --git a/Engine/Plugins/Media/PixelStreaming/Source/PixelStreamingEditor/Public/PixelStreamingVideoInputViewport.h b/Engine/Plugins/Media/PixelStreaming/Source/PixelStreamingEditor/Public/PixelStreamingVideoInputViewport.h
--- a/Engine/Plugins/Media/PixelStreaming/Source/PixelStreamingEditor/Public/PixelStreamingVideoInputViewport.h
+++ b/Engine/Plugins/Media/PixelStreaming/Source/PixelStreamingEditor/Public/PixelStreamingVideoInputViewport.h
@@ -3,6 +3,7 @@
 #pragma once
 
 #include "PixelStreamingVideoInputRHI.h"
+#include "IPixelStreamingStreamer.h"
 #include "Delegates/IDelegateInstance.h"
 #include "UnrealClient.h"
 
@@ -12,7 +13,7 @@
 class PIXELSTREAMINGEDITOR_API FPixelStreamingVideoInputViewport : public FPixelStreamingVideoInputRHI
 {
 public:
-	static TSharedPtr<FPixelStreamingVideoInputViewport> Create();
+	static TSharedPtr<FPixelStreamingVideoInputViewport> Create(TSharedPtr<IPixelStreamingStreamer> InAssociatedStreamer);
 	virtual ~FPixelStreamingVideoInputViewport();
 
 	virtual FString ToString() override;
@@ -20,9 +21,12 @@ public:
 private:
 	FPixelStreamingVideoInputViewport() = default;
 
+	bool FilterViewport(const FViewport* InViewport);
 	void OnViewportRendered(FViewport* Viewport);
 
 	FDelegateHandle DelegateHandle;
 
 	FName TargetViewportType = FName(FString(TEXT("SceneViewport")));
+
+	TWeakPtr<IPixelStreamingStreamer> AssociatedStreamer;
 };
