Non-unity fix
#rnx

diff --git a/Engine/Plugins/Media/PixelStreaming/Source/PixelStreaming/Private/PixelStreamingModule.cpp b/Engine/Plugins/Media/PixelStreaming/Source/PixelStreaming/Private/PixelStreamingModule.cpp
--- a/Engine/Plugins/Media/PixelStreaming/Source/PixelStreaming/Private/PixelStreamingModule.cpp
+++ b/Engine/Plugins/Media/PixelStreaming/Source/PixelStreaming/Private/PixelStreamingModule.cpp
@@ -10,6 +10,7 @@
 #include "PixelStreamingPrivate.h"
 #include "AudioSink.h"
 #include "CoreMinimal.h"
+#include "Engine/GameViewportClient.h"
 #include "Modules/ModuleManager.h"
 #include "UObject/UObjectIterator.h"
 #include "Engine/Texture2D.h"
