Fixing a crash on shutdown.
- Seems that if we don't release an FStreamer in the ShutdownModule call, it gets destroyed in engine destruction which causes a bunch of issues with things being deleted before we're done with them.

diff --git a/Engine/Plugins/Media/PixelStreaming/Source/PixelStreaming/Private/PixelStreamingModule.cpp b/Engine/Plugins/Media/PixelStreaming/Source/PixelStreaming/Private/PixelStreamingModule.cpp
--- a/Engine/Plugins/Media/PixelStreaming/Source/PixelStreaming/Private/PixelStreamingModule.cpp
+++ b/Engine/Plugins/Media/PixelStreaming/Source/PixelStreaming/Private/PixelStreamingModule.cpp
@@ -165,7 +165,10 @@ namespace UE::PixelStreaming
 		}
 
 		// We explicitly call release on streamer so WebRTC gets shutdown before our module is deleted
+		// additionally the streamer does a bunch of delegate calls and unbinds which seem to have issues
+		// when called during engine destruction rather than here.
 		Streamers.Empty();
+		DefaultStreamer.Reset();
 		ExternalVideoSourceGroup->Stop();
 
 		FPixelStreamingPeerConnection::Shutdown();
