Fixing a crash when -PixelStreamingURL does not provide a port.
- When no port is supplied the streamer will try to join on port 80 which means the signalling server will treat it as a player.
- This fix will detect when no port is provided and in such a case will manually add the default 8888 port to the url.

diff --git a/Engine/Plugins/Media/PixelStreaming/Source/PixelStreaming/PixelStreaming.Build.cs b/Engine/Plugins/Media/PixelStreaming/Source/PixelStreaming/PixelStreaming.Build.cs
--- a/Engine/Plugins/Media/PixelStreaming/Source/PixelStreaming/PixelStreaming.Build.cs
+++ b/Engine/Plugins/Media/PixelStreaming/Source/PixelStreaming/PixelStreaming.Build.cs
@@ -90,6 +90,7 @@ namespace UnrealBuildTool.Rules
 				"PixelStreamingServers",
 				"PixelStreamingHMD",
 				"TraceLog",
+				"HTTP",
 			});
 
 			PrivateDefinitions.Add("PIXELSTREAMING_DUMP_ENCODING=0");
diff --git a/Engine/Plugins/Media/PixelStreaming/Source/PixelStreaming/Private/PixelStreamingSignallingConnection.cpp b/Engine/Plugins/Media/PixelStreaming/Source/PixelStreaming/Private/PixelStreamingSignallingConnection.cpp
--- a/Engine/Plugins/Media/PixelStreaming/Source/PixelStreaming/Private/PixelStreamingSignallingConnection.cpp
+++ b/Engine/Plugins/Media/PixelStreaming/Source/PixelStreaming/Private/PixelStreamingSignallingConnection.cpp
@@ -11,6 +11,7 @@
 #include "PixelStreamingDelegates.h"
 #include "Utils.h"
 #include "ToStringExtensions.h"
+#include "GenericPlatform/GenericPlatformHttp.h"
 
 DECLARE_LOG_CATEGORY_EXTERN(LogPixelStreamingSS, Log, VeryVerbose);
 DEFINE_LOG_CATEGORY(LogPixelStreamingSS);
@@ -57,6 +58,18 @@ void FPixelStreamingSignallingConnection::Connect(FString InUrl, bool bIsReconne
 
 	Url = InUrl;
 
+	// add the default port if none is provided
+	TOptional<uint16> MaybePort = FGenericPlatformHttp::GetUrlPort(Url);
+	if (!MaybePort)
+	{
+		const FString Protocol = FGenericPlatformHttp::IsSecureProtocol(Url).Get(false) ? "wss" : "ws";
+		const FString Domain = FGenericPlatformHttp::GetUrlDomain(Url);
+		const uint16 Port = 8888;
+		const FString Path = FGenericPlatformHttp::GetUrlPath(Url);
+		const FString Final = FString::Printf(TEXT("%s://%s:%d%s"), *Protocol, *Domain, Port, *Path);
+		Url = Final;
+	}
+
 	WebSocket = FWebSocketsModule::Get().CreateWebSocket(Url, TEXT(""));
 	verifyf(WebSocket, TEXT("Web Socket Factory failed to return a valid Web Socket."));
 
