Fix: PixelStreaming module keeps strong references streamers, preventing them from being deleted

diff --git a/Engine/Plugins/Media/PixelStreaming/Source/PixelStreaming/Private/PixelStreamingAudioComponent.cpp b/Engine/Plugins/Media/PixelStreaming/Source/PixelStreaming/Private/PixelStreamingAudioComponent.cpp
--- a/Engine/Plugins/Media/PixelStreaming/Source/PixelStreaming/Private/PixelStreamingAudioComponent.cpp
+++ b/Engine/Plugins/Media/PixelStreaming/Source/PixelStreaming/Private/PixelStreamingAudioComponent.cpp
@@ -89,7 +89,7 @@ bool UPixelStreamingAudioComponent::StreamerListenTo(FString StreamerId, FString
 		StreamerToHear = StreamerId;
 	}
 
-	TSharedPtr<IPixelStreamingStreamer> Streamer = PixelStreamingModule.GetStreamer(StreamerToHear);
+	TSharedPtr<IPixelStreamingStreamer> Streamer = PixelStreamingModule.FindStreamer(StreamerToHear);
 	if (!Streamer)
 	{
 		return false;
diff --git a/Engine/Plugins/Media/PixelStreaming/Source/PixelStreaming/Private/PixelStreamingBlueprints.cpp b/Engine/Plugins/Media/PixelStreaming/Source/PixelStreaming/Private/PixelStreamingBlueprints.cpp
--- a/Engine/Plugins/Media/PixelStreaming/Source/PixelStreaming/Private/PixelStreamingBlueprints.cpp
+++ b/Engine/Plugins/Media/PixelStreaming/Source/PixelStreaming/Private/PixelStreamingBlueprints.cpp
@@ -11,12 +11,12 @@
 void UPixelStreamingBlueprints::SendFileAsByteArray(TArray<uint8> ByteArray, FString MimeType, FString FileExtension)
 {
 	IPixelStreamingModule* Module = UE::PixelStreaming::FPixelStreamingModule::GetModule();
-	if(!Module)
+	if (!Module)
 	{
 		return;
 	}
-	TSharedPtr<IPixelStreamingStreamer> Streamer = Module->GetStreamer(Module->GetDefaultStreamerID());
-	if(!Streamer)
+	TSharedPtr<IPixelStreamingStreamer> Streamer = Module->FindStreamer(Module->GetDefaultStreamerID());
+	if (!Streamer)
 	{
 		return;
 	}
@@ -27,12 +27,12 @@ void UPixelStreamingBlueprints::SendFileAsByteArray(TArray<uint8> ByteArray, FSt
 void UPixelStreamingBlueprints::StreamerSendFileAsByteArray(FString StreamerId, TArray<uint8> ByteArray, FString MimeType, FString FileExtension)
 {
 	IPixelStreamingModule* Module = UE::PixelStreaming::FPixelStreamingModule::GetModule();
-	if(!Module)
+	if (!Module)
 	{
 		return;
 	}
-	TSharedPtr<IPixelStreamingStreamer> Streamer = Module->GetStreamer(StreamerId);
-	if(!Streamer)
+	TSharedPtr<IPixelStreamingStreamer> Streamer = Module->FindStreamer(StreamerId);
+	if (!Streamer)
 	{
 		return;
 	}
@@ -43,62 +43,62 @@ void UPixelStreamingBlueprints::StreamerSendFileAsByteArray(FString StreamerId,
 void UPixelStreamingBlueprints::SendFile(FString FilePath, FString MimeType, FString FileExtension)
 {
 	IPixelStreamingModule* Module = UE::PixelStreaming::FPixelStreamingModule::GetModule();
-	if(!Module)
+	if (!Module)
 	{
 		return;
 	}
-	TSharedPtr<IPixelStreamingStreamer> Streamer = Module->GetStreamer(Module->GetDefaultStreamerID());
-	if(!Streamer)
+	TSharedPtr<IPixelStreamingStreamer> Streamer = Module->FindStreamer(Module->GetDefaultStreamerID());
+	if (!Streamer)
 	{
 		return;
 	}
-	
+
 	TArray64<uint8> ByteData;
 	bool bSuccess = FFileHelper::LoadFileToArray(ByteData, *FilePath);
-	if(bSuccess)
+	if (bSuccess)
 	{
 		Streamer->SendFileData(ByteData, MimeType, FileExtension);
 	}
 	else
 	{
 		UE_LOG(LogPixelStreaming, Error, TEXT("FileHelper failed to load file data"));
-	}	
+	}
 }
 
 void UPixelStreamingBlueprints::StreamerSendFile(FString StreamerId, FString FilePath, FString MimeType, FString FileExtension)
 {
 	IPixelStreamingModule* Module = UE::PixelStreaming::FPixelStreamingModule::GetModule();
-	if(!Module)
+	if (!Module)
 	{
 		return;
 	}
-	TSharedPtr<IPixelStreamingStreamer> Streamer = Module->GetStreamer(StreamerId);
-	if(!Streamer)
+	TSharedPtr<IPixelStreamingStreamer> Streamer = Module->FindStreamer(StreamerId);
+	if (!Streamer)
 	{
 		return;
 	}
-	
+
 	TArray64<uint8> ByteData;
 	bool bSuccess = FFileHelper::LoadFileToArray(ByteData, *FilePath);
-	if(bSuccess)
+	if (bSuccess)
 	{
 		Streamer->SendFileData(ByteData, MimeType, FileExtension);
 	}
 	else
 	{
 		UE_LOG(LogPixelStreaming, Error, TEXT("FileHelper failed to load file data"));
-	}	
+	}
 }
 
 void UPixelStreamingBlueprints::ForceKeyFrame()
 {
 	IPixelStreamingModule* Module = UE::PixelStreaming::FPixelStreamingModule::GetModule();
-	if(!Module)
+	if (!Module)
 	{
 		return;
 	}
-	TSharedPtr<IPixelStreamingStreamer> Streamer = Module->GetStreamer(Module->GetDefaultStreamerID());
-	if(!Streamer)
+	TSharedPtr<IPixelStreamingStreamer> Streamer = Module->FindStreamer(Module->GetDefaultStreamerID());
+	if (!Streamer)
 	{
 		return;
 	}
@@ -108,42 +108,42 @@ void UPixelStreamingBlueprints::ForceKeyFrame()
 void UPixelStreamingBlueprints::FreezeFrame(UTexture2D* Texture)
 {
 	IPixelStreamingModule* Module = UE::PixelStreaming::FPixelStreamingModule::GetModule();
-	if(!Module)
+	if (!Module)
 	{
 		return;
 	}
-	TSharedPtr<IPixelStreamingStreamer> Streamer = Module->GetStreamer(Module->GetDefaultStreamerID());
-	if(!Streamer)
+	TSharedPtr<IPixelStreamingStreamer> Streamer = Module->FindStreamer(Module->GetDefaultStreamerID());
+	if (!Streamer)
 	{
 		return;
 	}
-	Streamer->FreezeStream(Texture);	
+	Streamer->FreezeStream(Texture);
 }
 
 void UPixelStreamingBlueprints::StreamerFreezeStream(FString StreamerId, UTexture2D* Texture)
 {
 	IPixelStreamingModule* Module = UE::PixelStreaming::FPixelStreamingModule::GetModule();
-	if(!Module)
+	if (!Module)
 	{
 		return;
 	}
-	TSharedPtr<IPixelStreamingStreamer> Streamer = Module->GetStreamer(StreamerId);
-	if(!Streamer)
+	TSharedPtr<IPixelStreamingStreamer> Streamer = Module->FindStreamer(StreamerId);
+	if (!Streamer)
 	{
 		return;
 	}
-	Streamer->FreezeStream(Texture);	
+	Streamer->FreezeStream(Texture);
 }
 
 void UPixelStreamingBlueprints::UnfreezeFrame()
 {
 	IPixelStreamingModule* Module = UE::PixelStreaming::FPixelStreamingModule::GetModule();
-	if(!Module)
+	if (!Module)
 	{
 		return;
 	}
-	TSharedPtr<IPixelStreamingStreamer> Streamer = Module->GetStreamer(Module->GetDefaultStreamerID());
-	if(!Streamer)
+	TSharedPtr<IPixelStreamingStreamer> Streamer = Module->FindStreamer(Module->GetDefaultStreamerID());
+	if (!Streamer)
 	{
 		return;
 	}
@@ -153,12 +153,12 @@ void UPixelStreamingBlueprints::UnfreezeFrame()
 void UPixelStreamingBlueprints::StreamerUnfreezeStream(FString StreamerId)
 {
 	IPixelStreamingModule* Module = UE::PixelStreaming::FPixelStreamingModule::GetModule();
-	if(!Module)
+	if (!Module)
 	{
 		return;
 	}
-	TSharedPtr<IPixelStreamingStreamer> Streamer = Module->GetStreamer(StreamerId);
-	if(!Streamer)
+	TSharedPtr<IPixelStreamingStreamer> Streamer = Module->FindStreamer(StreamerId);
+	if (!Streamer)
 	{
 		return;
 	}
@@ -168,12 +168,12 @@ void UPixelStreamingBlueprints::StreamerUnfreezeStream(FString StreamerId)
 void UPixelStreamingBlueprints::KickPlayer(FString PlayerId)
 {
 	IPixelStreamingModule* Module = UE::PixelStreaming::FPixelStreamingModule::GetModule();
-	if(!Module)
+	if (!Module)
 	{
 		return;
 	}
-	TSharedPtr<IPixelStreamingStreamer> Streamer = Module->GetStreamer(Module->GetDefaultStreamerID());
-	if(!Streamer)
+	TSharedPtr<IPixelStreamingStreamer> Streamer = Module->FindStreamer(Module->GetDefaultStreamerID());
+	if (!Streamer)
 	{
 		return;
 	}
@@ -183,12 +183,12 @@ void UPixelStreamingBlueprints::KickPlayer(FString PlayerId)
 void UPixelStreamingBlueprints::StreamerKickPlayer(FString StreamerId, FString PlayerId)
 {
 	IPixelStreamingModule* Module = UE::PixelStreaming::FPixelStreamingModule::GetModule();
-	if(!Module)
+	if (!Module)
 	{
 		return;
 	}
-	TSharedPtr<IPixelStreamingStreamer> Streamer = Module->GetStreamer(StreamerId);
-	if(!Streamer)
+	TSharedPtr<IPixelStreamingStreamer> Streamer = Module->FindStreamer(StreamerId);
+	if (!Streamer)
 	{
 		return;
 	}
@@ -198,7 +198,7 @@ void UPixelStreamingBlueprints::StreamerKickPlayer(FString StreamerId, FString P
 FString UPixelStreamingBlueprints::GetDefaultStreamerID()
 {
 	IPixelStreamingModule* Module = UE::PixelStreaming::FPixelStreamingModule::GetModule();
-	if(!Module)
+	if (!Module)
 	{
 		FString();
 	}
diff --git a/Engine/Plugins/Media/PixelStreaming/Source/PixelStreaming/Private/PixelStreamingModule.cpp b/Engine/Plugins/Media/PixelStreaming/Source/PixelStreaming/Private/PixelStreamingModule.cpp
--- a/Engine/Plugins/Media/PixelStreaming/Source/PixelStreaming/Private/PixelStreamingModule.cpp
+++ b/Engine/Plugins/Media/PixelStreaming/Source/PixelStreaming/Private/PixelStreamingModule.cpp
@@ -240,7 +240,7 @@ namespace UE::PixelStreaming
 
 	TSharedPtr<IPixelStreamingStreamer> FPixelStreamingModule::CreateStreamer(const FString& StreamerId)
 	{
-		TSharedPtr<IPixelStreamingStreamer> ExistingStreamer = GetStreamer(StreamerId);
+		TSharedPtr<IPixelStreamingStreamer> ExistingStreamer = FindStreamer(StreamerId);
 		if (ExistingStreamer)
 		{
 			return ExistingStreamer;
@@ -283,11 +283,16 @@ namespace UE::PixelStreaming
 	}
 
 	TSharedPtr<IPixelStreamingStreamer> FPixelStreamingModule::GetStreamer(const FString& StreamerId)
+	{
+		return FindStreamer(StreamerId);
+	}
+
+	TSharedPtr<IPixelStreamingStreamer> FPixelStreamingModule::FindStreamer(const FString& StreamerId)
 	{
 		FScopeLock Lock(&StreamersCS);
 		if (Streamers.Contains(StreamerId))
 		{
-			return Streamers[StreamerId];
+			return Streamers[StreamerId].Pin();
 		}
 		return nullptr;
 	}
@@ -298,7 +303,7 @@ namespace UE::PixelStreaming
 		FScopeLock Lock(&StreamersCS);
 		if (Streamers.Contains(StreamerId))
 		{
-			ToBeDeleted = Streamers[StreamerId];
+			ToBeDeleted = Streamers[StreamerId].Pin();
 			Streamers.Remove(StreamerId);
 		}
 		return ToBeDeleted;
@@ -382,7 +387,7 @@ namespace UE::PixelStreaming
 		}
 		for (auto&& StreamerId : KeySet)
 		{
-			if (TSharedPtr<IPixelStreamingStreamer> Streamer = GetStreamer(StreamerId))
+			if (TSharedPtr<IPixelStreamingStreamer> Streamer = FindStreamer(StreamerId))
 			{
 				Func(Streamer);
 			}
diff --git a/Engine/Plugins/Media/PixelStreaming/Source/PixelStreaming/Private/PixelStreamingModule.h b/Engine/Plugins/Media/PixelStreaming/Source/PixelStreaming/Private/PixelStreamingModule.h
--- a/Engine/Plugins/Media/PixelStreaming/Source/PixelStreaming/Private/PixelStreamingModule.h
+++ b/Engine/Plugins/Media/PixelStreaming/Source/PixelStreaming/Private/PixelStreamingModule.h
@@ -34,6 +34,7 @@ namespace UE::PixelStreaming
 		virtual TSharedPtr<IPixelStreamingStreamer> CreateStreamer(const FString& StreamerId) override;
 		virtual TArray<FString> GetStreamerIds() override;
 		virtual TSharedPtr<IPixelStreamingStreamer> GetStreamer(const FString& StreamerId) override;
+		virtual TSharedPtr<IPixelStreamingStreamer> FindStreamer(const FString& StreamerId) override;
 		virtual TSharedPtr<IPixelStreamingStreamer> DeleteStreamer(const FString& StreamerId) override;
 		void DeleteStreamer(TSharedPtr<IPixelStreamingStreamer> ToBeDeleted) override;
 		virtual FString GetDefaultStreamerID() override;
@@ -80,7 +81,7 @@ namespace UE::PixelStreaming
 		TArray<UPixelStreamingInput*> InputComponents;
 		TSharedPtr<FVideoSourceGroup> ExternalVideoSourceGroup;
 		mutable FCriticalSection StreamersCS;
-		TMap<FString, TSharedPtr<IPixelStreamingStreamer>> Streamers;
+		TMap<FString, TWeakPtr<IPixelStreamingStreamer>> Streamers;
 		TSharedPtr<IPixelStreamingStreamer> DefaultStreamer;
 	};
 } // namespace UE::PixelStreaming
diff --git a/Engine/Plugins/Media/PixelStreaming/Source/PixelStreaming/Public/IPixelStreamingModule.h b/Engine/Plugins/Media/PixelStreaming/Source/PixelStreaming/Public/IPixelStreamingModule.h
--- a/Engine/Plugins/Media/PixelStreaming/Source/PixelStreaming/Public/IPixelStreamingModule.h
+++ b/Engine/Plugins/Media/PixelStreaming/Source/PixelStreaming/Public/IPixelStreamingModule.h
@@ -95,11 +95,19 @@ public:
 	virtual TArray<FString> GetStreamerIds() = 0;
 
 	/**
-	 * Get a streamer by an ID.
+	 * Get a streamer by an ID. 
+	 * // This method has been renamed as the module isn't responsible for a Streamer's lifetime, so renaming to Find as opposed to Get makes this more obvious
 	 * @return A pointer to the interface for a streamer. nullptr if the streamer isn't found
 	 */
+	UE_DEPRECATED(5.2, "GetStreamer() is deprecated. Please use FindStreamer instead.")
 	virtual TSharedPtr<IPixelStreamingStreamer> GetStreamer(const FString& StreamerId) = 0;
 
+	/**
+	 * Find a streamer by an ID.
+	 * @return A pointer to the interface for a streamer. nullptr if the streamer isn't found
+	 */
+	virtual TSharedPtr<IPixelStreamingStreamer> FindStreamer(const FString& StreamerId) = 0;
+
 	/**
 	 * Remove a streamer by an ID
 	 * @param StreamerId	-	The ID of the streamer to be removed.
diff --git a/Engine/Plugins/Media/PixelStreaming/Source/PixelStreamingEditor/Private/PixelStreamingToolbar.cpp b/Engine/Plugins/Media/PixelStreaming/Source/PixelStreamingEditor/Private/PixelStreamingToolbar.cpp
--- a/Engine/Plugins/Media/PixelStreaming/Source/PixelStreamingEditor/Private/PixelStreamingToolbar.cpp
+++ b/Engine/Plugins/Media/PixelStreaming/Source/PixelStreamingEditor/Private/PixelStreamingToolbar.cpp
@@ -72,7 +72,7 @@ namespace UE::EditorPixelStreaming
 				IPixelStreamingEditorModule::Get().StartStreaming(EStreamTypes::LevelEditorViewport);
 			}),
 			FCanExecuteAction::CreateLambda([] {
-				if (TSharedPtr<IPixelStreamingStreamer> Streamer = IPixelStreamingModule::Get().GetStreamer("Editor"))
+				if (TSharedPtr<IPixelStreamingStreamer> Streamer = IPixelStreamingModule::Get().FindStreamer("Editor"))
 				{
 					return !Streamer->IsStreaming();
 				}
@@ -85,7 +85,7 @@ namespace UE::EditorPixelStreaming
 				IPixelStreamingEditorModule::Get().StartStreaming(EStreamTypes::Editor);
 			}),
 			FCanExecuteAction::CreateLambda([] {
-				if (TSharedPtr<IPixelStreamingStreamer> Streamer = IPixelStreamingModule::Get().GetStreamer("Editor"))
+				if (TSharedPtr<IPixelStreamingStreamer> Streamer = IPixelStreamingModule::Get().FindStreamer("Editor"))
 				{
 					return !Streamer->IsStreaming();
 				}
@@ -318,7 +318,7 @@ namespace UE::EditorPixelStreaming
 					  .AutoWidth()
 						  [SNew(SEditableTextBox)
 								  .Text_Lambda([]() {
-									  TSharedPtr<IPixelStreamingStreamer> Streamer = IPixelStreamingModule::Get().GetStreamer("Editor");
+									  TSharedPtr<IPixelStreamingStreamer> Streamer = IPixelStreamingModule::Get().FindStreamer("Editor");
 									  return FText::FromString(Streamer->GetSignallingServerURL());
 								  })
 								  .OnTextChanged_Lambda([](const FText& InText) {
