Fix embedded signalling server not serving /images directory when images are now under /Public/images.

- Fixes showcase not displaying images when using embedded signalling server

diff --git a/Engine/Plugins/Media/PixelStreaming/Source/PixelStreamingServers/Private/SignallingServer.cpp b/Engine/Plugins/Media/PixelStreaming/Source/PixelStreamingServers/Private/SignallingServer.cpp
--- a/Engine/Plugins/Media/PixelStreaming/Source/PixelStreamingServers/Private/SignallingServer.cpp
+++ b/Engine/Plugins/Media/PixelStreaming/Source/PixelStreamingServers/Private/SignallingServer.cpp
@@ -152,17 +152,7 @@ namespace UE::PixelStreamingServers
 		PublicMount.SetDefaultFile(FString(TEXT("player.html")));
 		MountsArr.Add(PublicMount);
 
-		// Add /images
-		FWebSocketHttpMount ImagesMount;
-		ImagesMount.SetPathOnDisk(ServersDir / TEXT("images"));
-		ImagesMount.SetWebPath(FString(TEXT("/images")));
-		MountsArr.Add(ImagesMount);
-
-		// Add /scripts
-		FWebSocketHttpMount ScriptsMount;
-		ScriptsMount.SetPathOnDisk(ServersDir / TEXT("scripts"));
-		ScriptsMount.SetWebPath(FString(TEXT("/scripts")));
-		MountsArr.Add(ScriptsMount);
+		// Todo (Luke.Bermingham): Expose way for user to specify what directories to serve.
 
 		return MountsArr;
 	}
