Fix: Input Devices streamed from LiveLinkVCam App do not respect other Device IDs. PixelStreaming now keeps track of all of the connected gamepads and will assign a new id for every controller that connects regardless of which device it connects from

diff --git a/Engine/Plugins/Media/PixelStreaming/Source/PixelStreaming/Private/PixelStreamingModule.cpp b/Engine/Plugins/Media/PixelStreaming/Source/PixelStreaming/Private/PixelStreamingModule.cpp
--- a/Engine/Plugins/Media/PixelStreaming/Source/PixelStreaming/Private/PixelStreamingModule.cpp
+++ b/Engine/Plugins/Media/PixelStreaming/Source/PixelStreaming/Private/PixelStreamingModule.cpp
@@ -496,12 +496,12 @@ namespace UE::PixelStreaming
 
 			// Handle sending commands to peers
 			TWeakPtr<IPixelStreamingStreamer> WeakStreamer = Streamer;
-			InputHandler->OnSendMessage.AddLambda([WeakStreamer](FMemoryReader Ar) {
+			InputHandler->OnSendMessage.AddLambda([WeakStreamer](FString MessageName, FMemoryReader Ar) {
 				if (TSharedPtr<IPixelStreamingStreamer> Streamer = WeakStreamer.Pin())
 				{
 					FString Descriptor;
 					Ar << Descriptor;
-					Streamer->SendPlayerMessage(FPixelStreamingInputProtocol::FromStreamerProtocol.Find("Command")->GetID(), Descriptor);
+					Streamer->SendPlayerMessage(FPixelStreamingInputProtocol::FromStreamerProtocol.Find(MessageName)->GetID(), Descriptor);
 				}
 			});
 		}
diff --git a/Engine/Plugins/Media/PixelStreaming/Source/PixelStreamingInput/Private/PixelStreamingInputDevice.cpp b/Engine/Plugins/Media/PixelStreaming/Source/PixelStreamingInput/Private/PixelStreamingInputDevice.cpp
--- a/Engine/Plugins/Media/PixelStreaming/Source/PixelStreamingInput/Private/PixelStreamingInputDevice.cpp
+++ b/Engine/Plugins/Media/PixelStreaming/Source/PixelStreamingInput/Private/PixelStreamingInputDevice.cpp
@@ -86,4 +86,21 @@ namespace UE::PixelStreamingInput
 			}
 		}
 	}
+
+	uint8 FPixelStreamingInputDevice::OnControllerConnected()
+	{
+		uint8 NextControllerId = 0;
+		while (ConnectedControllers.Contains(NextControllerId))
+		{
+			NextControllerId++;
+		}
+
+		ConnectedControllers.Add(NextControllerId);
+		return NextControllerId;
+	}
+
+	void FPixelStreamingInputDevice::OnControllerDisconnected(uint8 DeleteControllerId)
+	{
+		ConnectedControllers.Remove(DeleteControllerId);
+	}
 } // namespace UE::PixelStreamingInput
diff --git a/Engine/Plugins/Media/PixelStreaming/Source/PixelStreamingInput/Private/PixelStreamingInputDevice.h b/Engine/Plugins/Media/PixelStreaming/Source/PixelStreamingInput/Private/PixelStreamingInputDevice.h
--- a/Engine/Plugins/Media/PixelStreaming/Source/PixelStreamingInput/Private/PixelStreamingInputDevice.h
+++ b/Engine/Plugins/Media/PixelStreaming/Source/PixelStreamingInput/Private/PixelStreamingInputDevice.h
@@ -27,6 +27,8 @@ namespace UE::PixelStreamingInput
 		virtual void SetChannelValues(int32 ControllerId, const FForceFeedbackValues& Values) override;
 
 		void AddInputHandler(TSharedPtr<IPixelStreamingInputHandler> InputHandler);
+		uint8 OnControllerConnected();
+		void OnControllerDisconnected(uint8 DeleteControllerId);
 
 		static TSharedPtr<FPixelStreamingInputDevice> GetInputDevice();
 
@@ -41,11 +43,20 @@ namespace UE::PixelStreamingInput
 		static TSharedPtr<FPixelStreamingInputDevice> InputDevice;
 
 		/**
-		 * The Map of input handlers. Each input handler belongs to a single streamer
+		 * The array of input handlers. Each input handler belongs to a single streamer
 		 *
 		 * NOTE: We store each input handler as a weakptr as we don't want to be the reason
 		 * the handler doesn't get deleted. Each handler should be tied to a Streamer's lifecycle
 		 */
 		TArray<TWeakPtr<IPixelStreamingInputHandler>> InputHandlers;
+
+		/**
+		 * The array of connected controllers. As each device can have multiple controllers, we want to make sure that each controller of each device is unique
+		 * As such, a simple incrementer approach is not applicable and we must instead keep track of all the connected controllers
+		 *
+		 * NOTE: We store each input handler as a weakptr as we don't want to be the reason
+		 * the handler doesn't get deleted. Each handler should be tied to a Streamer's lifecycle
+		 */
+		TArray<uint8> ConnectedControllers;
 	};
 } // namespace UE::PixelStreamingInput
\ No newline at end of file
diff --git a/Engine/Plugins/Media/PixelStreaming/Source/PixelStreamingInput/Private/PixelStreamingInputHandler.cpp b/Engine/Plugins/Media/PixelStreaming/Source/PixelStreamingInput/Private/PixelStreamingInputHandler.cpp
--- a/Engine/Plugins/Media/PixelStreaming/Source/PixelStreamingInput/Private/PixelStreamingInputHandler.cpp
+++ b/Engine/Plugins/Media/PixelStreaming/Source/PixelStreamingInput/Private/PixelStreamingInputHandler.cpp
@@ -17,13 +17,13 @@
 #include "Widgets/Input/SEditableTextBox.h"
 #include "Misc/CoreMiscDefines.h"
 #include "Input/HittestGrid.h"
-#include "IPixelStreamingInputModule.h"
 #include "PixelStreamingHMD.h"
 #include "IPixelStreamingHMDModule.h"
 #include "PixelStreamingInputEnums.h"
 #include "PixelStreamingInputConversion.h"
 #include "Widgets/Input/SEditableText.h"
 #include "Widgets/Input/SMultiLineEditableTextBox.h"
+#include "PixelStreamingInputDevice.h"
 
 DECLARE_LOG_CATEGORY_EXTERN(LogPixelStreamingInputHandler, Log, VeryVerbose);
 DEFINE_LOG_CATEGORY(LogPixelStreamingInputHandler);
@@ -56,9 +56,11 @@ namespace UE::PixelStreamingInput
 		RegisterMessageHandler("TouchMove", [this](FMemoryReader Ar) { HandleOnTouchMoved(Ar); });
 		RegisterMessageHandler("TouchEnd", [this](FMemoryReader Ar) { HandleOnTouchEnded(Ar); });
 
+		RegisterMessageHandler("GamepadConnected", [this](FMemoryReader Ar) { HandleOnControllerConnected(Ar); });
 		RegisterMessageHandler("GamepadAnalog", [this](FMemoryReader Ar) { HandleOnControllerAnalog(Ar); });
 		RegisterMessageHandler("GamepadButtonPressed", [this](FMemoryReader Ar) { HandleOnControllerButtonPressed(Ar); });
 		RegisterMessageHandler("GamepadButtonReleased", [this](FMemoryReader Ar) { HandleOnControllerButtonReleased(Ar); });
+		RegisterMessageHandler("GamepadDisconnected", [this](FMemoryReader Ar) { HandleOnControllerDisconnected(Ar); });
 
 		RegisterMessageHandler("MouseEnter", [this](FMemoryReader Ar) { HandleOnMouseEnter(Ar); });
 		RegisterMessageHandler("MouseLeave", [this](FMemoryReader Ar) { HandleOnMouseLeave(Ar); });
@@ -541,6 +543,22 @@ namespace UE::PixelStreamingInput
 		}
 	}
 
+	void FPixelStreamingInputHandler::HandleOnControllerConnected(FMemoryReader Ar)
+	{
+		TSharedPtr<FPixelStreamingInputDevice> InputDevice = FPixelStreamingInputDevice::GetInputDevice();
+		uint8 NextControllerId = InputDevice->OnControllerConnected();
+
+		FString Descriptor = FString::Printf(TEXT("{ \"controllerId\": %d }"), NextControllerId);
+
+		FBufferArchive Buffer;
+		Buffer << Descriptor;
+		TArray<uint8> Data(Buffer.GetData(), Buffer.Num());
+		// Specific implementation for this method is handled per streamer
+		OnSendMessage.Broadcast("GamepadResponse", FMemoryReader(Data));
+
+		UE_LOG(LogPixelStreamingInputHandler, Verbose, TEXT("GAMEPAD_CONNECTED: ControllerId = %d"), NextControllerId);
+	}
+
 	void FPixelStreamingInputHandler::HandleOnControllerAnalog(FMemoryReader Ar)
 	{
 		TPayloadThreeParam<uint8, uint8, double> Payload(Ar);
@@ -554,8 +572,22 @@ namespace UE::PixelStreamingInput
 		float AnalogValue = (float)Payload.Param3;
 		FPlatformUserId UserId = IPlatformInputDeviceMapper::Get().GetPrimaryPlatformUser();
 
-		UE_LOG(LogPixelStreamingInputHandler, Verbose, TEXT("GAMEPAD_ANALOG: ControllerId = %d; KeyName = %s; AnalogValue = %.4f;"), ControllerId.GetId(), *AxisPtr->ToString(), AnalogValue);
-		MessageHandler->OnControllerAnalog(AxisPtr->GetFName(), UserId, ControllerId, AnalogValue);
+		FSlateApplication& SlateApplication = FSlateApplication::Get();
+		FAnalogInputEvent AnalogInputEvent(
+			*AxisPtr,													  /* InKey */
+			SlateApplication.GetPlatformApplication()->GetModifierKeys(), /* InModifierKeys */
+			ControllerId,												  /* InDeviceId */
+			false,														  /* bInIsRepeat*/
+			0,															  /* InCharacterCode */
+			0,															  /* InKeyCode */
+			AnalogValue,												  /* InAnalogValue */
+			// TODO (william.belcher): This user idx should be the playerId
+			0 /* InUserIndex */
+		);
+		FSlateApplication::Get().ProcessAnalogInputEvent(AnalogInputEvent);
+
+		// Commented to reduce log spam
+		// UE_LOG(LogPixelStreamingInputHandler, Verbose, TEXT("GAMEPAD_ANALOG: ControllerId = %d; KeyName = %s; AnalogValue = %.4f;"), ControllerId.GetId(), *AxisPtr->ToString(), AnalogValue);
 	}
 
 	void FPixelStreamingInputHandler::HandleOnControllerButtonPressed(FMemoryReader Ar)
@@ -571,8 +603,19 @@ namespace UE::PixelStreamingInput
 		bool bIsRepeat = Payload.Param3 != 0;
 		FPlatformUserId UserId = IPlatformInputDeviceMapper::Get().GetPrimaryPlatformUser();
 
+		FKeyEvent KeyEvent(
+			*ButtonPtr,															  /* InKey */
+			FSlateApplication::Get().GetPlatformApplication()->GetModifierKeys(), /* InModifierKeys */
+			ControllerId,														  /* InDeviceId */
+			bIsRepeat,															  /* bInIsRepeat */
+			0,																	  /* InCharacterCode */
+			0,																	  /* InKeyCode */
+			// TODO (william.belcher): This user idx should be the playerId
+			0 /* InUserIndex*/
+		);
+		FSlateApplication::Get().ProcessKeyDownEvent(KeyEvent);
+
 		UE_LOG(LogPixelStreamingInputHandler, Verbose, TEXT("GAMEPAD_PRESSED: ControllerId = %d; KeyName = %s; IsRepeat = %s;"), ControllerId.GetId(), *ButtonPtr->ToString(), bIsRepeat ? TEXT("True") : TEXT("False"));
-		MessageHandler->OnControllerButtonPressed(ButtonPtr->GetFName(), UserId, ControllerId, bIsRepeat);
 	}
 
 	void FPixelStreamingInputHandler::HandleOnControllerButtonReleased(FMemoryReader Ar)
@@ -587,8 +630,31 @@ namespace UE::PixelStreamingInput
 		FInputDeviceId ControllerId = FInputDeviceId::CreateFromInternalId((int32)Payload.Param1);
 		FPlatformUserId UserId = IPlatformInputDeviceMapper::Get().GetPrimaryPlatformUser();
 
+		FKeyEvent KeyEvent(
+			*ButtonPtr,															  /* InKey */
+			FSlateApplication::Get().GetPlatformApplication()->GetModifierKeys(), /* InModifierKeys */
+			ControllerId,														  /* InDeviceId */
+			false,																  /* bInIsRepeat */
+			0,																	  /* InCharacterCode */
+			0,																	  /* InKeyCode */
+			// TODO (william.belcher): This user idx should be the playerId
+			0 /* InUserIndex*/
+		);
+		FSlateApplication::Get().ProcessKeyUpEvent(KeyEvent);
+
 		UE_LOG(LogPixelStreamingInputHandler, Verbose, TEXT("GAMEPAD_RELEASED: ControllerId = %d; KeyName = %s;"), ControllerId.GetId(), *ButtonPtr->ToString());
-		MessageHandler->OnControllerButtonReleased(ButtonPtr->GetFName(), UserId, ControllerId, false);
+	}
+
+	void FPixelStreamingInputHandler::HandleOnControllerDisconnected(FMemoryReader Ar)
+	{
+		TPayloadOneParam<uint8> Payload(Ar);
+
+		uint8 DisconnectedControllerId = Payload.Param1;
+
+		TSharedPtr<FPixelStreamingInputDevice> InputDevice = FPixelStreamingInputDevice::GetInputDevice();
+		InputDevice->OnControllerDisconnected(DisconnectedControllerId);
+
+		UE_LOG(LogPixelStreamingInputHandler, Verbose, TEXT("GAMEPAD_DISCONNECTED: ControllerId = %d"), DisconnectedControllerId);
 	}
 
 	/**
@@ -1337,7 +1403,7 @@ namespace UE::PixelStreamingInput
 			Buffer << Descriptor;
 			TArray<uint8> Data(Buffer.GetData(), Buffer.Num());
 			// Specific implementation for this method is handled per streamer
-			OnSendMessage.Broadcast(FMemoryReader(Data));
+			OnSendMessage.Broadcast("Command", FMemoryReader(Data));
 		});
 	}
 
diff --git a/Engine/Plugins/Media/PixelStreaming/Source/PixelStreamingInput/Private/PixelStreamingInputHandler.h b/Engine/Plugins/Media/PixelStreaming/Source/PixelStreamingInput/Private/PixelStreamingInputHandler.h
--- a/Engine/Plugins/Media/PixelStreaming/Source/PixelStreamingInput/Private/PixelStreamingInputHandler.h
+++ b/Engine/Plugins/Media/PixelStreaming/Source/PixelStreamingInput/Private/PixelStreamingInputHandler.h
@@ -78,9 +78,11 @@ namespace UE::PixelStreamingInput
 		/**
 		 * Controller handling
 		 */
+		virtual void HandleOnControllerConnected(FMemoryReader Ar);
 		virtual void HandleOnControllerAnalog(FMemoryReader Ar);
 		virtual void HandleOnControllerButtonPressed(FMemoryReader Ar);
 		virtual void HandleOnControllerButtonReleased(FMemoryReader Ar);
+		virtual void HandleOnControllerDisconnected(FMemoryReader Ar);
 		/**
 		 * Mouse handling
 		 */
diff --git a/Engine/Plugins/Media/PixelStreaming/Source/PixelStreamingInput/Private/PixelStreamingInputModule.cpp b/Engine/Plugins/Media/PixelStreaming/Source/PixelStreamingInput/Private/PixelStreamingInputModule.cpp
--- a/Engine/Plugins/Media/PixelStreaming/Source/PixelStreamingInput/Private/PixelStreamingInputModule.cpp
+++ b/Engine/Plugins/Media/PixelStreaming/Source/PixelStreamingInput/Private/PixelStreamingInputModule.cpp
@@ -100,6 +100,8 @@ namespace UE::PixelStreamingInput
 		FPixelStreamingInputProtocol::ToStreamerProtocol.Add("GamepadButtonPressed", FPixelStreamingInputMessage(90, { EType::Uint8, EType::Uint8, EType::Uint8 }));
 		FPixelStreamingInputProtocol::ToStreamerProtocol.Add("GamepadButtonReleased", FPixelStreamingInputMessage(91, { EType::Uint8, EType::Uint8, EType::Uint8 }));
 		FPixelStreamingInputProtocol::ToStreamerProtocol.Add("GamepadAnalog", FPixelStreamingInputMessage(92, { EType::Uint8, EType::Uint8, EType::Double }));
+		FPixelStreamingInputProtocol::ToStreamerProtocol.Add("GamepadConnected", FPixelStreamingInputMessage(93));
+		FPixelStreamingInputProtocol::ToStreamerProtocol.Add("GamepadDisconnected", FPixelStreamingInputMessage(94, { EType::Uint8 }));
 
 		// XR Input Messages.
 		// clang-format off
@@ -145,6 +147,7 @@ namespace UE::PixelStreamingInput
 		FPixelStreamingInputProtocol::FromStreamerProtocol.Add("FileContents", FPixelStreamingInputMessage(10));
 		FPixelStreamingInputProtocol::FromStreamerProtocol.Add("TestEcho", FPixelStreamingInputMessage(11));
 		FPixelStreamingInputProtocol::FromStreamerProtocol.Add("InputControlOwnership", FPixelStreamingInputMessage(12));
+		FPixelStreamingInputProtocol::FromStreamerProtocol.Add("GamepadResponse", FPixelStreamingInputMessage(13));
 		FPixelStreamingInputProtocol::FromStreamerProtocol.Add("Protocol", FPixelStreamingInputMessage(255));
 	}
 } // namespace UE::PixelStreamingInput
diff --git a/Engine/Plugins/Media/PixelStreaming/Source/PixelStreamingInput/Public/IPixelStreamingInputHandler.h b/Engine/Plugins/Media/PixelStreaming/Source/PixelStreamingInput/Public/IPixelStreamingInputHandler.h
--- a/Engine/Plugins/Media/PixelStreaming/Source/PixelStreamingInput/Public/IPixelStreamingInputHandler.h
+++ b/Engine/Plugins/Media/PixelStreaming/Source/PixelStreamingInput/Public/IPixelStreamingInputHandler.h
@@ -110,6 +110,6 @@ public:
 	 */
 	virtual void SetInputType(EPixelStreamingInputType InputType) = 0;
 
-	DECLARE_MULTICAST_DELEGATE_OneParam(FOnSendMessage, FMemoryReader);
+	DECLARE_MULTICAST_DELEGATE_TwoParams(FOnSendMessage, FString, FMemoryReader);
 	FOnSendMessage OnSendMessage;
 };
\ No newline at end of file
