Axis input continues to be applied after stopping when device is connected to the iPad

diff --git a/Engine/Plugins/Media/PixelStreaming/Source/PixelStreamingInput/Private/PixelStreamingInputHandler.cpp b/Engine/Plugins/Media/PixelStreaming/Source/PixelStreamingInput/Private/PixelStreamingInputHandler.cpp
--- a/Engine/Plugins/Media/PixelStreaming/Source/PixelStreamingInput/Private/PixelStreamingInputHandler.cpp
+++ b/Engine/Plugins/Media/PixelStreaming/Source/PixelStreamingInput/Private/PixelStreamingInputHandler.cpp
@@ -251,7 +251,8 @@ namespace UE::PixelStreamingInput
 			FMemoryReader Ar(Message.Data);
 			(*Message.Handler)(Ar);
 		}
-
+		
+		ProcessLatestAnalogInputFromThisTick();
 		BroadcastActiveTouchMoveEvents();
 	}
 
@@ -561,33 +562,10 @@ namespace UE::PixelStreamingInput
 
 	void FPixelStreamingInputHandler::HandleOnControllerAnalog(FMemoryReader Ar)
 	{
-		TPayloadThreeParam<uint8, uint8, double> Payload(Ar);
-		FKey* AxisPtr = FPixelStreamingInputConverter::GamepadInputToFKey.Find(MakeTuple(Payload.Param2, Action::Axis));
-		if (AxisPtr == nullptr)
-		{
-			return;
-		}
-
-		FInputDeviceId ControllerId = FInputDeviceId::CreateFromInternalId((int32)Payload.Param1);
-		float AnalogValue = (float)Payload.Param3;
-		FPlatformUserId UserId = IPlatformInputDeviceMapper::Get().GetPrimaryPlatformUser();
-
-		FSlateApplication& SlateApplication = FSlateApplication::Get();
-		FAnalogInputEvent AnalogInputEvent(
-			*AxisPtr,													  /* InKey */
-			SlateApplication.GetPlatformApplication()->GetModifierKeys(), /* InModifierKeys */
-			ControllerId,												  /* InDeviceId */
-			false,														  /* bInIsRepeat*/
-			0,															  /* InCharacterCode */
-			0,															  /* InKeyCode */
-			AnalogValue,												  /* InAnalogValue */
-			// TODO (william.belcher): This user idx should be the playerId
-			0 /* InUserIndex */
-		);
-		FSlateApplication::Get().ProcessAnalogInputEvent(AnalogInputEvent);
-
-		// Commented to reduce log spam
-		// UE_LOG(LogPixelStreamingInputHandler, Verbose, TEXT("GAMEPAD_ANALOG: ControllerId = %d; KeyName = %s; AnalogValue = %.4f;"), ControllerId.GetId(), *AxisPtr->ToString(), AnalogValue);
+		const TPayloadThreeParam<uint8, uint8, double> Payload(Ar);
+		const FInputDeviceId ControllerId = FInputDeviceId::CreateFromInternalId((int32)Payload.Param1);
+		// Overwrite the last data: every tick only process the latest
+		AnalogEventsReceivedThisTick.FindOrAdd(ControllerId).FindOrAdd(Payload.Param2) = Payload.Param3;
 	}
 
 	void FPixelStreamingInputHandler::HandleOnControllerButtonPressed(FMemoryReader Ar)
@@ -1255,6 +1233,41 @@ namespace UE::PixelStreamingInput
 		return true;
 	}
 
+	void FPixelStreamingInputHandler::ProcessLatestAnalogInputFromThisTick()
+	{
+		for (auto AnalogInputIt = AnalogEventsReceivedThisTick.CreateConstIterator(); AnalogInputIt; ++AnalogInputIt)
+		{
+			for (auto FKeyIt = AnalogInputIt->Value.CreateConstIterator(); FKeyIt; ++FKeyIt)
+			{
+				ProcessAnalog(AnalogInputIt->Key, FKeyIt->Key, FKeyIt->Value);
+			}
+		}
+		AnalogEventsReceivedThisTick.Reset();
+	}
+
+	void FPixelStreamingInputHandler::ProcessAnalog(const FInputDeviceId& ControllerId, FKeyId Key, FAnalogValue AnalogValue)
+	{
+		const FKey* AxisPtr = FPixelStreamingInputConverter::GamepadInputToFKey.Find(MakeTuple(Key, Action::Axis));
+		if (!AxisPtr)
+		{
+			return;
+		}
+		
+		FSlateApplication& SlateApplication = FSlateApplication::Get();
+		const FAnalogInputEvent AnalogInputEvent(
+			*AxisPtr,														/* InKey */
+			SlateApplication.GetPlatformApplication()->GetModifierKeys(),	/* InModifierKeys */
+			ControllerId,													/* InDeviceId */
+			false,													/* bInIsRepeat*/
+			0,													/* InCharacterCode */
+			0,														/* InKeyCode */
+			AnalogValue,													/* InAnalogValue */
+			// TODO (william.belcher): This user idx should be the playerId
+			0 /* InUserIndex */
+		);
+		FSlateApplication::Get().ProcessAnalogInputEvent(AnalogInputEvent);
+	}
+
 	void FPixelStreamingInputHandler::BroadcastActiveTouchMoveEvents()
 	{
 		if (!ensure(MessageHandler))
diff --git a/Engine/Plugins/Media/PixelStreaming/Source/PixelStreamingInput/Private/PixelStreamingInputHandler.h b/Engine/Plugins/Media/PixelStreaming/Source/PixelStreamingInput/Private/PixelStreamingInputHandler.h
--- a/Engine/Plugins/Media/PixelStreaming/Source/PixelStreamingInput/Private/PixelStreamingInputHandler.h
+++ b/Engine/Plugins/Media/PixelStreaming/Source/PixelStreamingInput/Private/PixelStreamingInputHandler.h
@@ -63,6 +63,7 @@ namespace UE::PixelStreamingInput
 		// End IMotionController Interface
 
 	protected:
+		
 		/**
 		 * Key press handling
 		 */
@@ -134,6 +135,26 @@ namespace UE::PixelStreamingInput
 
 		// Keep a cache of the last touch events as we need to fire Touch Moved every frame while touch is down
 		TMap<int32, FCachedTouchEvent> CachedTouchEvents;
+		
+		using FKeyId = uint8;
+		using FAnalogValue = double;
+		/**
+		 * If more values are received in a single tick (e.g. could be temp network issue),
+		 * then we only forward the latest value.
+		 * 
+		 * Reason: The input system seems to expect at most one raw analog value per FKey per Tick.
+		 * If this is not done, the input system can get stuck on non-zero input value even if the user has
+		 * already stopped moving the analog stick. It would stay stuck until the next time the user moves the stick.
+		 * 
+		 * The values arrive in the order of recording: that means once the player releases the analog,
+		 * the last analog value would be 0.
+		 */
+		TMap<FInputDeviceId, TMap<FKeyId, FAnalogValue>> AnalogEventsReceivedThisTick;
+
+		/** Forwards the latest analog input received for each key this tick. */
+		void ProcessLatestAnalogInputFromThisTick();
+		/** Forward a single analog input the engine. */
+		void ProcessAnalog(const FInputDeviceId& ControllerId, FKeyId Key, FAnalogValue AnalogValue);
 
 		// Track which touch events we processed this frame so we can avoid re-processing them
 		TSet<int32> TouchIndicesProcessedThisFrame;
