Fix: Non-unity compilation errors in PixelStreamingSettings.h

diff --git a/Engine/Plugins/Media/PixelStreaming/Source/PixelStreaming/Private/PixelStreamingPeerConnection.cpp b/Engine/Plugins/Media/PixelStreaming/Source/PixelStreaming/Private/PixelStreamingPeerConnection.cpp
--- a/Engine/Plugins/Media/PixelStreaming/Source/PixelStreaming/Private/PixelStreamingPeerConnection.cpp
+++ b/Engine/Plugins/Media/PixelStreaming/Source/PixelStreaming/Private/PixelStreamingPeerConnection.cpp
@@ -26,7 +26,7 @@ namespace
 
 		if (IsSFU && Settings::SimulcastParameters.Layers.Num() > 0)
 		{
-			using FLayer = FPixelStreamingSimulcastParameters::FLayer;
+			using FLayer = FPixelStreamingSimulcastParameters::FPixelStreamingSimulcastLayer;
 
 			// encodings should be lowest res to highest
 			TArray<FLayer*> SortedLayers;
diff --git a/Engine/Plugins/Media/PixelStreaming/Source/PixelStreaming/Private/Settings.cpp b/Engine/Plugins/Media/PixelStreaming/Source/PixelStreaming/Private/Settings.cpp
--- a/Engine/Plugins/Media/PixelStreaming/Source/PixelStreaming/Private/Settings.cpp
+++ b/Engine/Plugins/Media/PixelStreaming/Source/PixelStreaming/Private/Settings.cpp
@@ -437,7 +437,7 @@ namespace UE::PixelStreaming::Settings
 		int NextOption = 0;
 		while (bSuccess && ((OptionCount - NextOption) >= 3))
 		{
-			FPixelStreamingSimulcastParameters::FLayer Layer;
+			FPixelStreamingSimulcastParameters::FPixelStreamingSimulcastLayer Layer;
 			bSuccess = FDefaultValueHelper::ParseFloat(ParameterArray[NextOption++], Layer.Scaling);
 			bSuccess = FDefaultValueHelper::ParseInt(ParameterArray[NextOption++], Layer.MinBitrate);
 			bSuccess = FDefaultValueHelper::ParseInt(ParameterArray[NextOption++], Layer.MaxBitrate);
diff --git a/Engine/Plugins/Media/PixelStreaming/Source/PixelStreaming/Public/PixelStreamingSettings.h b/Engine/Plugins/Media/PixelStreaming/Source/PixelStreaming/Public/PixelStreamingSettings.h
--- a/Engine/Plugins/Media/PixelStreaming/Source/PixelStreaming/Public/PixelStreamingSettings.h
+++ b/Engine/Plugins/Media/PixelStreaming/Source/PixelStreaming/Public/PixelStreamingSettings.h
@@ -2,6 +2,9 @@
 
 #pragma once
 
+#include "CoreTypes.h"
+#include "Containers/Array.h"
+
 /**
  * A struct representing the simulcast paramaters used by PixelStreaming. The parameters
  * contain an array of layers, with each layer having a Scale, MinBitrate and MaxBitrate
@@ -9,14 +12,14 @@
  */
 struct PIXELSTREAMING_API FPixelStreamingSimulcastParameters
 {
-	struct FLayer
+	struct PIXELSTREAMING_API FPixelStreamingSimulcastLayer
 	{
 		float Scaling;
 		int MinBitrate;
 		int MaxBitrate;
 	};
 
-	TArray<FLayer> Layers;
+	TArray<FPixelStreamingSimulcastLayer> Layers;
 };
 
 /**
