Fix: Frontend folder not being copied during packaging

diff --git a/Engine/Plugins/Media/PixelStreaming/Source/PixelStreaming/PixelStreaming.Build.cs b/Engine/Plugins/Media/PixelStreaming/Source/PixelStreaming/PixelStreaming.Build.cs
--- a/Engine/Plugins/Media/PixelStreaming/Source/PixelStreaming/PixelStreaming.Build.cs
+++ b/Engine/Plugins/Media/PixelStreaming/Source/PixelStreaming/PixelStreaming.Build.cs
@@ -109,6 +109,7 @@ namespace UnrealBuildTool.Rules
 			// When we build a Game target we also package the servers with it as runtime dependencies
 			if(Target.Type == TargetType.Game && Target.ProjectFile != null)
 			{
+				AddFolder("Frontend");
 				AddFolder("SignallingWebServer");
 				AddFolder("Matchmaker");
 				AddFolder("SFU");
