Fix: SWindow properties would update during composition, leading to RHI errors when trying to copy textures

diff --git a/Engine/Plugins/Media/PixelStreaming/Source/PixelStreamingEditor/Private/PixelStreamingVideoInputBackBufferComposited.cpp b/Engine/Plugins/Media/PixelStreaming/Source/PixelStreamingEditor/Private/PixelStreamingVideoInputBackBufferComposited.cpp
--- a/Engine/Plugins/Media/PixelStreaming/Source/PixelStreamingEditor/Private/PixelStreamingVideoInputBackBufferComposited.cpp
+++ b/Engine/Plugins/Media/PixelStreaming/Source/PixelStreamingEditor/Private/PixelStreamingVideoInputBackBufferComposited.cpp
@@ -42,8 +42,22 @@ FPixelStreamingVideoInputBackBufferComposited::FPixelStreamingVideoInputBackBuff
 	UE::PixelStreaming::DoOnGameThread([this]() {
 		FSlateApplication::Get().OnPreTick().AddLambda([this](float DeltaTime) {
 			FScopeLock Lock(&TopLevelWindowsCriticalSection);
+			TArray<TSharedRef<SWindow>> TopLevelSWindows;
+			FSlateApplication::Get().GetAllVisibleWindowsOrdered(TopLevelSWindows);
+
 			TopLevelWindows.Empty();
-			FSlateApplication::Get().GetAllVisibleWindowsOrdered(TopLevelWindows);
+			// We store all the necessary window information in structs. This prevents window information from updating
+			// underneath us while we composite and also means we aren't holding on to any shared refs between compositions
+			for(TSharedRef<SWindow>& CurrentWindow : TopLevelSWindows)
+			{
+				TopLevelWindows.Add(FPixelStreamingVideoInputBackBufferComposited::Window(
+					CurrentWindow->GetPositionInScreen(),
+					CurrentWindow->GetSizeInScreen(),
+					CurrentWindow->GetOpacity(),
+					CurrentWindow->GetType(),
+					&CurrentWindow.Get()
+				));
+			}
 		});
 	});
 }
@@ -72,13 +86,22 @@ void FPixelStreamingVideoInputBackBufferComposited::OnBackBufferReady(SWindow& S
 	UE_LOG(LogPixelStreamingBackBufferComposited, Verbose, TEXT("Type: %s"), UE::EditorPixelStreaming::ToString(SlateWindow.GetType()));
 
 	FIntPoint FrameSize = FrameBuffer->GetSizeXY();
-	FString Hash = UE::EditorPixelStreaming::HashWindow(SlateWindow, FrameBuffer);
-	FTextureRHIRef StagingTexture = StagingTextures.FindRef(Hash);
-	if (!StagingTexture.IsValid())
+	FTextureRHIRef StagingTexture;
+	if (SlateWindow.GetType() != EWindowType::ToolTip)
+	{
+		FString Hash = UE::EditorPixelStreaming::HashWindow(SlateWindow, FrameBuffer);
+		StagingTexture = StagingTextures.FindRef(Hash);
+		if (!StagingTexture.IsValid())
+		{
+			UE_LOG(LogPixelStreamingBackBufferComposited, Verbose, TEXT("Creating new staging texture: %dx%d"), FrameSize.X, FrameSize.Y);
+			StagingTexture = UE::PixelStreaming::CreateRHITexture(FrameSize.X, FrameSize.Y);
+			StagingTextures.Add(Hash, StagingTexture);
+		}
+	}
+	else
 	{
-		UE_LOG(LogPixelStreamingBackBufferComposited, Verbose, TEXT("Creating new staging texture: %dx%d"), FrameSize.X, FrameSize.Y);
+		// Always create a new staging texture for tooltips due to their odd handling
 		StagingTexture = UE::PixelStreaming::CreateRHITexture(FrameSize.X, FrameSize.Y);
-		StagingTextures.Add(Hash, StagingTexture);
 	}
 
 	FRHICommandListImmediate& RHICmdList = FRHICommandListExecutor::GetImmediateCommandList();
@@ -89,9 +112,9 @@ void FPixelStreamingVideoInputBackBufferComposited::OnBackBufferReady(SWindow& S
 	TopLevelWindowsCriticalSection.Lock();
 	// Check that we have received a texture from every window in the TopLevelWindows array
 	uint8 NumWindowsRendered = 0;
-	for (TSharedRef<SWindow> Window : TopLevelWindows)
+	for (FPixelStreamingVideoInputBackBufferComposited::Window TopWindow : TopLevelWindows)
 	{
-		if (TopLevelWindowTextures.FindRef(&Window.Get()))
+		if (TopLevelWindowTextures.FindRef(TopWindow.OwningWindow))
 		{
 			++NumWindowsRendered;
 		}
@@ -109,22 +132,22 @@ void FPixelStreamingVideoInputBackBufferComposited::CompositeWindows()
 	FRHICommandListImmediate& RHICmdList = FRHICommandListExecutor::GetImmediateCommandList();
 
 	// Process all of the windows we will need to render. This processing step finds the extents of the
-	// composited texture as well as the top-left most point
+	// composited texture as well as the top-left point
 	FIntPoint TopLeft = FIntPoint(MAX_int32, MAX_int32);
 	FIntPoint BottomRight = FIntPoint(MIN_int32, MIN_int32);
 	for (uint8 i = 0; i < TopLevelWindows.Num(); i++)
 	{
-		TSharedRef<SWindow> CurrentWindow = TopLevelWindows[i];
+		FPixelStreamingVideoInputBackBufferComposited::Window CurrentWindow = TopLevelWindows[i];
 		// Early out if we have an "invalid" window type
-		if (CurrentWindow->GetOpacity() == 0.f || CurrentWindow->GetSizeInScreen() == FVector2f(0, 0))
+		if (CurrentWindow.Opacity == 0.f || CurrentWindow.SizeInScreen == FVector2D(0, 0))
 		{
 			continue;
 		}
 
-		FTextureRHIRef CurrentTexture = TopLevelWindowTextures.FindRef(&CurrentWindow.Get());
-		FIntPoint TextureExtent = VectorMin(CurrentTexture->GetSizeXY(), CurrentWindow->GetSizeInScreen().IntPoint());
+		FTextureRHIRef CurrentTexture = TopLevelWindowTextures.FindRef(CurrentWindow.OwningWindow);
+		FIntPoint TextureExtent = VectorMin(CurrentTexture->GetSizeXY(), CurrentWindow.SizeInScreen.IntPoint());
 
-		FIntPoint WindowPosition = FIntPoint(CurrentWindow->GetPositionInScreen().X, CurrentWindow->GetPositionInScreen().Y);
+		FIntPoint WindowPosition = FIntPoint(CurrentWindow.PositionInScreen.X, CurrentWindow.PositionInScreen.Y);
 		//
 		TopLeft = VectorMin(TopLeft, WindowPosition);
 		//
@@ -135,19 +158,19 @@ void FPixelStreamingVideoInputBackBufferComposited::CompositeWindows()
 	FTextureRHIRef CompositedTexture = UE::PixelStreaming::CreateRHITexture(Extent.X, Extent.Y);
 	for (uint8 i = 0; i < TopLevelWindows.Num(); i++)
 	{
-		TSharedRef<SWindow> CurrentWindow = TopLevelWindows[i];
+		FPixelStreamingVideoInputBackBufferComposited::Window CurrentWindow = TopLevelWindows[i];
 		// Early out if we have an "invalid" window type
-		if (CurrentWindow->GetOpacity() == 0.f || CurrentWindow->GetSizeInScreen() == FVector2f(0, 0))
+		if (CurrentWindow.Opacity == 0.f || CurrentWindow.SizeInScreen == FVector2D(0, 0))
 		{
 			continue;
 		}
-		FIntPoint WindowPosition = FIntPoint(CurrentWindow->GetPositionInScreen().X, CurrentWindow->GetPositionInScreen().Y) - TopLeft;
+		FIntPoint WindowPosition = FIntPoint(CurrentWindow.PositionInScreen.X, CurrentWindow.PositionInScreen.Y) - TopLeft;
 
-		FTextureRHIRef CurrentTexture = TopLevelWindowTextures.FindRef(&CurrentWindow.Get());
+		FTextureRHIRef CurrentTexture = TopLevelWindowTextures.FindRef(CurrentWindow.OwningWindow);
 		// There is only ever one tooltip and as such UE keeps the same texture for each and just rerenders the content
 		// this can lead to small tooltips having a large texture from a previously displayed long tooltip
 		// so we use the tooltips window size which is guaranteed to be correct
-		FIntPoint TextureExtent = VectorMin(CurrentTexture->GetSizeXY(), CurrentWindow->GetSizeInScreen().IntPoint());
+		FIntPoint TextureExtent = VectorMin(CurrentTexture->GetSizeXY(), CurrentWindow.SizeInScreen.IntPoint());
 
 		// Copy our new texture to the compositedframe texture
 		FRHICopyTextureInfo CopyInfo;
diff --git a/Engine/Plugins/Media/PixelStreaming/Source/PixelStreamingEditor/Public/PixelStreamingVideoInputBackBufferComposited.h b/Engine/Plugins/Media/PixelStreaming/Source/PixelStreamingEditor/Public/PixelStreamingVideoInputBackBufferComposited.h
--- a/Engine/Plugins/Media/PixelStreaming/Source/PixelStreamingEditor/Public/PixelStreamingVideoInputBackBufferComposited.h
+++ b/Engine/Plugins/Media/PixelStreaming/Source/PixelStreamingEditor/Public/PixelStreamingVideoInputBackBufferComposited.h
@@ -23,6 +23,24 @@ public:
 	DECLARE_MULTICAST_DELEGATE_OneParam(FOnFrameSizeChanged, TWeakPtr<FIntRect>);
 	FOnFrameSizeChanged OnFrameSizeChanged;
 
+private:
+	// Our struct to keep window information. Use of this struct prevents window information being updated during composition
+	struct Window {
+		Window(FVector2D InPositionInScreen, FVector2D InSizeInScreen, float InOpacity, EWindowType InType, SWindow* InOwningWindow)
+			: PositionInScreen(InPositionInScreen)
+			, SizeInScreen(InSizeInScreen)
+			, Opacity(InOpacity)
+			, Type(InType)
+			, OwningWindow(InOwningWindow)
+		{}
+
+		FVector2D PositionInScreen;
+		FVector2D SizeInScreen;
+		float Opacity;
+		EWindowType Type;
+		SWindow* OwningWindow;
+	};
+
 private:
 	FPixelStreamingVideoInputBackBufferComposited();
 	void CompositeWindows();
@@ -31,11 +49,15 @@ private:
 
 	FDelegateHandle DelegateHandle;
 
-	TArray<TSharedRef<SWindow>> TopLevelWindows;
+	TArray<struct Window> TopLevelWindows;
+
 	TMap<SWindow*, FTextureRHIRef> TopLevelWindowTextures;
+
 	TMap<FString, FTextureRHIRef> StagingTextures;
 
+
 	FCriticalSection TopLevelWindowsCriticalSection;
+
 	TSharedPtr<FIntRect> SharedFrameRect;
 
 private:
