Fix: Exclude node_modules from the recently added Frontend folder when packaging

diff --git a/Engine/Plugins/Media/PixelStreaming/Source/PixelStreaming/PixelStreaming.Build.cs b/Engine/Plugins/Media/PixelStreaming/Source/PixelStreaming/PixelStreaming.Build.cs
--- a/Engine/Plugins/Media/PixelStreaming/Source/PixelStreaming/PixelStreaming.Build.cs
+++ b/Engine/Plugins/Media/PixelStreaming/Source/PixelStreaming/PixelStreaming.Build.cs
@@ -24,6 +24,11 @@ namespace UnrealBuildTool.Rules
 
 			List<string> PathsToIgnore = new List<string>();
 			PathsToIgnore.Add(new DirectoryInfo(DirectoryToAdd + "/node_modules").FullName);
+			PathsToIgnore.Add(new DirectoryInfo(DirectoryToAdd + "/implementations/typescript/node_modules").FullName);
+			PathsToIgnore.Add(new DirectoryInfo(DirectoryToAdd + "/implementations/EpicGames/node_modules").FullName);
+			PathsToIgnore.Add(new DirectoryInfo(DirectoryToAdd + "/implementations/react/node_modules").FullName);
+			PathsToIgnore.Add(new DirectoryInfo(DirectoryToAdd + "/library/node_modules").FullName);
+			PathsToIgnore.Add(new DirectoryInfo(DirectoryToAdd + "/ui-library/node_modules").FullName);
 			PathsToIgnore.Add(new DirectoryInfo(DirectoryToAdd + "/logs").FullName);
 			PathsToIgnore.Add(new DirectoryInfo(DirectoryToAdd + "/platform_scripts/cmd/node").FullName);
 			PathsToIgnore.Add(new DirectoryInfo(DirectoryToAdd + "/platform_scripts/cmd/coturn").FullName);
