Fix use after free. Remove delegate handle on destruction.

diff --git a/Engine/Plugins/Media/PixelStreaming/Source/PixelStreaming/Private/Tests/TestUtils.h b/Engine/Plugins/Media/PixelStreaming/Source/PixelStreaming/Private/Tests/TestUtils.h
--- a/Engine/Plugins/Media/PixelStreaming/Source/PixelStreaming/Private/Tests/TestUtils.h
+++ b/Engine/Plugins/Media/PixelStreaming/Source/PixelStreaming/Private/Tests/TestUtils.h
@@ -22,18 +22,22 @@ namespace UE::PixelStreaming
 		, Callback(InCallback)
 		, bComplete(InbComplete)
 		{
-			Player->OnMessageReceived.AddLambda([this](uint8 Type, const webrtc::DataBuffer& RawBuffer) {
+			MessageReceivedHandle = Player->OnMessageReceived.AddLambda([this](uint8 Type, const webrtc::DataBuffer& RawBuffer) {
 				(Callback)(Type, RawBuffer);
 			});
 		}
 
-		virtual ~FWaitForDataChannelMessageOrTimeout() {}
+		virtual ~FWaitForDataChannelMessageOrTimeout() 
+		{
+			Player->OnMessageReceived.Remove(MessageReceivedHandle);
+		}
 		virtual bool Update() override; 
 	private:
 		double TimeoutSeconds;
 		TSharedPtr<FMockPlayer> Player;
 		TFunction<void(uint8, const webrtc::DataBuffer&)> Callback;
 		TSharedPtr<bool> bComplete;
+		FDelegateHandle MessageReceivedHandle;
 	};
 
 	DEFINE_LATENT_AUTOMATION_COMMAND_THREE_PARAMETER(FSendDataChannelMessageToPlayer, TSharedPtr<IPixelStreamingStreamer>, Streamer, uint8, Id, const FString, Body);
