Fix: Context menu memory leak (move everything to RDG), Editor streamer not respecting passed in streamer id (parse streamer id), input not working when modal is present (add alternate delegate to tick our input devices)

diff --git a/Engine/Plugins/Media/PixelStreaming/Source/PixelStreamingEditor/Private/PixelStreamingEditorModule.cpp b/Engine/Plugins/Media/PixelStreaming/Source/PixelStreamingEditor/Private/PixelStreamingEditorModule.cpp
--- a/Engine/Plugins/Media/PixelStreaming/Source/PixelStreamingEditor/Private/PixelStreamingEditorModule.cpp
+++ b/Engine/Plugins/Media/PixelStreaming/Source/PixelStreamingEditor/Private/PixelStreamingEditorModule.cpp
@@ -60,7 +60,13 @@ void FPixelStreamingEditorModule::ShutdownModule()
 
 void FPixelStreamingEditorModule::InitEditorStreaming(IPixelStreamingModule& Module)
 {
-	EditorStreamer = Module.CreateStreamer("Editor");
+	FString EditorStreamerID;
+	if (!FParse::Value(FCommandLine::Get(), TEXT("PixelStreamingID="), EditorStreamerID))
+	{
+		EditorStreamerID = "Editor";
+	}
+
+	EditorStreamer = Module.CreateStreamer(EditorStreamerID);
 
 	// Give the editor streamer the default url if the user hasn't specified one when launching the editor
 	if (EditorStreamer->GetSignallingServerURL().IsEmpty())
@@ -93,10 +99,6 @@ void FPixelStreamingEditorModule::InitEditorStreaming(IPixelStreamingModule& Mod
 	IMainFrameModule::Get().OnMainFrameCreationFinished().AddLambda([&](TSharedPtr<SWindow> RootWindow, bool bIsRunningStartupDialog) {
 		MaybeResizeEditor(RootWindow);
 
-		// We don't want to show tooltips in render off screen as they're currently broken
-		bool bIsRenderingOffScreen = FParse::Param(FCommandLine::Get(), TEXT("RenderOffScreen"));
-		FSlateApplication::Get().SetAllowTooltips(!bIsRenderingOffScreen);
-
 		if (UE::EditorPixelStreaming::Settings::CVarEditorPixelStreamingStartOnLaunch.GetValueOnAnyThread())
 		{
 			StartStreaming(UE::EditorPixelStreaming::EStreamTypes::Editor);
diff --git a/Engine/Plugins/Media/PixelStreaming/Source/PixelStreamingEditor/Private/PixelStreamingVideoInputBackBufferComposited.cpp b/Engine/Plugins/Media/PixelStreaming/Source/PixelStreamingEditor/Private/PixelStreamingVideoInputBackBufferComposited.cpp
--- a/Engine/Plugins/Media/PixelStreaming/Source/PixelStreamingEditor/Private/PixelStreamingVideoInputBackBufferComposited.cpp
+++ b/Engine/Plugins/Media/PixelStreaming/Source/PixelStreamingEditor/Private/PixelStreamingVideoInputBackBufferComposited.cpp
@@ -1,25 +1,32 @@
 // Copyright Epic Games, Inc. All Rights Reserved.
 
 #include "PixelStreamingVideoInputBackBufferComposited.h"
-#include "Settings.h"
-#include "PixelStreamingEditorUtils.h"
-#include "Utils.h"
-#include "UtilsRender.h"
-#include "ToStringExtensions.h"
-#include "PixelStreamingPrivate.h"
-#include "PixelCaptureInputFrameRHI.h"
-#include "PixelCaptureCapturerRHI.h"
-#include "PixelCaptureCapturerRHIToI420CPU.h"
-#include "PixelCaptureCapturerRHIToI420Compute.h"
-#include "PixelCaptureBufferFormat.h"
+
+#include "EngineModule.h"
 #include "Framework/Application/SlateApplication.h"
-#include "GenericPlatform/GenericWindowDefinition.h"
-#include "Application/SlateApplicationBase.h"
-#include "Widgets/SWindow.h"
+#include "MediaShaders.h"
+#include "PixelCaptureInputFrameRHI.h"
+#include "RenderGraphBuilder.h"
+#include "RenderGraphUtils.h"
+#include "RHI.h"
+#include "ScreenPass.h"
+#include "ScreenRendering.h"
+#include "Utils.h"
 
 DECLARE_LOG_CATEGORY_EXTERN(LogPixelStreamingBackBufferComposited, Log, VeryVerbose);
 DEFINE_LOG_CATEGORY(LogPixelStreamingBackBufferComposited);
 
+/** Pass to push our texture to the rest of the pixel streaming pipeline */
+BEGIN_SHADER_PARAMETER_STRUCT(FExtractCompositedTextureParameters, )
+RDG_TEXTURE_ACCESS(Input, ERHIAccess::CopySrc)
+END_SHADER_PARAMETER_STRUCT()
+
+/** Pass to simple copy one texture to another */
+BEGIN_SHADER_PARAMETER_STRUCT(FCopyToStagingParameters, )
+RDG_TEXTURE_ACCESS(Input, ERHIAccess::CopySrc)
+RDG_TEXTURE_ACCESS(Output, ERHIAccess::CopyDest)
+END_SHADER_PARAMETER_STRUCT()
+
 TSharedPtr<FPixelStreamingVideoInputBackBufferComposited> FPixelStreamingVideoInputBackBufferComposited::Create()
 {
 	TSharedPtr<FPixelStreamingVideoInputBackBufferComposited> NewInput = TSharedPtr<FPixelStreamingVideoInputBackBufferComposited>(new FPixelStreamingVideoInputBackBufferComposited());
@@ -30,8 +37,7 @@ TSharedPtr<FPixelStreamingVideoInputBackBufferComposited> FPixelStreamingVideoIn
 		{
 			FSlateRenderer* Renderer = FSlateApplication::Get().GetRenderer();
 			Input->DelegateHandle = Renderer->OnBackBufferReadyToPresent().AddSP(Input.ToSharedRef(), &FPixelStreamingVideoInputBackBufferComposited::OnBackBufferReady);
-		}
-	});
+		} });
 
 	return NewInput;
 }
@@ -39,156 +45,224 @@ TSharedPtr<FPixelStreamingVideoInputBackBufferComposited> FPixelStreamingVideoIn
 FPixelStreamingVideoInputBackBufferComposited::FPixelStreamingVideoInputBackBufferComposited()
 {
 	SharedFrameRect = MakeShared<FIntRect>();
-	UE::PixelStreaming::DoOnGameThread([this]() {
-		FSlateApplication::Get().OnPreTick().AddLambda([this](float DeltaTime) {
+	UE::PixelStreaming::DoOnGameThread([this]() { FSlateApplication::Get().OnPreTick().AddLambda([this](float DeltaTime) {
 			FScopeLock Lock(&TopLevelWindowsCriticalSection);
-			TArray<TSharedRef<SWindow>> TopLevelSWindows;
-			FSlateApplication::Get().GetAllVisibleWindowsOrdered(TopLevelSWindows);
+			TArray<TSharedRef<SWindow>> TopLevelSlateWindows;
+			FSlateApplication::Get().GetAllVisibleWindowsOrdered(TopLevelSlateWindows);
 
-			TopLevelWindows.Empty();
 			// We store all the necessary window information in structs. This prevents window information from updating
 			// underneath us while we composite and also means we aren't holding on to any shared refs between compositions
-			for(TSharedRef<SWindow>& CurrentWindow : TopLevelSWindows)
+			TArray<FTexturedWindow> TempWindows = TopLevelWindows;
+			TopLevelWindows.Empty();
+
+			for (uint8 i = 0; i < TopLevelSlateWindows.Num(); i++)
 			{
-				TopLevelWindows.Add(FPixelStreamingVideoInputBackBufferComposited::Window(
-					CurrentWindow->GetPositionInScreen(),
-					CurrentWindow->GetSizeInScreen(),
-					CurrentWindow->GetOpacity(),
-					CurrentWindow->GetType(),
-					&CurrentWindow.Get()
-				));
-			}
-		});
-	});
+				TSharedRef<SWindow> CurrentWindow = TopLevelSlateWindows[i];
+				FVector2D WindowExtent = CurrentWindow->GetPositionInScreen() - CurrentWindow->GetSizeInScreen();
+
+				// No need to keep track of "invalid" windows
+				if (CurrentWindow->GetOpacity() == 0.f
+					|| CurrentWindow->GetSizeInScreen() == FVector2D(0, 0)
+					|| CurrentWindow->GetPositionInScreen().X > 16384
+					|| CurrentWindow->GetPositionInScreen().X < -16384
+					|| CurrentWindow->GetPositionInScreen().Y > 16384
+					|| CurrentWindow->GetPositionInScreen().Y < -16384)
+				{
+					continue;
+				}
+
+				TopLevelWindows.Add(FTexturedWindow(CurrentWindow->GetPositionInScreen(), CurrentWindow->GetSizeInScreen(), CurrentWindow->GetOpacity(), CurrentWindow->GetType(), &CurrentWindow.Get()));
+
+				// HACK (william.belcher): This following section is a bit of a nasty work-around to the fact that when a modal is displayed, previous windows (eg the editor) won't
+				// trigger the OnBackBufferReady delegate. This means that we need to check if we have previously seen this top level window, and if we have we need to copy across
+				// its staging texture just in case we won't see it again until after the modal is closed
+				int32 Index = TempWindows.IndexOfByPredicate([&CurrentWindow](FTexturedWindow Window) {
+					return Window.GetOwningWindow() == &CurrentWindow.Get();
+				});
+
+				if (Index != INDEX_NONE)
+				{
+					// This isn't the first time we've seen this window, so copy across its staging texture
+					TopLevelWindows[TopLevelWindows.Num() - 1].SetTexture(TempWindows[Index].GetTexture());
+				}
+			} }); });
 }
 
 FPixelStreamingVideoInputBackBufferComposited::~FPixelStreamingVideoInputBackBufferComposited()
 {
 	if (!IsEngineExitRequested())
 	{
-		UE::PixelStreaming::DoOnGameThread([HandleCopy = DelegateHandle]() {
-			FSlateApplication::Get().GetRenderer()->OnBackBufferReadyToPresent().Remove(HandleCopy);
-		});
+		UE::PixelStreaming::DoOnGameThread([HandleCopy = DelegateHandle]() { FSlateApplication::Get().GetRenderer()->OnBackBufferReadyToPresent().Remove(HandleCopy); });
 	}
 }
 
 void FPixelStreamingVideoInputBackBufferComposited::OnBackBufferReady(SWindow& SlateWindow, const FTexture2DRHIRef& FrameBuffer)
 {
+	/**
+	 * When we receive a texture from this delegate, the texture will undergo a two copy process.
+	 *
+	 * The first copy performed in this function copies the texture we receive to the "Texture" member
+	 * of the FWindow instance corresponding to the window provided. This is necessary as UE sometimes deletes
+	 * before we have a chance to use them when compositing, so we need our own copy.
+	 *
+	 * The second copy is completed within CompositeWindows and applies a render pass to ensure format match
+	 * (editor usually renders in RGB10A2 but WebRTC only supports RGBA8) before copying the texture to appropriate
+	 * location in the composited frame.
+	 *
+	 * Finally, the composited frame is extracted from the RDG pipeline and we send it on its way through the PixelCapturer
+	 *
+	 */
+	UE_LOG(LogPixelStreamingBackBufferComposited, Verbose, TEXT("Type: %s"), *SlateWindow.GetTitle().ToString());
+
 	{
 		FScopeLock Lock(&TopLevelWindowsCriticalSection);
 		if (TopLevelWindows.IsEmpty())
 		{
 			return;
 		}
-	}
 
-	UE_LOG(LogPixelStreamingBackBufferComposited, Verbose, TEXT("=== Window Rendered ==="));
-	UE_LOG(LogPixelStreamingBackBufferComposited, Verbose, TEXT("Type: %s"), UE::EditorPixelStreaming::ToString(SlateWindow.GetType()));
+		// Find the index of the window that called this delegate in our array of windows + textures
+		int32 Index = TopLevelWindows.IndexOfByPredicate([&SlateWindow](FTexturedWindow Window) { return Window.GetOwningWindow() == &SlateWindow; });
 
-	FIntPoint FrameSize = FrameBuffer->GetSizeXY();
-	FTextureRHIRef StagingTexture;
-	if (SlateWindow.GetType() != EWindowType::ToolTip)
-	{
-		FString Hash = UE::EditorPixelStreaming::HashWindow(SlateWindow, FrameBuffer);
-		StagingTexture = StagingTextures.FindRef(Hash);
-		if (!StagingTexture.IsValid())
+		if (Index == INDEX_NONE)
 		{
-			UE_LOG(LogPixelStreamingBackBufferComposited, Verbose, TEXT("Creating new staging texture: %dx%d"), FrameSize.X, FrameSize.Y);
-			StagingTexture = UE::PixelStreaming::CreateRHITexture(FrameSize.X, FrameSize.Y);
-			StagingTextures.Add(Hash, StagingTexture);
+			// Early out if we've received a texture without knowing if it's a part of GetAllVisibleWindowsOrdered()
+			return;
 		}
-	}
-	else
-	{
-		// Always create a new staging texture for tooltips due to their odd handling
-		StagingTexture = UE::PixelStreaming::CreateRHITexture(FrameSize.X, FrameSize.Y);
-	}
 
-	FRHICommandListImmediate& RHICmdList = FRHICommandListExecutor::GetImmediateCommandList();
-	// Re-render FrameBuffer to StagingTexture (ensure format match)
-	UE::PixelStreaming::CopyTextureRDG(RHICmdList, FrameBuffer, StagingTexture);
-
-	TopLevelWindowTextures.Add(&SlateWindow, StagingTexture);
-	TopLevelWindowsCriticalSection.Lock();
-	// Check that we have received a texture from every window in the TopLevelWindows array
-	uint8 NumWindowsRendered = 0;
-	for (FPixelStreamingVideoInputBackBufferComposited::Window TopWindow : TopLevelWindows)
-	{
-		if (TopLevelWindowTextures.FindRef(TopWindow.OwningWindow))
+		FRHICommandListImmediate& RHICmdList = FRHICommandListExecutor::GetImmediateCommandList();
+		FRDGBuilder GraphBuilder(RHICmdList);
+
+		// Register an external RDG texture from the provided frame buffer
+		FRDGTextureRef InputTexture = GraphBuilder.RegisterExternalTexture(CreateRenderTarget(FrameBuffer, *SlateWindow.GetTitle().ToString()));
+
+		// Create an internal RDG texture with the same extent and format as the source
+		FRDGTextureRef OutputTexture = GraphBuilder.CreateTexture(FRDGTextureDesc::Create2D(FrameBuffer->GetDesc().Extent, FrameBuffer->GetDesc().Format, FClearValueBinding::None, ETextureCreateFlags::Shared | ETextureCreateFlags::RenderTargetable), TEXT("VideoInputBackBufferCompositedStaging"));
+		// Bit cheeky, but when attempting to create two textures with the same description, RDG was just re-allocating which would lead to flickering. By converting to external, we force immediate allocation of the underlying pooled resource
+		TopLevelWindows[Index].SetTexture(GraphBuilder.ConvertToExternalTexture(OutputTexture));
+
+		// Initialise our copy paramaters
+		FCopyToStagingParameters* PassParameters = GraphBuilder.AllocParameters<FCopyToStagingParameters>();
+		PassParameters->Input = InputTexture;
+		PassParameters->Output = OutputTexture;
+
+		GraphBuilder.AddPass(
+			RDG_EVENT_NAME("VideoInputBackBufferCopyToStaging (%s)", *SlateWindow.GetTitle().ToString()),
+			PassParameters,
+			ERDGPassFlags::Readback,
+			[InputTexture, OutputTexture, Index, this](FRHICommandList& RHICmdList) {
+				// Simple texture copy as we know pixel format is the same
+				RHICmdList.CopyTexture(InputTexture->GetRHI(), OutputTexture->GetRHI(), {});
+			});
+
+		GraphBuilder.Execute();
+
+		// Check all of our windows have a texture
+		bool bSkipComposite = TopLevelWindows.ContainsByPredicate([](FTexturedWindow Window) { return !Window.GetTexture(); });
+
+		if (!bSkipComposite)
 		{
-			++NumWindowsRendered;
+			// All windows have a texture so we can composite
+			CompositeWindows();
 		}
 	}
-
-	if (NumWindowsRendered == TopLevelWindows.Num())
-	{
-		CompositeWindows();
-	}
-	TopLevelWindowsCriticalSection.Unlock();
 }
 
 void FPixelStreamingVideoInputBackBufferComposited::CompositeWindows()
 {
 	FRHICommandListImmediate& RHICmdList = FRHICommandListExecutor::GetImmediateCommandList();
+	FRDGBuilder GraphBuilder(RHICmdList);
 
 	// Process all of the windows we will need to render. This processing step finds the extents of the
 	// composited texture as well as the top-left point
 	FIntPoint TopLeft = FIntPoint(MAX_int32, MAX_int32);
 	FIntPoint BottomRight = FIntPoint(MIN_int32, MIN_int32);
-	for (uint8 i = 0; i < TopLevelWindows.Num(); i++)
+	for (FTexturedWindow& CurrentWindow : TopLevelWindows)
 	{
-		FPixelStreamingVideoInputBackBufferComposited::Window CurrentWindow = TopLevelWindows[i];
-		// Early out if we have an "invalid" window type
-		if (CurrentWindow.Opacity == 0.f || CurrentWindow.SizeInScreen == FVector2D(0, 0))
-		{
-			continue;
-		}
-
-		FTextureRHIRef CurrentTexture = TopLevelWindowTextures.FindRef(CurrentWindow.OwningWindow);
-		FIntPoint TextureExtent = VectorMin(CurrentTexture->GetSizeXY(), CurrentWindow.SizeInScreen.IntPoint());
-
-		FIntPoint WindowPosition = FIntPoint(CurrentWindow.PositionInScreen.X, CurrentWindow.PositionInScreen.Y);
-		//
+		FIntPoint TextureExtent = VectorMin(CurrentWindow.GetTexture()->GetDesc().Extent, CurrentWindow.GetSizeInScreen().IntPoint());
+		FIntPoint WindowPosition = FIntPoint(CurrentWindow.GetPositionInScreen().X, CurrentWindow.GetPositionInScreen().Y);
+		// Update the top left to be the element-wise minimum
 		TopLeft = VectorMin(TopLeft, WindowPosition);
-		//
+		// Update the bottom right to be the element-wise maximum
 		BottomRight = VectorMax(BottomRight, (WindowPosition + TextureExtent));
 	}
 
-	FIntPoint Extent = BottomRight - TopLeft;
-	FTextureRHIRef CompositedTexture = UE::PixelStreaming::CreateRHITexture(Extent.X, Extent.Y);
-	for (uint8 i = 0; i < TopLevelWindows.Num(); i++)
+	// Clamp the texture dimensions to ensure no RHI crashes
+	// Create an RDG texture that is the size of our extent for use as the composited frame
+	FRDGTextureRef CompositedTexture = GraphBuilder.CreateTexture(FRDGTextureDesc::Create2D(VectorMin(BottomRight - TopLeft, FIntPoint(16384, 16384)), EPixelFormat::PF_B8G8R8A8, FClearValueBinding::None, ETextureCreateFlags::Shared | ETextureCreateFlags::RenderTargetable), TEXT("VideoInputBackBufferCompositedCompositedTexture"));
+
+	// Shader globals used in the conversion pass
+	FGlobalShaderMap* GlobalShaderMap = GetGlobalShaderMap(GMaxRHIFeatureLevel);
+	TShaderMapRef<FScreenPassVS> VertexShader(GlobalShaderMap);
+
+	// In cases where texture is converted from a format that doesn't have A channel, we want to force set it to 1.
+	int32 ConversionOperation = 0; // None
+	FModifyAlphaSwizzleRgbaPS::FPermutationDomain PermutationVector;
+	PermutationVector.Set<FModifyAlphaSwizzleRgbaPS::FConversionOp>(ConversionOperation);
+
+	for (FTexturedWindow& CurrentWindow : TopLevelWindows)
 	{
-		FPixelStreamingVideoInputBackBufferComposited::Window CurrentWindow = TopLevelWindows[i];
-		// Early out if we have an "invalid" window type
-		if (CurrentWindow.Opacity == 0.f || CurrentWindow.SizeInScreen == FVector2D(0, 0))
+		FIntPoint WindowPosition = FIntPoint(CurrentWindow.GetPositionInScreen().X, CurrentWindow.GetPositionInScreen().Y) - TopLeft;
+
+		FRHITexture* CurrentTexture = CurrentWindow.GetTexture()->GetRHI();
+
+		FRDGTextureRef InputTexture = GraphBuilder.RegisterExternalTexture(CreateRenderTarget(CurrentTexture, TEXT("VideoInputBackBufferCompositedStaging")));
+		// There is only ever one tooltip and as such UE keeps the same texture for each and just rerenders the
+		// content this can lead to small tooltips having a large texture from a previously displayed long tooltip
+		// so we use the tooltips window size which is guaranteed to be correct
+		FIntPoint Extent = VectorMin(CurrentTexture->GetDesc().Extent, CurrentWindow.GetSizeInScreen().IntPoint());
+
+		// Ensure we have a valid extent (texture or window > 0,0)
+		if (Extent.X == 0 || Extent.Y == 0)
 		{
 			continue;
 		}
-		FIntPoint WindowPosition = FIntPoint(CurrentWindow.PositionInScreen.X, CurrentWindow.PositionInScreen.Y) - TopLeft;
 
-		FTextureRHIRef CurrentTexture = TopLevelWindowTextures.FindRef(CurrentWindow.OwningWindow);
-		// There is only ever one tooltip and as such UE keeps the same texture for each and just rerenders the content
-		// this can lead to small tooltips having a large texture from a previously displayed long tooltip
-		// so we use the tooltips window size which is guaranteed to be correct
-		FIntPoint TextureExtent = VectorMin(CurrentTexture->GetSizeXY(), CurrentWindow.SizeInScreen.IntPoint());
-
-		// Copy our new texture to the compositedframe texture
-		FRHICopyTextureInfo CopyInfo;
-		CopyInfo.Size.X = TextureExtent.X;
-		CopyInfo.Size.Y = TextureExtent.Y;
-		CopyInfo.Size.Z = 1;
-		CopyInfo.DestPosition.X = WindowPosition.X;
-		CopyInfo.DestPosition.Y = WindowPosition.Y;
-		CopyInfo.DestPosition.Z = 0;
-		// source and dest are the same. simple copy
-		RHICmdList.Transition(FRHITransitionInfo(CurrentTexture, ERHIAccess::Unknown, ERHIAccess::CopySrc));
-		RHICmdList.Transition(FRHITransitionInfo(CompositedTexture, ERHIAccess::Unknown, ERHIAccess::CopyDest));
-		RHICmdList.CopyTexture(CurrentTexture, CompositedTexture, CopyInfo);
+		// Configure our viewports appropriately
+		FScreenPassTextureViewport InputViewport(InputTexture, FIntRect(FIntPoint(0, 0), Extent));
+		FScreenPassTextureViewport OutputViewport(CompositedTexture, FIntRect(WindowPosition, WindowPosition + Extent));
+
+		// Rectangle area to use from the source texture
+		const FIntRect ViewRect(FIntPoint(0, 0), Extent);
+
+		// Dummy ViewFamily/ViewInfo created to use built in Draw Screen/Texture Pass
+		FSceneViewFamily ViewFamily(FSceneViewFamily::ConstructionValues(nullptr, nullptr, FEngineShowFlags(ESFIM_Game))
+										.SetTime(FGameTime())
+										.SetGammaCorrection(1.0f));
+		FSceneViewInitOptions ViewInitOptions;
+		ViewInitOptions.ViewFamily = &ViewFamily;
+		ViewInitOptions.SetViewRectangle(ViewRect);
+		ViewInitOptions.ViewOrigin = FVector::ZeroVector;
+		ViewInitOptions.ViewRotationMatrix = FMatrix::Identity;
+		ViewInitOptions.ProjectionMatrix = FMatrix::Identity;
+
+		GetRendererModule().CreateAndInitSingleView(RHICmdList, &ViewFamily, &ViewInitOptions);
+		const FSceneView& View = *ViewFamily.Views[0];
+
+		TShaderMapRef<FModifyAlphaSwizzleRgbaPS> PixelShader(GlobalShaderMap, PermutationVector);
+		FModifyAlphaSwizzleRgbaPS::FParameters* PixelShaderParameters = PixelShader->AllocateAndSetParameters(GraphBuilder, InputTexture, CompositedTexture);
+		// Add screen pass to convert whatever format the editor produces to BGRA8
+		AddDrawScreenPass(GraphBuilder, RDG_EVENT_NAME("VideoInputBackBufferCompositedSwizzle"), View, OutputViewport, InputViewport, VertexShader, PixelShader, PixelShaderParameters);
 	}
 
-	// Our composition is complete, send it to along the pipeline
-	OnFrame(FPixelCaptureInputFrameRHI(CompositedTexture));
-	// Update the default streamer to let it know our compositedframe size and position. This way it can correctly scale input from the browser
+	// Final pass to extract the composited frames underlying RHI resource for passing to the rest of the pixel streaming pipeline
+	FTextureRHIRef OutTexture = nullptr;
+	FExtractCompositedTextureParameters* PassParameters = GraphBuilder.AllocParameters<FExtractCompositedTextureParameters>();
+	PassParameters->Input = CompositedTexture;
+	GraphBuilder.AddPass(
+		RDG_EVENT_NAME("PushCompositedFrame"),
+		PassParameters,
+		ERDGPassFlags::Readback,
+		[CompositedTexture, &OutTexture, this](FRHICommandList& RHICmdList) {
+			// Our composition is complete out the underlying rhi resource
+			OutTexture = CompositedTexture->GetRHI();
+		});
+
+	GraphBuilder.Execute();
+
+	OnFrame(FPixelCaptureInputFrameRHI(OutTexture));
+
+	// Update any subscribed streamers to let them know our composited frame size and position. This way it can correctly scale input from the browser
 	*SharedFrameRect.Get() = FIntRect(TopLeft, BottomRight);
 	OnFrameSizeChanged.Broadcast(SharedFrameRect);
 }
diff --git a/Engine/Plugins/Media/PixelStreaming/Source/PixelStreamingEditor/Public/PixelStreamingVideoInputBackBufferComposited.h b/Engine/Plugins/Media/PixelStreaming/Source/PixelStreamingEditor/Public/PixelStreamingVideoInputBackBufferComposited.h
--- a/Engine/Plugins/Media/PixelStreaming/Source/PixelStreamingEditor/Public/PixelStreamingVideoInputBackBufferComposited.h
+++ b/Engine/Plugins/Media/PixelStreaming/Source/PixelStreamingEditor/Public/PixelStreamingVideoInputBackBufferComposited.h
@@ -3,11 +3,10 @@
 #pragma once
 
 #include "PixelStreamingVideoInputRHI.h"
+
 #include "Widgets/SWindow.h"
-#include "RHI.h"
+#include "RendererInterface.h"
 #include "Delegates/IDelegateInstance.h"
-#include "Widgets/SWindow.h"
-#include "GenericPlatform/GenericWindowDefinition.h"
 
 /*
  * Use this if you want to send the full UE editor as video input.
@@ -15,64 +14,69 @@
 class PIXELSTREAMINGEDITOR_API FPixelStreamingVideoInputBackBufferComposited : public FPixelStreamingVideoInputRHI
 {
 public:
-	static TSharedPtr<FPixelStreamingVideoInputBackBufferComposited> Create();
-	virtual ~FPixelStreamingVideoInputBackBufferComposited();
+    static TSharedPtr<FPixelStreamingVideoInputBackBufferComposited> Create();
+    virtual ~FPixelStreamingVideoInputBackBufferComposited();
 
-	virtual FString ToString() override;
+    virtual FString ToString() override;
 
-	DECLARE_MULTICAST_DELEGATE_OneParam(FOnFrameSizeChanged, TWeakPtr<FIntRect>);
-	FOnFrameSizeChanged OnFrameSizeChanged;
+    DECLARE_MULTICAST_DELEGATE_OneParam(FOnFrameSizeChanged, TWeakPtr<FIntRect>);
+    FOnFrameSizeChanged OnFrameSizeChanged;
 
 private:
-	// Our struct to keep window information. Use of this struct prevents window information being updated during composition
-	struct Window {
-		Window(FVector2D InPositionInScreen, FVector2D InSizeInScreen, float InOpacity, EWindowType InType, SWindow* InOwningWindow)
-			: PositionInScreen(InPositionInScreen)
-			, SizeInScreen(InSizeInScreen)
-			, Opacity(InOpacity)
-			, Type(InType)
-			, OwningWindow(InOwningWindow)
-		{}
-
-		FVector2D PositionInScreen;
-		FVector2D SizeInScreen;
-		float Opacity;
-		EWindowType Type;
-		SWindow* OwningWindow;
-	};
+    // Our class to keep window and texture information. Use of this struct prevents SWindow propeties being updated composition as well
+    // as preventing deletion of our staging textures pre composition
+    class FTexturedWindow
+    {
+    public:
+        FTexturedWindow(FVector2D InPositionInScreen, FVector2D InSizeInScreen, float InOpacity, EWindowType InType, SWindow *InOwningWindow)
+            : PositionInScreen(InPositionInScreen), SizeInScreen(InSizeInScreen), Opacity(InOpacity), Type(InType), Texture(nullptr), OwningWindow(InOwningWindow)
+        {
+        }
+
+        FVector2D GetPositionInScreen() { return PositionInScreen; }
+        FVector2D GetSizeInScreen() { return SizeInScreen; }
+        float GetOpacity() { return Opacity; }
+        EWindowType GetType() { return Type; }
+        SWindow *GetOwningWindow() { return OwningWindow; }
+        TRefCountPtr<IPooledRenderTarget> &GetTexture() { return Texture; }
+        void SetTexture(TRefCountPtr<IPooledRenderTarget> InTexture) { Texture = InTexture; }
+
+    private:
+        FVector2D PositionInScreen;
+        FVector2D SizeInScreen;
+        float Opacity;
+        EWindowType Type;
+        TRefCountPtr<IPooledRenderTarget> Texture;
+        SWindow *OwningWindow;
+    };
 
 private:
-	FPixelStreamingVideoInputBackBufferComposited();
-	void CompositeWindows();
-
-	void OnBackBufferReady(SWindow& SlateWindow, const FTexture2DRHIRef& FrameBuffer);
-
-	FDelegateHandle DelegateHandle;
-
-	TArray<struct Window> TopLevelWindows;
+    FPixelStreamingVideoInputBackBufferComposited();
+    void CompositeWindows();
 
-	TMap<SWindow*, FTextureRHIRef> TopLevelWindowTextures;
+    void OnBackBufferReady(SWindow &SlateWindow, const FTexture2DRHIRef &FrameBuffer);
 
-	TMap<FString, FTextureRHIRef> StagingTextures;
+    FDelegateHandle DelegateHandle;
 
+    TArray<FTexturedWindow> TopLevelWindows;
 
-	FCriticalSection TopLevelWindowsCriticalSection;
+    FCriticalSection TopLevelWindowsCriticalSection;
 
-	TSharedPtr<FIntRect> SharedFrameRect;
+    TSharedPtr<FIntRect> SharedFrameRect;
 
 private:
-	// Util functions for 2D vectors
-	template <class T>
-	T VectorMax(const T A, const T B)
-	{
-		// Returns the component-wise maximum of two vectors
-		return T(FMath::Max(A.X, B.X), FMath::Max(A.Y, B.Y));
-	}
-
-	template <class T>
-	T VectorMin(const T A, const T B)
-	{
-		// Returns the component-wise minimum of two vectors
-		return T(FMath::Min(A.X, B.X), FMath::Min(A.Y, B.Y));
-	}
+    // Util functions for 2D vectors
+    template <class T>
+    T VectorMax(const T A, const T B)
+    {
+        // Returns the component-wise maximum of two vectors
+        return T(FMath::Max(A.X, B.X), FMath::Max(A.Y, B.Y));
+    }
+
+    template <class T>
+    T VectorMin(const T A, const T B)
+    {
+        // Returns the component-wise minimum of two vectors
+        return T(FMath::Min(A.X, B.X), FMath::Min(A.Y, B.Y));
+    }
 };
\ No newline at end of file
diff --git a/Engine/Plugins/Media/PixelStreaming/Source/PixelStreamingInput/Private/PixelStreamingInputDevice.cpp b/Engine/Plugins/Media/PixelStreaming/Source/PixelStreamingInput/Private/PixelStreamingInputDevice.cpp
--- a/Engine/Plugins/Media/PixelStreaming/Source/PixelStreamingInput/Private/PixelStreamingInputDevice.cpp
+++ b/Engine/Plugins/Media/PixelStreaming/Source/PixelStreamingInput/Private/PixelStreamingInputDevice.cpp
@@ -1,6 +1,7 @@
 // Copyright Epic Games, Inc. All Rights Reserved.
 
 #include "PixelStreamingInputDevice.h"
+#include "Framework/Application/SlateApplication.h"
 
 TSharedPtr<UE::PixelStreamingInput::FPixelStreamingInputDevice> UE::PixelStreamingInput::FPixelStreamingInputDevice::InputDevice;
 
@@ -23,6 +24,8 @@ namespace UE::PixelStreamingInput
 
 	FPixelStreamingInputDevice::FPixelStreamingInputDevice()
 	{
+		// This is imperative for editor streaming as when a modal is open or we've hit a BP breakpoint, the engine tick loop will not run, so instead we rely on this delegate to tick for us
+		FSlateApplication::Get().OnPreTick().AddRaw(this, &FPixelStreamingInputDevice::Tick);
 	}
 
 	void FPixelStreamingInputDevice::AddInputHandler(TSharedPtr<IPixelStreamingInputHandler> InputHandler)
