Some fixes for the failing horde unit tests.

diff --git a/Engine/Plugins/Media/PixelStreaming/Source/PixelStreaming/Private/Tests/CodecTest.cpp b/Engine/Plugins/Media/PixelStreaming/Source/PixelStreaming/Private/Tests/CodecTest.cpp
deleted file mode 100644
diff --git a/Engine/Plugins/Media/PixelStreaming/Source/PixelStreaming/Private/Tests/CodecUtils.h b/Engine/Plugins/Media/PixelStreaming/Source/PixelStreaming/Private/Tests/CodecUtils.h
new file mode 100644
--- /dev/null
+++ b/Engine/Plugins/Media/PixelStreaming/Source/PixelStreaming/Private/Tests/CodecUtils.h
@@ -0,0 +1,77 @@
+// Copyright Epic Games, Inc. All Rights Reserved.
+
+#pragma once
+
+#include "Misc/AutomationTest.h"
+#include "TestUtils.h"
+
+#if WITH_DEV_AUTOMATION_TESTS
+
+namespace
+{
+	using namespace UE::PixelStreaming;
+
+	void DoFrameReceiveTest()
+	{
+		int32 StreamerPort = 7564;
+		int32 PlayerPort = 4578;
+
+		FMockVideoFrameConfig FrameConfig = { 128 /*Width*/, 128 /*Height*/, 255 /*Y*/, 137 /*U*/, 216 /*V*/ };
+
+		TSharedPtr<UE::PixelStreamingServers::IServer> SignallingServer = CreateSignallingServer(StreamerPort, PlayerPort);
+
+		TSharedPtr<IPixelStreamingStreamer> Streamer = CreateStreamer(StreamerPort);
+		TSharedPtr<FPixelStreamingVideoInputI420> VideoInput = MakeShared<FPixelStreamingVideoInputI420>();
+		VideoInput->AddOutputFormat(PixelCaptureBufferFormat::FORMAT_I420);
+		Streamer->SetVideoInput(VideoInput);
+
+		TSharedPtr<FMockPlayer> Player = CreatePlayer(FMockPlayer::EMode::AcceptOffers);
+		TSharedPtr<FMockVideoSink> VideoSink = MakeShared<FMockVideoSink>();
+		Player->VideoSink = VideoSink;
+
+		Streamer->StartStreaming();
+
+		ADD_LATENT_AUTOMATION_COMMAND(FConnectPlayerAfterStreamerConnectedOrTimeout(5.0, Streamer, Player, PlayerPort))
+		ADD_LATENT_AUTOMATION_COMMAND(FWaitForPlayerConnectedOrTimeout(5.0, Player, PlayerPort))
+		ADD_LATENT_AUTOMATION_COMMAND(FWaitForICEConnectedOrTimeout(5.0, Player))
+		ADD_LATENT_AUTOMATION_COMMAND(FSendSolidColorFrame(VideoInput, FrameConfig))
+		ADD_LATENT_AUTOMATION_COMMAND(FWaitForFrameReceived(5.0, VideoSink, FrameConfig))
+		ADD_LATENT_AUTOMATION_COMMAND(FCleanupAll(SignallingServer, Streamer, Player))
+	}
+
+	void DoFrameResizeMultipleTimesTest()
+	{
+		int32 StreamerPort = 7564;
+		int32 PlayerPort = 4578;
+
+		TSharedPtr<UE::PixelStreamingServers::IServer> SignallingServer = CreateSignallingServer(StreamerPort, PlayerPort);
+
+		TSharedPtr<IPixelStreamingStreamer> Streamer = CreateStreamer(StreamerPort);
+		TSharedPtr<FPixelStreamingVideoInputI420> VideoInput = MakeShared<FPixelStreamingVideoInputI420>();
+		VideoInput->AddOutputFormat(PixelCaptureBufferFormat::FORMAT_I420);
+		Streamer->SetVideoInput(VideoInput);
+
+		TSharedPtr<FMockPlayer> Player = CreatePlayer(FMockPlayer::EMode::AcceptOffers);
+		TSharedPtr<FMockVideoSink> VideoSink = MakeShared<FMockVideoSink>();
+		Player->VideoSink = VideoSink;
+
+		// Note: Important to couple framerate as we are manually passing frames and don't want any cached frames
+		Streamer->SetCoupleFramerate(true);
+		Streamer->StartStreaming();
+
+		ADD_LATENT_AUTOMATION_COMMAND(FConnectPlayerAfterStreamerConnectedOrTimeout(5.0, Streamer, Player, PlayerPort))
+		ADD_LATENT_AUTOMATION_COMMAND(FWaitForPlayerConnectedOrTimeout(5.0, Player, PlayerPort))
+		ADD_LATENT_AUTOMATION_COMMAND(FWaitForICEConnectedOrTimeout(5.0, Player))
+
+		for (int Res = 2; Res < 512; Res *= 2)
+		{
+			FMockVideoFrameConfig FrameConfig = { Res /*Width*/, Res /*Height*/, 255 /*Y*/, 0 /*U*/, 255 /*V*/ };
+			ADD_LATENT_AUTOMATION_COMMAND(FSendSolidColorFrame(VideoInput, FrameConfig))
+			ADD_LATENT_AUTOMATION_COMMAND(FWaitForFrameReceived(5.0, VideoSink, FrameConfig))
+		}
+
+		ADD_LATENT_AUTOMATION_COMMAND(FCleanupAll(SignallingServer, Streamer, Player))
+	}
+} // namespace UE::PixelStreaming
+
+#endif // WITH_DEV_AUTOMATION_TESTS
diff --git a/Engine/Plugins/Media/PixelStreaming/Source/PixelStreaming/Private/Tests/DataChannelTest.cpp b/Engine/Plugins/Media/PixelStreaming/Source/PixelStreaming/Private/Tests/DataChannelTest.cpp
--- a/Engine/Plugins/Media/PixelStreaming/Source/PixelStreaming/Private/Tests/DataChannelTest.cpp
+++ b/Engine/Plugins/Media/PixelStreaming/Source/PixelStreaming/Private/Tests/DataChannelTest.cpp
@@ -53,7 +53,7 @@ namespace UE::PixelStreaming
 
 		ADD_LATENT_AUTOMATION_COMMAND(FConnectPlayerAfterStreamerConnectedOrTimeout(5.0, Streamer, Player, PlayerPort))
 		ADD_LATENT_AUTOMATION_COMMAND(FWaitForPlayerConnectedOrTimeout(5.0, Player, PlayerPort))
-		ADD_LATENT_AUTOMATION_COMMAND(FWaitForICEConnectedOrTimeout(5.0, Player))
+		ADD_LATENT_AUTOMATION_COMMAND(FWaitForDataChannelOrTimeout(5.0, Player))
 		ADD_LATENT_AUTOMATION_COMMAND(FSendDataChannelMessageToStreamer(Player, ToStreamerEchoId, EchoContent))
 		ADD_LATENT_AUTOMATION_COMMAND(FWaitForDataChannelMessageOrTimeout(15.0, Player, Callback, bComplete))
 		ADD_LATENT_AUTOMATION_COMMAND(FCleanupAll(SignallingServer, Streamer, Player))
diff --git a/Engine/Plugins/Media/PixelStreaming/Source/PixelStreaming/Private/Tests/MockPlayer.h b/Engine/Plugins/Media/PixelStreaming/Source/PixelStreaming/Private/Tests/MockPlayer.h
--- a/Engine/Plugins/Media/PixelStreaming/Source/PixelStreaming/Private/Tests/MockPlayer.h
+++ b/Engine/Plugins/Media/PixelStreaming/Source/PixelStreaming/Private/Tests/MockPlayer.h
@@ -118,7 +118,6 @@ namespace UE::PixelStreaming
 
 		virtual void OnSignallingConfig(const webrtc::PeerConnectionInterface::RTCConfiguration& Config) override
 		{
-
 			PeerConnection = FPixelStreamingPeerConnection::Create(Config);
 
 			PeerConnection->OnEmitIceCandidate.AddLambda([this](const webrtc::IceCandidateInterface* Candidate) {
@@ -132,6 +131,8 @@ namespace UE::PixelStreaming
 				DataChannel->OnMessageReceived.AddLambda([this](uint8 Type, const webrtc::DataBuffer& RawBuffer) {
 					OnMessageReceived.Broadcast(Type, RawBuffer);
 				});
+
+				HaveDataChannel = true;
 			});
 
 			PeerConnection->OnIceStateChanged.AddLambda([this](webrtc::PeerConnectionInterface::IceConnectionState NewState) {
@@ -211,6 +212,7 @@ namespace UE::PixelStreaming
 		TUniquePtr<FPixelStreamingPeerConnection> PeerConnection;
 		TSharedPtr<FPixelStreamingDataChannel> DataChannel;
 		bool Completed = false;
+		bool HaveDataChannel = false;
 		TSharedPtr<rtc::VideoSinkInterface<webrtc::VideoFrame>> VideoSink;
 	};
 
diff --git a/Engine/Plugins/Media/PixelStreaming/Source/PixelStreaming/Private/Tests/ProtocolTest.cpp b/Engine/Plugins/Media/PixelStreaming/Source/PixelStreaming/Private/Tests/ProtocolTest.cpp
--- a/Engine/Plugins/Media/PixelStreaming/Source/PixelStreaming/Private/Tests/ProtocolTest.cpp
+++ b/Engine/Plugins/Media/PixelStreaming/Source/PixelStreaming/Private/Tests/ProtocolTest.cpp
@@ -119,7 +119,7 @@ namespace UE::PixelStreaming
 
 		ADD_LATENT_AUTOMATION_COMMAND(FConnectPlayerAfterStreamerConnectedOrTimeout(5.0, Streamer, Player, PlayerPort))
 		ADD_LATENT_AUTOMATION_COMMAND(FWaitForPlayerConnectedOrTimeout(5.0, Player, PlayerPort))
-		ADD_LATENT_AUTOMATION_COMMAND(FWaitForICEConnectedOrTimeout(5.0, Player))
+		ADD_LATENT_AUTOMATION_COMMAND(FWaitForDataChannelOrTimeout(5.0, Player))
 		ADD_LATENT_AUTOMATION_COMMAND(FSendCustomMessageToStreamer(Player, Message.GetID(), 1337))
 		ADD_LATENT_AUTOMATION_COMMAND(FWaitForDataChannelMessageOrTimeout(15.0, Player, Callback, bComplete))
 		ADD_LATENT_AUTOMATION_COMMAND(FCleanupAll(SignallingServer, Streamer, Player))
diff --git a/Engine/Plugins/Media/PixelStreaming/Source/PixelStreaming/Private/Tests/TestUtils.cpp b/Engine/Plugins/Media/PixelStreaming/Source/PixelStreaming/Private/Tests/TestUtils.cpp
--- a/Engine/Plugins/Media/PixelStreaming/Source/PixelStreaming/Private/Tests/TestUtils.cpp
+++ b/Engine/Plugins/Media/PixelStreaming/Source/PixelStreaming/Private/Tests/TestUtils.cpp
@@ -149,6 +149,28 @@ namespace UE::PixelStreaming
 		return true;
 	}
 
+	bool FWaitForDataChannelOrTimeout::Update()
+	{
+		// If no signalling we can early exit.
+		if (!OutPlayer->IsSignallingConnected())
+		{
+			UE_LOG(LogPixelStreaming, Error, TEXT("Early exiting waiting for data channel as player is not connected to signalling server."));
+			return true;
+		}
+
+		if (OutPlayer)
+		{
+			double DeltaTime = FPlatformTime::Seconds() - StartTime;
+			if (DeltaTime > TimeoutSeconds)
+			{
+				UE_LOG(LogPixelStreaming, Error, TEXT("Timed out waiting for data channel between streamer and player."));
+				return true;
+			}
+			return OutPlayer->HaveDataChannel;
+		}
+		return true;
+	}
+
 	bool FWaitForStreamerConnectedOrTimeout::Update()
 	{
 		if (OutStreamer->IsSignallingConnected())
diff --git a/Engine/Plugins/Media/PixelStreaming/Source/PixelStreaming/Private/Tests/TestUtils.h b/Engine/Plugins/Media/PixelStreaming/Source/PixelStreaming/Private/Tests/TestUtils.h
--- a/Engine/Plugins/Media/PixelStreaming/Source/PixelStreaming/Private/Tests/TestUtils.h
+++ b/Engine/Plugins/Media/PixelStreaming/Source/PixelStreaming/Private/Tests/TestUtils.h
@@ -46,6 +46,7 @@ namespace UE::PixelStreaming
 	DEFINE_LATENT_AUTOMATION_COMMAND_TWO_PARAMETER(FSendSolidColorFrame, TSharedPtr<FPixelStreamingVideoInputI420>, VideoInput, FMockVideoFrameConfig, FrameConfig);
 	DEFINE_LATENT_AUTOMATION_COMMAND_THREE_PARAMETER(FWaitForFrameReceived, double, TimeoutSeconds, TSharedPtr<FMockVideoSink>, VideoSink, FMockVideoFrameConfig, FrameConfig);
 	DEFINE_LATENT_AUTOMATION_COMMAND_TWO_PARAMETER(FWaitForICEConnectedOrTimeout, double, TimeoutSeconds, TSharedPtr<FMockPlayer>, OutPlayer);
+	DEFINE_LATENT_AUTOMATION_COMMAND_TWO_PARAMETER(FWaitForDataChannelOrTimeout, double, TimeoutSeconds, TSharedPtr<FMockPlayer>, OutPlayer);
 	DEFINE_LATENT_AUTOMATION_COMMAND_TWO_PARAMETER(FWaitForStreamerConnectedOrTimeout, double, TimeoutSeconds, TSharedPtr<IPixelStreamingStreamer>, OutStreamer);
 	DEFINE_LATENT_AUTOMATION_COMMAND_FOUR_PARAMETER(FConnectPlayerAfterStreamerConnectedOrTimeout, double, TimeoutSeconds, TSharedPtr<IPixelStreamingStreamer>, OutStreamer, TSharedPtr<FMockPlayer>, OutPlayer, int, PlayerPort);
 	DEFINE_LATENT_AUTOMATION_COMMAND_THREE_PARAMETER(FWaitForPlayerConnectedOrTimeout, double, TimeoutSeconds, TSharedPtr<FMockPlayer>, OutPlayer, int, PlayerPort);
diff --git a/Engine/Plugins/Media/PixelStreaming/Source/PixelStreaming/Private/Tests/VP8Tests.cpp b/Engine/Plugins/Media/PixelStreaming/Source/PixelStreaming/Private/Tests/VP8Tests.cpp
new file mode 100644
--- /dev/null
+++ b/Engine/Plugins/Media/PixelStreaming/Source/PixelStreaming/Private/Tests/VP8Tests.cpp
@@ -0,0 +1,42 @@
+// Copyright Epic Games, Inc. All Rights Reserved.
+
+#include "CodecUtils.h"
+
+#if WITH_DEV_AUTOMATION_TESTS
+
+namespace UE::PixelStreaming
+{
+	IMPLEMENT_SIMPLE_AUTOMATION_TEST(FVP8FrameReceivedTest, "System.Plugins.PixelStreaming.FVP8FrameReceivedTest", EAutomationTestFlags::EditorContext | EAutomationTestFlags::ClientContext | EAutomationTestFlags::ProductFilter)
+	bool FVP8FrameReceivedTest::RunTest(const FString& Parameters)
+	{
+		SetCodec(EPixelStreamingCodec::VP8, false /*bUseComputeShader*/);
+		DoFrameReceiveTest();
+		return true;
+	}
+
+	IMPLEMENT_SIMPLE_AUTOMATION_TEST(FVP8ComputeShaderFrameReceivedTest, "System.Plugins.PixelStreaming.FVP8ComuteShaderFrameReceivedTest", EAutomationTestFlags::EditorContext | EAutomationTestFlags::ClientContext | EAutomationTestFlags::ProductFilter)
+	bool FVP8ComputeShaderFrameReceivedTest::RunTest(const FString& Parameters)
+	{
+		SetCodec(EPixelStreamingCodec::VP8, true /*bUseComputeShader*/);
+		DoFrameReceiveTest();
+		return true;
+	}
+
+	IMPLEMENT_SIMPLE_AUTOMATION_TEST(FVP8FrameResizeTest, "System.Plugins.PixelStreaming.FVP8FrameResizeTest", EAutomationTestFlags::EditorContext | EAutomationTestFlags::ClientContext | EAutomationTestFlags::ProductFilter)
+	bool FVP8FrameResizeTest::RunTest(const FString& Parameters)
+	{
+		SetCodec(EPixelStreamingCodec::VP8, false /*bUseComputeShader*/);
+		DoFrameResizeMultipleTimesTest();
+		return true;
+	}
+
+	IMPLEMENT_SIMPLE_AUTOMATION_TEST(FVP8ComputeShaderFrameResizeTest, "System.Plugins.PixelStreaming.FVP8ComputeShaderFrameResizeTest", EAutomationTestFlags::EditorContext | EAutomationTestFlags::ClientContext | EAutomationTestFlags::ProductFilter)
+	bool FVP8ComputeShaderFrameResizeTest::RunTest(const FString& Parameters)
+	{
+		SetCodec(EPixelStreamingCodec::VP8, true /*bUseComputeShader*/);
+		DoFrameResizeMultipleTimesTest();
+		return true;
+	}
+} // namespace UE::PixelStreaming
+
+#endif // WITH_DEV_AUTOMATION_TESTS
\ No newline at end of file
diff --git a/Engine/Plugins/Media/PixelStreaming/Source/PixelStreaming/Private/Tests/VP9Tests.cpp b/Engine/Plugins/Media/PixelStreaming/Source/PixelStreaming/Private/Tests/VP9Tests.cpp
new file mode 100644
--- /dev/null
+++ b/Engine/Plugins/Media/PixelStreaming/Source/PixelStreaming/Private/Tests/VP9Tests.cpp
@@ -0,0 +1,42 @@
+// Copyright Epic Games, Inc. All Rights Reserved.
+
+#include "CodecUtils.h"
+
+#if WITH_DEV_AUTOMATION_TESTS
+
+namespace UE::PixelStreaming
+{
+	IMPLEMENT_SIMPLE_AUTOMATION_TEST(FVP9FrameReceivedTest, "System.Plugins.PixelStreaming.FVP9FrameReceivedTest", EAutomationTestFlags::EditorContext | EAutomationTestFlags::ClientContext | EAutomationTestFlags::ProductFilter)
+	bool FVP9FrameReceivedTest::RunTest(const FString& Parameters)
+	{
+		SetCodec(EPixelStreamingCodec::VP9, false /*bUseComputeShader*/);
+		DoFrameReceiveTest();
+		return true;
+	}
+
+	IMPLEMENT_SIMPLE_AUTOMATION_TEST(FVP9ComputeShaderFrameReceivedTest, "System.Plugins.PixelStreaming.FVP9ComuteShaderFrameReceivedTest", EAutomationTestFlags::EditorContext | EAutomationTestFlags::ClientContext | EAutomationTestFlags::ProductFilter)
+	bool FVP9ComputeShaderFrameReceivedTest::RunTest(const FString& Parameters)
+	{
+		SetCodec(EPixelStreamingCodec::VP9, true /*bUseComputeShader*/);
+		DoFrameReceiveTest();
+		return true;
+	}
+
+	IMPLEMENT_SIMPLE_AUTOMATION_TEST(FVP9FrameResizeTest, "System.Plugins.PixelStreaming.FVP9FrameResizeTest", EAutomationTestFlags::EditorContext | EAutomationTestFlags::ClientContext | EAutomationTestFlags::ProductFilter)
+	bool FVP9FrameResizeTest::RunTest(const FString& Parameters)
+	{
+		SetCodec(EPixelStreamingCodec::VP9, false /*bUseComputeShader*/);
+		DoFrameResizeMultipleTimesTest();
+		return true;
+	}
+
+	IMPLEMENT_SIMPLE_AUTOMATION_TEST(FVP9ComputeShaderFrameResizeTest, "System.Plugins.PixelStreaming.FVP9ComputeShaderFrameResizeTest", EAutomationTestFlags::EditorContext | EAutomationTestFlags::ClientContext | EAutomationTestFlags::ProductFilter)
+	bool FVP9ComputeShaderFrameResizeTest::RunTest(const FString& Parameters)
+	{
+		SetCodec(EPixelStreamingCodec::VP9, false /*bUseComputeShader*/);
+		DoFrameResizeMultipleTimesTest();
+		return true;
+	}
+} // namespace UE::PixelStreaming
+
+#endif // WITH_DEV_AUTOMATION_TESTS
\ No newline at end of file
diff --git a/Engine/Plugins/Media/PixelStreaming/Source/PixelStreaming/Private/VideoDecoderFactory.cpp b/Engine/Plugins/Media/PixelStreaming/Source/PixelStreaming/Private/VideoDecoderFactory.cpp
--- a/Engine/Plugins/Media/PixelStreaming/Source/PixelStreaming/Private/VideoDecoderFactory.cpp
+++ b/Engine/Plugins/Media/PixelStreaming/Source/PixelStreaming/Private/VideoDecoderFactory.cpp
@@ -25,49 +25,69 @@ namespace UE::PixelStreaming
 
 		Codecs[EPixelStreamingCodec::VP8].push_back(webrtc::SdpVideoFormat(cricket::kVp8CodecName));
 		Codecs[EPixelStreamingCodec::VP9].push_back(webrtc::SdpVideoFormat(cricket::kVp9CodecName));
-		// Codecs[EPixelStreamingCodec::H264].push_back(UE::PixelStreaming::CreateH264Format(webrtc::H264Profile::kProfileConstrainedBaseline, webrtc::H264Level::kLevel3_1));
-		// Codecs[EPixelStreamingCodec::H264].push_back(UE::PixelStreaming::CreateH264Format(webrtc::H264Profile::kProfileBaseline, webrtc::H264Level::kLevel3_1));
+		// Codecs[EPixelStreamingCodec::H264].push_back(CreateH264Format(webrtc::H264Profile::kProfileConstrainedBaseline, webrtc::H264Level::kLevel3_1));
+		// Codecs[EPixelStreamingCodec::H264].push_back(CreateH264Format(webrtc::H264Profile::kProfileBaseline, webrtc::H264Level::kLevel3_1));
 		Codecs[EPixelStreamingCodec::H265].push_back(webrtc::SdpVideoFormat(cricket::kH265CodecName));
 
 		return Codecs;
 	}
 
+	/**
+	 * Adds all the formats of a given codec to a destination list according to a list of supported formats
+	 */
+	void AddSupportedCodecFormats(EPixelStreamingCodec Codec,
+							 	  const TMap<EPixelStreamingCodec, std::vector<webrtc::SdpVideoFormat>>& SupportedFormatsForCodecs,
+							      std::vector<webrtc::SdpVideoFormat>& OutFormats)
+	{
+		if (SupportedFormatsForCodecs.Contains(Codec))
+		{
+			for (auto& Format : SupportedFormatsForCodecs[Codec])
+			{
+				if (std::find(std::begin(OutFormats), std::end(OutFormats), Format) == std::end(OutFormats))
+				{
+					OutFormats.push_back(Format);
+				}
+			}
+		}
+	}
+
 	std::vector<webrtc::SdpVideoFormat> FVideoDecoderFactory::GetSupportedFormats() const
 	{
 		// static so we dont create the list every time this is called since the list will not change during runtime.
 		static TMap<EPixelStreamingCodec, std::vector<webrtc::SdpVideoFormat>> CodecMap = CreateSupportedDecoderFormatMap();
 
-		// since this method is const we need to store this state statically. it means all instances will share this state
-		// but that actually works in our favor since we're describing more about the plugin state than the actual
-		// instance of this factory.
+		static EPixelStreamingCodec LastSelectedCodec = EPixelStreamingCodec::Invalid;
+		const EPixelStreamingCodec SelectedCodec = Settings::GetSelectedCodec();
+		const bool NegotiateCodecs = Settings::CVarPixelStreamingWebRTCNegotiateCodecs.GetValueOnAnyThread();
+
+		// this is static so we don't have to generate it every time we get a new connection and we cant do that
+		// on a member since this is a const method.
 		static std::vector<webrtc::SdpVideoFormat> SupportedFormats;
 
-		// If we are not negotiating codecs simply return just the one codec that is selected in UE
-		if (!Settings::CVarPixelStreamingWebRTCNegotiateCodecs.GetValueOnAnyThread())
+		// if we're not negotiating codecs we only support the selected codec. If the selected codec is for
+		// some reason unsupported, then we fall through to negotiating codecs.
+		if (!NegotiateCodecs)
 		{
-			const EPixelStreamingCodec SelectedCodec = UE::PixelStreaming::Settings::GetSelectedCodec();
-			if (CodecMap.Contains(SelectedCodec))
+			SupportedFormats.clear();
+			AddSupportedCodecFormats(SelectedCodec, CodecMap, SupportedFormats);
+			
+			if (SupportedFormats.empty())
 			{
-				for (auto& Format : CodecMap[SelectedCodec])
-				{
-					SupportedFormats.push_back(Format);
-				}
-				return SupportedFormats;
+				UE_LOG(LogPixelStreaming, Error, TEXT("Selected codec was not a supported codec, falling back to negotiating codecs..."));
 			}
 			else
 			{
-				UE_LOG(LogPixelStreaming, Error, TEXT("Selected codec was not a supported codec, falling back to negotiating codecs..."));
+				return SupportedFormats;
 			}
 		}
 
-		static EPixelStreamingCodec LastSelectedCodec = EPixelStreamingCodec::Invalid;
-
-		const EPixelStreamingCodec SelectedCodec = UE::PixelStreaming::Settings::GetSelectedCodec();
+		// negotiating codecs means we have a list of codecs we can support and webrtc can pick one.
+		// the list will have the currently selected codec on the top so its the "preferred" codec.
+		// only if this changes do we need to rebuild the list
 		if (LastSelectedCodec != SelectedCodec)
 		{
-			// build a new format list
-			LastSelectedCodec = SelectedCodec;
 			SupportedFormats.clear();
+			LastSelectedCodec = SelectedCodec;
 
 			// order the codecs so the selected is first
 			TArray<EPixelStreamingCodec> OrderedCodecList;
@@ -83,13 +103,7 @@ namespace UE::PixelStreaming
 			// now just add each of the formats in order
 			for (auto& Codec : OrderedCodecList)
 			{
-				if (CodecMap.Contains(Codec))
-				{
-					for (auto& Format : CodecMap[Codec])
-					{
-						SupportedFormats.push_back(Format);
-					}
-				}
+				AddSupportedCodecFormats(Codec, CodecMap, SupportedFormats);
 			}
 		}
 
