Fixing the case where min bitrate is higher than max bitrate. This would cause a crash due to the failure of creating the stream with the given parameters.
- Clamped the max bitrate to the min bitrate if the min bitrate was higher.

diff --git a/Engine/Plugins/Media/PixelStreaming/Source/PixelStreaming/Private/PixelStreamingPeerConnection.cpp b/Engine/Plugins/Media/PixelStreaming/Source/PixelStreaming/Private/PixelStreamingPeerConnection.cpp
--- a/Engine/Plugins/Media/PixelStreaming/Source/PixelStreaming/Private/PixelStreamingPeerConnection.cpp
+++ b/Engine/Plugins/Media/PixelStreaming/Source/PixelStreaming/Private/PixelStreamingPeerConnection.cpp
@@ -43,8 +43,11 @@ namespace
 				const FLayer* SimulcastLayer = SortedLayers[i];
 				webrtc::RtpEncodingParameters LayerEncoding{};
 				LayerEncoding.rid = TCHAR_TO_UTF8(*(FString("simulcast") + FString::FromInt(LayerCount - i)));
-				LayerEncoding.min_bitrate_bps = SimulcastLayer->MinBitrate;
-				LayerEncoding.max_bitrate_bps = SimulcastLayer->MaxBitrate;
+				int minBps, maxBps;
+				minBps = SimulcastLayer->MinBitrate;
+				maxBps = FMath::Max(minBps, SimulcastLayer->MaxBitrate);
+				LayerEncoding.min_bitrate_bps = minBps;
+				LayerEncoding.max_bitrate_bps = maxBps;
 				LayerEncoding.scale_resolution_down_by = SimulcastLayer->Scaling;
 
 				// In M84 this will crash with "Attempted to set an unimplemented parameter of RtpParameters".
@@ -59,8 +62,19 @@ namespace
 		{
 			webrtc::RtpEncodingParameters Encoding{};
 			Encoding.rid = "base";
-			Encoding.max_bitrate_bps = Settings::CVarPixelStreamingWebRTCMaxBitrate.GetValueOnAnyThread();
-			Encoding.min_bitrate_bps = Settings::CVarPixelStreamingWebRTCMinBitrate.GetValueOnAnyThread();
+			// if the min/max bitrates are in the wrong order the stream will fail
+			int MinBps, MaxBps;
+			MinBps = Settings::CVarPixelStreamingWebRTCMinBitrate.GetValueOnAnyThread();
+			MaxBps = Settings::CVarPixelStreamingWebRTCMaxBitrate.GetValueOnAnyThread();
+			if (MinBps > MaxBps)
+			{
+				MaxBps = MinBps;
+				// to try to not be misleading with debug texts etc, we reset these sanitised settings here
+				Settings::CVarPixelStreamingWebRTCMinBitrate->Set(MinBps, ECVF_SetByCode);
+				Settings::CVarPixelStreamingWebRTCMaxBitrate->Set(MaxBps, ECVF_SetByCode);
+			}
+			Encoding.max_bitrate_bps = MaxBps;
+			Encoding.min_bitrate_bps = MinBps;
 			Encoding.max_framerate = Settings::CVarPixelStreamingWebRTCFps.GetValueOnAnyThread();
 			Encoding.scale_resolution_down_by.reset();
 			Encoding.network_priority = webrtc::Priority::kHigh;
