When destroying a peer connection, if there are still webrtc messages in the queue, such as setting local/remote descriptions etc. It was possible to receive a callback during the destructor. This is a small fix to prevent usage of bad stuff during the destructor.

diff --git a/Engine/Plugins/Media/PixelStreaming/Source/PixelStreaming/Private/PixelStreamingPeerConnection.cpp b/Engine/Plugins/Media/PixelStreaming/Source/PixelStreaming/Private/PixelStreamingPeerConnection.cpp
--- a/Engine/Plugins/Media/PixelStreaming/Source/PixelStreaming/Private/PixelStreamingPeerConnection.cpp
+++ b/Engine/Plugins/Media/PixelStreaming/Source/PixelStreaming/Private/PixelStreamingPeerConnection.cpp
@@ -257,6 +257,9 @@ FPixelStreamingPeerConnection::FPixelStreamingPeerConnection() = default;
 
 FPixelStreamingPeerConnection::~FPixelStreamingPeerConnection()
 {
+	// callbacks might still be fired as we're destroying this object, we want to make sure we dont try and access this still.
+	bIsDestroying = true;
+
 	// need to remove the audio/video sinks before deleting it.
 	if (PeerConnection)
 	{
@@ -796,7 +799,11 @@ void FPixelStreamingPeerConnection::CreateSDP(ESDPType SDPType, EReceiveMediaOpt
 			}
 		};
 
-		SetLocalDescription(SDP, OnSetLocalDescriptionSuccess, OnSetLocalDescriptionFail);
+		// prevents this callback from continuing if we're currently in the destructor
+		if (!bIsDestroying)
+		{
+			SetLocalDescription(SDP, OnSetLocalDescriptionSuccess, OnSetLocalDescriptionFail);
+		}
 	};
 
 	auto OnCreateSDPFail = [ErrorCallback](const FString& Error) {
diff --git a/Engine/Plugins/Media/PixelStreaming/Source/PixelStreaming/Public/PixelStreamingPeerConnection.h b/Engine/Plugins/Media/PixelStreaming/Source/PixelStreaming/Public/PixelStreamingPeerConnection.h
--- a/Engine/Plugins/Media/PixelStreaming/Source/PixelStreaming/Public/PixelStreamingPeerConnection.h
+++ b/Engine/Plugins/Media/PixelStreaming/Source/PixelStreaming/Public/PixelStreamingPeerConnection.h
@@ -284,6 +284,7 @@ private:
 	std::unique_ptr<class FPeerWebRTCStatsSource> StatsSource;
 
 	bool IsSFU = false;
+	bool bIsDestroying = false;
 
 	// TODO these static methods can probably be moved off into a singleton or something
 	static void CreatePeerConnectionFactory();
