Fixing an on exit crash/assert.
UPixelStreamingDelegates gets destroyed upon exiting the application but there is also StopStreaming() which tries to broadcast on these delegates.

diff --git a/Engine/Plugins/Media/PixelStreaming/Source/PixelStreaming/Private/PixelStreamingDelegates.cpp b/Engine/Plugins/Media/PixelStreaming/Source/PixelStreaming/Private/PixelStreamingDelegates.cpp
--- a/Engine/Plugins/Media/PixelStreaming/Source/PixelStreaming/Private/PixelStreamingDelegates.cpp
+++ b/Engine/Plugins/Media/PixelStreaming/Source/PixelStreaming/Private/PixelStreamingDelegates.cpp
@@ -3,10 +3,11 @@
 #include "PixelStreamingDelegates.h"
 
 UPixelStreamingDelegates* UPixelStreamingDelegates::Singleton = nullptr;
+bool UPixelStreamingDelegates::bIsExiting = false;
 
 UPixelStreamingDelegates* UPixelStreamingDelegates::CreateInstance()
 {
-	if (Singleton == nullptr)
+	if (Singleton == nullptr && !bIsExiting)
 	{
 		Singleton = NewObject<UPixelStreamingDelegates>();
 		Singleton->AddToRoot();
diff --git a/Engine/Plugins/Media/PixelStreaming/Source/PixelStreaming/Public/PixelStreamingDelegates.h b/Engine/Plugins/Media/PixelStreaming/Source/PixelStreaming/Public/PixelStreamingDelegates.h
--- a/Engine/Plugins/Media/PixelStreaming/Source/PixelStreaming/Public/PixelStreamingDelegates.h
+++ b/Engine/Plugins/Media/PixelStreaming/Source/PixelStreaming/Public/PixelStreamingDelegates.h
@@ -125,14 +125,17 @@ public:
 
 	static UPixelStreamingDelegates* GetPixelStreamingDelegates()
 	{
-		if (Singleton == nullptr)
-		{
-			return CreateInstance();
-		}
-		return Singleton;
+		return CreateInstance();
+	}
+
+	virtual ~UPixelStreamingDelegates()
+	{
+		Singleton = nullptr;
+		bIsExiting = true;
 	}
 
 private:
 	// The singleton object.
 	static UPixelStreamingDelegates* Singleton;
+	static bool bIsExiting;
 };
