Fixing a crash when using OfferToReceive on FF. THe list of supported codecs was not being generated properly.

diff --git a/Engine/Plugins/Media/PixelStreaming/Source/PixelStreaming/Private/VideoDecoderFactory.cpp b/Engine/Plugins/Media/PixelStreaming/Source/PixelStreaming/Private/VideoDecoderFactory.cpp
--- a/Engine/Plugins/Media/PixelStreaming/Source/PixelStreaming/Private/VideoDecoderFactory.cpp
+++ b/Engine/Plugins/Media/PixelStreaming/Source/PixelStreaming/Private/VideoDecoderFactory.cpp
@@ -68,15 +68,17 @@ namespace UE::PixelStreaming
 		// some reason unsupported, then we fall through to negotiating codecs.
 		if (!NegotiateCodecs)
 		{
-			SupportedFormats.clear();
-			AddSupportedCodecFormats(SelectedCodec, CodecMap, SupportedFormats);
+			std::vector<webrtc::SdpVideoFormat> TempSupportedFormats;
+			AddSupportedCodecFormats(SelectedCodec, CodecMap, TempSupportedFormats);
 			
-			if (SupportedFormats.empty())
+			if (TempSupportedFormats.empty())
 			{
 				UE_LOG(LogPixelStreaming, Error, TEXT("Selected codec was not a supported codec, falling back to negotiating codecs..."));
 			}
 			else
 			{
+				SupportedFormats.clear();
+				SupportedFormats = TempSupportedFormats;
 				return SupportedFormats;
 			}
 		}
