Fixing a debug assert on pixel streaming input

diff --git a/Engine/Plugins/Media/PixelStreaming/Source/PixelStreamingInput/Private/PixelStreamingInputHandler.cpp b/Engine/Plugins/Media/PixelStreaming/Source/PixelStreamingInput/Private/PixelStreamingInputHandler.cpp
--- a/Engine/Plugins/Media/PixelStreaming/Source/PixelStreamingInput/Private/PixelStreamingInputHandler.cpp
+++ b/Engine/Plugins/Media/PixelStreaming/Source/PixelStreamingInput/Private/PixelStreamingInputHandler.cpp
@@ -719,6 +719,7 @@ namespace UE::PixelStreamingInput
 		// Set cursor pos on mouse down - we may not have moved if this is the very first click
 		FSlateApplication& SlateApplication = FSlateApplication::Get();
 		SlateApplication.OnCursorSet();
+		PixelStreamerApplicationWrapper->Cursor->SetPosition(ScreenLocation.X, ScreenLocation.Y);
 		// Force window focus
 		SlateApplication.ProcessApplicationActivationEvent(true);
 
