Fixing stats related race conditions.

diff --git a/Engine/Plugins/Media/PixelStreaming/Source/PixelStreaming/Private/RTCStatsCollector.cpp b/Engine/Plugins/Media/PixelStreaming/Source/PixelStreaming/Private/RTCStatsCollector.cpp
--- a/Engine/Plugins/Media/PixelStreaming/Source/PixelStreaming/Private/RTCStatsCollector.cpp
+++ b/Engine/Plugins/Media/PixelStreaming/Source/PixelStreaming/Private/RTCStatsCollector.cpp
@@ -267,6 +267,11 @@ namespace UE::PixelStreaming
 			return;
 		}
 
+		// This function can get fired from a bunch of different threads thanks to WebRTC which can cause
+		// race conditions not only here but also in stats changed delegates etc. This makes sure we do this
+		// work all synchronously
+		FScopeLock StatsLock(&StatsCS);
+
 		for (auto&& Stats : *Report)
 		{
 			const FString StatsType = FString(Stats.type());
diff --git a/Engine/Plugins/Media/PixelStreaming/Source/PixelStreaming/Private/RTCStatsCollector.h b/Engine/Plugins/Media/PixelStreaming/Source/PixelStreaming/Private/RTCStatsCollector.h
--- a/Engine/Plugins/Media/PixelStreaming/Source/PixelStreaming/Private/RTCStatsCollector.h
+++ b/Engine/Plugins/Media/PixelStreaming/Source/PixelStreaming/Private/RTCStatsCollector.h
@@ -88,5 +88,7 @@ namespace UE::PixelStreaming
 
 		// Stream stats are further broken down and stored per RID, as some peers (e.g. an SFU) can have multiple streams.
 		TMap<FString, FStatsSink> StreamStatsSinks;
+
+		FCriticalSection StatsCS;
 	};
 } // namespace UE::PixelStreaming
