Better solution to the race condition relating to stats.

diff --git a/Engine/Plugins/Media/PixelStreaming/Source/PixelStreaming/Private/RTCStatsCollector.cpp b/Engine/Plugins/Media/PixelStreaming/Source/PixelStreaming/Private/RTCStatsCollector.cpp
--- a/Engine/Plugins/Media/PixelStreaming/Source/PixelStreaming/Private/RTCStatsCollector.cpp
+++ b/Engine/Plugins/Media/PixelStreaming/Source/PixelStreaming/Private/RTCStatsCollector.cpp
@@ -267,11 +267,6 @@ namespace UE::PixelStreaming
 			return;
 		}
 
-		// This function can get fired from a bunch of different threads thanks to WebRTC which can cause
-		// race conditions not only here but also in stats changed delegates etc. This makes sure we do this
-		// work all synchronously
-		FScopeLock StatsLock(&StatsCS);
-
 		for (auto&& Stats : *Report)
 		{
 			const FString StatsType = FString(Stats.type());
diff --git a/Engine/Plugins/Media/PixelStreaming/Source/PixelStreaming/Private/RTCStatsCollector.h b/Engine/Plugins/Media/PixelStreaming/Source/PixelStreaming/Private/RTCStatsCollector.h
--- a/Engine/Plugins/Media/PixelStreaming/Source/PixelStreaming/Private/RTCStatsCollector.h
+++ b/Engine/Plugins/Media/PixelStreaming/Source/PixelStreaming/Private/RTCStatsCollector.h
@@ -88,7 +88,5 @@ namespace UE::PixelStreaming
 
 		// Stream stats are further broken down and stored per RID, as some peers (e.g. an SFU) can have multiple streams.
 		TMap<FString, FStatsSink> StreamStatsSinks;
-
-		FCriticalSection StatsCS;
 	};
 } // namespace UE::PixelStreaming
diff --git a/Engine/Plugins/Media/PixelStreaming/Source/PixelStreaming/Private/Stats.cpp b/Engine/Plugins/Media/PixelStreaming/Source/PixelStreaming/Private/Stats.cpp
--- a/Engine/Plugins/Media/PixelStreaming/Source/PixelStreaming/Private/Stats.cpp
+++ b/Engine/Plugins/Media/PixelStreaming/Source/PixelStreaming/Private/Stats.cpp
@@ -181,6 +181,8 @@ namespace UE::PixelStreaming
 
 	void FStats::FireStatChanged(FPixelStreamingPlayerId PlayerId, FName StatName, float StatValue)
 	{
+		// firing off these delegates is not thread safe so we want to mutex this call.
+		FScopeLock Lock(&StatNotificationCS);
 		if (UPixelStreamingDelegates* Delegates = UPixelStreamingDelegates::GetPixelStreamingDelegates())
 		{
 			Delegates->OnStatChangedNative.Broadcast(PlayerId, StatName, StatValue);
diff --git a/Engine/Plugins/Media/PixelStreaming/Source/PixelStreaming/Private/Stats.h b/Engine/Plugins/Media/PixelStreaming/Source/PixelStreaming/Private/Stats.h
--- a/Engine/Plugins/Media/PixelStreaming/Source/PixelStreaming/Private/Stats.h
+++ b/Engine/Plugins/Media/PixelStreaming/Source/PixelStreaming/Private/Stats.h
@@ -147,5 +147,7 @@ namespace UE::PixelStreaming
 
 		TMap<FName, FDebugGraph> Graphs;
 		TMap<FName, FCanvasTileItem> Tiles;
+
+		FCriticalSection StatNotificationCS;
 	};
 } // namespace UE::PixelStreaming
