Merging CL24829126 from Dev-TensorWorks:

Add: SFU layer switching API

diff --git a/Engine/Plugins/Media/PixelStreaming/Source/PixelStreaming/Private/PixelStreamingBlueprints.cpp b/Engine/Plugins/Media/PixelStreaming/Source/PixelStreaming/Private/PixelStreamingBlueprints.cpp
--- a/Engine/Plugins/Media/PixelStreaming/Source/PixelStreaming/Private/PixelStreamingBlueprints.cpp
+++ b/Engine/Plugins/Media/PixelStreaming/Source/PixelStreaming/Private/PixelStreamingBlueprints.cpp
@@ -195,6 +195,66 @@ void UPixelStreamingBlueprints::StreamerKickPlayer(FString StreamerId, FString P
 	Streamer->KickPlayer(ToPlayerId(PlayerId));
 }
 
+void UPixelStreamingBlueprints::SetPlayerLayerPreference(FString PlayerId, int SpatialLayerId, int TemporalLayerId)
+{
+	IPixelStreamingModule* Module = UE::PixelStreaming::FPixelStreamingModule::GetModule();
+	if (!Module)
+	{
+		return;
+	}
+	TSharedPtr<IPixelStreamingStreamer> Streamer = Module->FindStreamer(Module->GetDefaultStreamerID());
+	if (!Streamer)
+	{
+		return;
+	}
+	Streamer->SetPlayerLayerPreference(ToPlayerId(PlayerId), SpatialLayerId, TemporalLayerId);
+}
+
+void UPixelStreamingBlueprints::StreamerSetPlayerLayerPreference(FString StreamerId, FString PlayerId, int SpatialLayerId, int TemporalLayerId)
+{
+	IPixelStreamingModule* Module = UE::PixelStreaming::FPixelStreamingModule::GetModule();
+	if (!Module)
+	{
+		return;
+	}
+	TSharedPtr<IPixelStreamingStreamer> Streamer = Module->FindStreamer(StreamerId);
+	if (!Streamer)
+	{
+		return;
+	}
+	Streamer->SetPlayerLayerPreference(ToPlayerId(PlayerId), SpatialLayerId, TemporalLayerId);
+}
+
+TArray<FString> UPixelStreamingBlueprints::GetConnectedPlayers()
+{
+	IPixelStreamingModule* Module = UE::PixelStreaming::FPixelStreamingModule::GetModule();
+	if (!Module)
+	{
+		return {};
+	}
+	TSharedPtr<IPixelStreamingStreamer> Streamer = Module->FindStreamer(Module->GetDefaultStreamerID());
+	if (!Streamer)
+	{
+		return {};
+	}
+	return Streamer->GetConnectedPlayers();
+}
+
+TArray<FString> UPixelStreamingBlueprints::StreamerGetConnectedPlayers(FString StreamerId)
+{
+	IPixelStreamingModule* Module = UE::PixelStreaming::FPixelStreamingModule::GetModule();
+	if (!Module)
+	{
+		return {};
+	}
+	TSharedPtr<IPixelStreamingStreamer> Streamer = Module->FindStreamer(StreamerId);
+	if (!Streamer)
+	{
+		return {};
+	}
+	return Streamer->GetConnectedPlayers();
+}
+
 FString UPixelStreamingBlueprints::GetDefaultStreamerID()
 {
 	IPixelStreamingModule* Module = UE::PixelStreaming::FPixelStreamingModule::GetModule();
diff --git a/Engine/Plugins/Media/PixelStreaming/Source/PixelStreaming/Private/PixelStreamingBlueprints.h b/Engine/Plugins/Media/PixelStreaming/Source/PixelStreaming/Private/PixelStreamingBlueprints.h
--- a/Engine/Plugins/Media/PixelStreaming/Source/PixelStreaming/Private/PixelStreamingBlueprints.h
+++ b/Engine/Plugins/Media/PixelStreaming/Source/PixelStreaming/Private/PixelStreamingBlueprints.h
@@ -16,45 +16,44 @@ class PIXELSTREAMING_API UPixelStreamingBlueprints : public UBlueprintFunctionLi
 	GENERATED_BODY()
 public:
 	/**
-     * Send a specified byte array over the WebRTC peer connection data channel. The extension and mime type are used for file reconstruction on the front end 
-     * 
-     * @param   ByteArray       The raw data that will be sent over the data channel
-     * @param   MimeType        The mime type of the file. Used for reconstruction on the front end
-     * @param   FileExtension   The file extension. Used for file reconstruction on the front end
-     */
+	 * Send a specified byte array over the WebRTC peer connection data channel. The extension and mime type are used for file reconstruction on the front end
+	 *
+	 * @param   ByteArray       The raw data that will be sent over the data channel
+	 * @param   MimeType        The mime type of the file. Used for reconstruction on the front end
+	 * @param   FileExtension   The file extension. Used for file reconstruction on the front end
+	 */
 	UFUNCTION(BlueprintCallable, Category = "Pixel Streaming")
 	static void SendFileAsByteArray(TArray<uint8> ByteArray, FString MimeType, FString FileExtension);
 
 	/**
-     * Send a specified byte array over the WebRTC peer connection data channel. The extension and mime type are used for file reconstruction on the front end 
-     * 
+	 * Send a specified byte array over the WebRTC peer connection data channel. The extension and mime type are used for file reconstruction on the front end
+	 *
 	 * @param	StreamerId		The streamer use when sending the data
-     * @param   ByteArray       The raw data that will be sent over the data channel
-     * @param   MimeType        The mime type of the file. Used for reconstruction on the front end
-     * @param   FileExtension   The file extension. Used for file reconstruction on the front end
-     */
+	 * @param   ByteArray       The raw data that will be sent over the data channel
+	 * @param   MimeType        The mime type of the file. Used for reconstruction on the front end
+	 * @param   FileExtension   The file extension. Used for file reconstruction on the front end
+	 */
 	UFUNCTION(BlueprintCallable, Category = "Pixel Streaming")
 	static void StreamerSendFileAsByteArray(FString StreamerId, TArray<uint8> ByteArray, FString MimeType, FString FileExtension);
 
 	/**
-     * Send a specified file over the WebRTC peer connection data channel. The extension and mime type are used for file reconstruction on the front end 
-     * 
-     * @param   FilePath        The path to the file that will be sent
-     * @param   MimeType        The mime type of the file. Used for file reconstruction on the front end
-     * @param   FileExtension   The file extension. Used for file reconstruction on the front end
-     */
+	 * Send a specified file over the WebRTC peer connection data channel. The extension and mime type are used for file reconstruction on the front end
+	 *
+	 * @param   FilePath        The path to the file that will be sent
+	 * @param   MimeType        The mime type of the file. Used for file reconstruction on the front end
+	 * @param   FileExtension   The file extension. Used for file reconstruction on the front end
+	 */
 	UFUNCTION(BlueprintCallable, Category = "Pixel Streaming")
 	static void SendFile(FString Filepath, FString MimeType, FString FileExtension);
 
-
 	/**
-     * Send a specified file over the WebRTC peer connection data channel. The extension and mime type are used for file reconstruction on the front end 
-     * 
+	 * Send a specified file over the WebRTC peer connection data channel. The extension and mime type are used for file reconstruction on the front end
+	 *
 	 * @param	StreamerId		The streamer use when sending the data
-     * @param   FilePath        The path to the file that will be sent
-     * @param   MimeType        The mime type of the file. Used for file reconstruction on the front end
-     * @param   FileExtension   The file extension. Used for file reconstruction on the front end
-     */
+	 * @param   FilePath        The path to the file that will be sent
+	 * @param   MimeType        The mime type of the file. Used for file reconstruction on the front end
+	 * @param   FileExtension   The file extension. Used for file reconstruction on the front end
+	 */
 	UFUNCTION(BlueprintCallable, Category = "Pixel Streaming")
 	static void StreamerSendFile(FString StreamerId, FString Filepath, FString MimeType, FString FileExtension);
 
@@ -65,12 +64,12 @@ public:
 	static void ForceKeyFrame();
 
 	/**
-    * The functions allow Pixel Streaming to be frozen and unfrozen from
-    * Blueprint. When frozen, a freeze frame (a still image) will be used by the
-    * browser instead of the video stream.
-    */
+	 * The functions allow Pixel Streaming to be frozen and unfrozen from
+	 * Blueprint. When frozen, a freeze frame (a still image) will be used by the
+	 * browser instead of the video stream.
+	 */
 
-   /**
+	/**
 	 * Freeze Pixel Streaming.
 	 * @param   Texture         The freeze frame to display. If null then the back buffer is captured.
 	 */
@@ -78,12 +77,11 @@ public:
 	static void FreezeFrame(UTexture2D* Texture);
 
 	/**
-	 * Unfreeze Pixel Streaming. 
+	 * Unfreeze Pixel Streaming.
 	 */
 	UFUNCTION(BlueprintCallable, Category = "Pixel Streaming")
 	static void UnfreezeFrame();
 
-
 	/**
 	 * Freeze Pixel Streaming.
 	 * @param	StreamerId		The id of the streamer to freeze.
@@ -93,7 +91,7 @@ public:
 	static void StreamerFreezeStream(FString StreamerId, UTexture2D* Texture);
 
 	/**
-	 * Unfreeze Pixel Streaming. 
+	 * Unfreeze Pixel Streaming.
 	 * @param StreamerId		The id of the streamer to unfreeze.
 	 */
 	UFUNCTION(BlueprintCallable, Category = "Pixel Streaming")
@@ -114,6 +112,42 @@ public:
 	UFUNCTION(BlueprintCallable, Category = "Pixel Streaming")
 	static void StreamerKickPlayer(FString StreamerId, FString PlayerId);
 
+	/**
+	 * Set the layer preference for a specific player
+	 * @param PlayerId			The ID of the player to set the layer preference for.
+	 * @param SpatialLayerId	The spatial layer ID
+	 * @param TemporalLayerId	The temporal layer ID
+	 */
+	UFUNCTION(BlueprintCallable, Category = "Pixel Streaming")
+	static void SetPlayerLayerPreference(FString PlayerId, int SpatialLayerId, int TemporalLayerId);
+
+	/**
+	 * Set the layer preference for a specific player
+	 * @param StreamerId 		The streamer which the player belongs
+	 * @param PlayerId			The ID of the player to set the layer preference for.
+	 * @param SpatialLayerId	The spatial layer ID
+	 * @param TemporalLayerId	The temporal layer ID
+	 */
+	UFUNCTION(BlueprintCallable, Category = "Pixel Streaming")
+	static void StreamerSetPlayerLayerPreference(FString StreamerId, FString PlayerId, int SpatialLayerId, int TemporalLayerId);
+
+	/**
+	 * @brief Get the connected players
+	 *
+	 * @return TArray<FString> The connected players
+	 */
+	UFUNCTION(BlueprintCallable, Category = "Pixel Streaming")
+	static TArray<FString> GetConnectedPlayers();
+
+	/**
+	 * @brief Get the connected players
+	 *
+	 * @param StreamerId	The streamer whose list of players you wish to get
+	 * @return TArray<FString> The connected players
+	 */
+	UFUNCTION(BlueprintCallable, Category = "Pixel Streaming")
+	static TArray<FString> StreamerGetConnectedPlayers(FString StreamerId);
+
 	/**
 	 * Get the default Streamer ID
 	 */
diff --git a/Engine/Plugins/Media/PixelStreaming/Source/PixelStreaming/Private/Streamer.cpp b/Engine/Plugins/Media/PixelStreaming/Source/PixelStreaming/Private/Streamer.cpp
--- a/Engine/Plugins/Media/PixelStreaming/Source/PixelStreaming/Private/Streamer.cpp
+++ b/Engine/Plugins/Media/PixelStreaming/Source/PixelStreaming/Private/Streamer.cpp
@@ -35,6 +35,7 @@
 #include "PixelCaptureOutputFrameRHI.h"
 #include "PixelCaptureInputFrameRHI.h"
 #include "SignallingConnectionObserver.h"
+#include "ToStringExtensions.h"
 
 namespace UE::PixelStreaming
 {
@@ -379,6 +380,26 @@ namespace UE::PixelStreaming
 		// TODO Delete player session?
 	}
 
+	void FStreamer::SetPlayerLayerPreference(FPixelStreamingPlayerId PlayerId, int SpatialLayerId, int TemporalLayerId)
+	{
+		TSharedPtr<FJsonObject> LayerJson = MakeShared<FJsonObject>();
+		LayerJson->SetStringField(TEXT("type"), TEXT("layerPreference"));
+		int32 PlayerIdAsInt = PlayerIdToInt(PlayerId);
+		LayerJson->SetNumberField(TEXT("playerId"), PlayerIdAsInt);
+		LayerJson->SetNumberField(TEXT("spatialLayer"), SpatialLayerId);
+		LayerJson->SetNumberField(TEXT("temporalLayer"), TemporalLayerId);
+		SignallingServerConnection->SendMessage(UE::PixelStreaming::ToString(LayerJson, false));
+	}
+
+	TArray<FPixelStreamingPlayerId> FStreamer::GetConnectedPlayers()
+	{
+		TArray<FPixelStreamingPlayerId> ConnectedPlayerIds;
+		Players.Apply([&ConnectedPlayerIds, this](FPixelStreamingPlayerId PlayerId, FPlayerContext& PlayerContext) {
+			ConnectedPlayerIds.Add(PlayerId);
+		});
+		return ConnectedPlayerIds;
+	}
+
 	void FStreamer::SetInputHandlerType(EPixelStreamingInputType InputType)
 	{
 		InputHandler->SetInputType(InputType);
diff --git a/Engine/Plugins/Media/PixelStreaming/Source/PixelStreaming/Private/Streamer.h b/Engine/Plugins/Media/PixelStreaming/Source/PixelStreaming/Private/Streamer.h
--- a/Engine/Plugins/Media/PixelStreaming/Source/PixelStreaming/Private/Streamer.h
+++ b/Engine/Plugins/Media/PixelStreaming/Source/PixelStreaming/Private/Streamer.h
@@ -63,8 +63,10 @@ namespace UE::PixelStreaming
 		virtual void SendPlayerMessage(uint8 Type, const FString& Descriptor) override;
 		virtual void SendFileData(const TArray64<uint8>& ByteData, FString& MimeType, FString& FileExtension) override;
 		virtual void KickPlayer(FPixelStreamingPlayerId PlayerId) override;
+		virtual void SetPlayerLayerPreference(FPixelStreamingPlayerId PlayerId, int SpatialLayerId, int TemporalLayerId) override;
+		virtual TArray<FPixelStreamingPlayerId> GetConnectedPlayers() override;
 
-		virtual void SetInputHandler(TSharedPtr<IPixelStreamingInputHandler> InInputHandler) override { InputHandler = InInputHandler; }
+		virtual void SetInputHandler(TSharedPtr<IPixelStreamingInputHandler> InInputHandler) override {	InputHandler = InInputHandler; }
 		virtual TWeakPtr<IPixelStreamingInputHandler> GetInputHandler() override { return InputHandler; }
 		virtual void SetInputHandlerType(EPixelStreamingInputType InputType) override;
 
diff --git a/Engine/Plugins/Media/PixelStreaming/Source/PixelStreaming/Public/IPixelStreamingStreamer.h b/Engine/Plugins/Media/PixelStreaming/Source/PixelStreaming/Public/IPixelStreamingStreamer.h
--- a/Engine/Plugins/Media/PixelStreaming/Source/PixelStreaming/Public/IPixelStreamingStreamer.h
+++ b/Engine/Plugins/Media/PixelStreaming/Source/PixelStreaming/Public/IPixelStreamingStreamer.h
@@ -212,7 +212,20 @@ public:
 	 */
 	virtual void KickPlayer(FPixelStreamingPlayerId PlayerId) = 0;
 
-	// virtual bool IsQualityController(FPixelStreamingPlayerId PlayerId) = 0;
+	/**
+	 * @brief Set the Layer Preference for a specific player. Only useful if using the SFU
+	 * @param PlayerId The player id
+	 * @param SpatialLayerId The id of the preferred spatial layer (resolution)
+	 * @param TemporalLayerId The id of the preferred temporal layer (frame rate)
+	 */
+	virtual void SetPlayerLayerPreference(FPixelStreamingPlayerId PlayerId, int SpatialLayerId, int TemporalLayerId) = 0;
+
+	/**
+	 * @brief Get the connected players
+	 * 
+	 * @return TArray<FPixelStreamingPlayerId> The connected players
+	 */
+	virtual TArray<FPixelStreamingPlayerId> GetConnectedPlayers() = 0;
 
 	/**
 	 * @brief Set the streamer's input handler
