Merging CL24113859 from Dev-TensorWorks:

Updating threadsafe map to be more thread safe.
Added a level of indirection to the map elements so points do not invalidate on map changes.

diff --git a/Engine/Plugins/Media/PixelStreaming/Source/PixelStreaming/Private/Streamer.cpp b/Engine/Plugins/Media/PixelStreaming/Source/PixelStreaming/Private/Streamer.cpp
--- a/Engine/Plugins/Media/PixelStreaming/Source/PixelStreaming/Private/Streamer.cpp
+++ b/Engine/Plugins/Media/PixelStreaming/Source/PixelStreaming/Private/Streamer.cpp
@@ -727,7 +727,7 @@ namespace UE::PixelStreaming
 	void FStreamer::DeleteAllPlayerSessions()
 	{
 		VideoSourceGroup->RemoveAllVideoSources();
-		Players.Clear();
+		Players.Empty();
 		SFUPlayerId = INVALID_PLAYER_ID;
 		QualityControllingId = INVALID_PLAYER_ID;
 		InputControllingId = INVALID_PLAYER_ID;
diff --git a/Engine/Plugins/Media/PixelStreaming/Source/PixelStreaming/Private/ThreadSafeMap.h b/Engine/Plugins/Media/PixelStreaming/Source/PixelStreaming/Private/ThreadSafeMap.h
--- a/Engine/Plugins/Media/PixelStreaming/Source/PixelStreaming/Private/ThreadSafeMap.h
+++ b/Engine/Plugins/Media/PixelStreaming/Source/PixelStreaming/Private/ThreadSafeMap.h
@@ -11,13 +11,27 @@ namespace UE::PixelStreaming
 	template <typename KeyType, typename ValueType>
 	class TThreadSafeMap
 	{
+	private:
+		// We add a level of indirection here so when the map changes the location of the ValueType
+		// does not change. This allows Find to return a pointer to the ValueType and release the lock
+		// and not worry that changes to the map will change the location of the ValueType pointed to.
+		using InternalElementType = TSharedPtr<ValueType>;
+		TMap<KeyType, InternalElementType> InnerMap;
+		mutable FCriticalSection Mutex;
+
+		// Helpers for manipulating internal elements
+		InternalElementType NewEmptyElement() { return MakeShared<ValueType>(); }
+		InternalElementType NewElement(const ValueType& Value) { return MakeShared<ValueType>(Value); }
+		ValueType* GetValuePtr(const InternalElementType& Element) const { return Element.Get(); }
+		ValueType& GetValue(const InternalElementType& Element) const { return *GetValuePtr(Element); }
+
 	public:
 		bool Add(KeyType Key, const ValueType& Value)
 		{
 			FScopeLock Lock(&Mutex);
 			if (!InnerMap.Contains(Key))
 			{
-				InnerMap.Add(Key, Value);
+				InnerMap.Add(Key, NewElement(Value));
 				return true;
 			}
 			return false;
@@ -34,24 +48,33 @@ namespace UE::PixelStreaming
 			FScopeLock Lock(&Mutex);
 			if (InnerMap.Contains(Key))
 			{
-				return InnerMap[Key];
+				return *InnerMap[Key].Get();
 			}
-			return InnerMap.Add(Key);
+			InternalElementType& NewElement = InnerMap.Add(Key, NewEmptyElement());
+			return GetValue(NewElement);
 		}
 
 		ValueType* Find(KeyType Key)
 		{
 			FScopeLock Lock(&Mutex);
-			return InnerMap.Find(Key);
+			if (InternalElementType* Found = InnerMap.Find(Key))
+			{
+				return GetValuePtr(*Found);
+			}
+			return nullptr;
 		}
 
 		const ValueType* Find(KeyType Key) const
 		{
 			FScopeLock Lock(&Mutex);
-			return InnerMap.Find(Key);
+			if (const InternalElementType* Found = InnerMap.Find(Key))
+			{
+				return GetValuePtr(*Found);
+			}
+			return nullptr;
 		}
 
-		void Clear()
+		void Empty()
 		{
 			FScopeLock Lock(&Mutex);
 			InnerMap.Empty();
@@ -68,14 +91,14 @@ namespace UE::PixelStreaming
 		void Apply(T&& Visitor)
 		{
 			// work on a copy of the map so we dont lock while calling the visitor
-			TMap<KeyType, ValueType> MapCopy;
+			TMap<KeyType, InternalElementType> MapCopy;
 			{
 				FScopeLock Lock(&Mutex);
 				MapCopy = InnerMap;
 			}
-			for (auto&& [Key, Value] : MapCopy)
+			for (auto&& [Key, Element] : MapCopy)
 			{
-				Visitor(Key, Value);
+				Visitor(Key, GetValue(Element));
 			}
 		}
 
@@ -84,22 +107,18 @@ namespace UE::PixelStreaming
 		void ApplyUntil(T&& Visitor)
 		{
 			// work on a copy of the map so we dont lock while calling the visitor
-			TMap<KeyType, ValueType> MapCopy;
+			TMap<KeyType, InternalElementType> MapCopy;
 			{
 				FScopeLock Lock(&Mutex);
 				MapCopy = InnerMap;
 			}
-			for (auto&& [Key, Value] : MapCopy)
+			for (auto&& [Key, Element] : MapCopy)
 			{
-				if (Visitor(Key, Value))
+				if (Visitor(Key, GetValue(Element)))
 				{
 					break;
 				}
 			}
 		}
-
-	private:
-		TMap<KeyType, ValueType> InnerMap;
-		mutable FCriticalSection Mutex;
 	};
 } // namespace UE::PixelStreaming
