Merging CL24758353, CL24758495 from Dev-TensorWorks:

PixelStreaming: Add WebRTC field trials to Pixel Streaming command line arguments
Passing field trials to WebRTC enables us to change existing settings nested deep inside the WebRTC library without changing it's source code.
With Ouroboros project we want to change packet buffer size to accomodate large frames. With this CL we would be able to just pass "-PixelStreamingWebRTCFieldTrials=WebRTC-PacketBufferMaxSize/8192/".

Test steps:
1. Start Ouroboros Streamer with `UnrealEditor.exe Ouroboros -Game -Windowed -ResX=800 -ResY=600 -Streamer -SourceImageFolder={FOLDER_WITH_NV12_FRAMES} -Resolution=3840x2160 -PixelStreamingWebRTCStartBitrate=150000000 -PixelStreamingWebRTCMinBitrate=150000000 -PixelStreamingWebRTCMaxBitrate=1000000000`
2. Start Ouroboros Consumer with `UnrealEditor.exe Ouroboros -Game -Windowed -ResX=1280 -ResY=720 -PixelStreamingHUDStats=true`
3. Observe video freeze and growing PLI at Consumer side
4. Restart Consumer with `UnrealEditor.exe Ouroboros -Game -Windowed -ResX=1280 -ResY=720 -PixelStreamingHUDStats=true -PixelStreamingWebRTCFieldTrials="WebRTC-PacketBufferMaxSize/8192/"`
5. Consumer should play video stream with no PLI and high bitrate

diff --git a/Engine/Plugins/Media/PixelStreaming/Source/PixelStreaming/Private/PixelStreamingPeerConnection.cpp b/Engine/Plugins/Media/PixelStreaming/Source/PixelStreaming/Private/PixelStreamingPeerConnection.cpp
--- a/Engine/Plugins/Media/PixelStreaming/Source/PixelStreaming/Private/PixelStreamingPeerConnection.cpp
+++ b/Engine/Plugins/Media/PixelStreaming/Source/PixelStreaming/Private/PixelStreamingPeerConnection.cpp
@@ -920,6 +920,17 @@ void FPixelStreamingPeerConnection::CreatePeerConnectionFactory()
 
 	std::unique_ptr<FVideoEncoderFactoryLayered> VideoEncoderFactory = std::make_unique<FVideoEncoderFactoryLayered>();
 	GVideoEncoderFactory = VideoEncoderFactory.get();
+	 
+	FString FieldTrials = Settings::CVarPixelStreamingWebRTCFieldTrials.GetValueOnAnyThread();
+	if (!FieldTrials.IsEmpty()) {
+		//Pass the field trials string to WebRTC. String must never be destroyed.
+		TStringConversion<TStringConvert<TCHAR, ANSICHAR>> Str = StringCast<ANSICHAR>(*FieldTrials);
+		int length = Str.Length() + 1;
+		char* WRTCFieldTrials = (char*)FMemory::SystemMalloc(length);
+		FMemory::Memcpy(WRTCFieldTrials, Str.Get(), Str.Length());
+		WRTCFieldTrials[length - 1] = '\0';
+		webrtc::field_trial::InitFieldTrialsFromString(WRTCFieldTrials);
+	}
 
 	PeerConnectionFactory = webrtc::CreatePeerConnectionFactory(
 		nullptr,													   // network_thread
diff --git a/Engine/Plugins/Media/PixelStreaming/Source/PixelStreaming/Private/Settings.cpp b/Engine/Plugins/Media/PixelStreaming/Source/PixelStreaming/Private/Settings.cpp
--- a/Engine/Plugins/Media/PixelStreaming/Source/PixelStreaming/Private/Settings.cpp
+++ b/Engine/Plugins/Media/PixelStreaming/Source/PixelStreaming/Private/Settings.cpp
@@ -101,7 +101,7 @@ namespace UE::PixelStreaming::Settings
 	TAutoConsoleVariable<FString> CVarPixelStreamingEncoderCodec(
 		TEXT("PixelStreaming.Encoder.Codec"),
 		TEXT("H264"),
-		TEXT("PixelStreaming encoder codec. Supported values are `H264`, 'H265', `VP8`, `VP9`"),
+		TEXT("PixelStreaming encoder codec. Supported values are `H264`, `H265`, `VP8`, `VP9`"),
 		ECVF_Default);
 
 	TAutoConsoleVariable<int32> CVarPixelStreamingEncoderMaxSessions(
@@ -197,6 +197,12 @@ namespace UE::PixelStreaming::Settings
 		TEXT("Sets the amount of gain to apply to audio. Default: 1.0"),
 		ECVF_Default);
 
+	TAutoConsoleVariable<FString> CVarPixelStreamingWebRTCFieldTrials(
+		TEXT("PixelStreaming.WebRTC.FieldTrials"),
+		TEXT(""),
+		TEXT("Sets the WebRTC field trials string. Format:\"TRIAL1/VALUE1/TRIAL2/VALUE2/\""),
+		ECVF_Default);
+
 	// End WebRTC CVars
 
 	// Begin Pixel Streaming Plugin CVars
@@ -643,6 +649,7 @@ namespace UE::PixelStreaming::Settings
 		CommandLineParseValue(TEXT("PixelStreamingFreezeFrameQuality"), CVarPixelStreamingFreezeFrameQuality);
 		CommandLineParseValue(TEXT("PixelStreamingInputController="), CVarPixelStreamingInputController);
 		CommandLineParseValue(TEXT("PixelStreamingSignalingReconnectInterval="), CVarPixelStreamingSignalingReconnectInterval);
+		CommandLineParseValue(TEXT("PixelStreamingWebRTCFieldTrials="), CVarPixelStreamingWebRTCFieldTrials);
 
 		// Options parse (if these exist they are set to true)
 		CommandLineParseOption(TEXT("PixelStreamingOnScreenStats"), CVarPixelStreamingOnScreenStats);
diff --git a/Engine/Plugins/Media/PixelStreaming/Source/PixelStreaming/Private/Settings.h b/Engine/Plugins/Media/PixelStreaming/Source/PixelStreaming/Private/Settings.h
--- a/Engine/Plugins/Media/PixelStreaming/Source/PixelStreaming/Private/Settings.h
+++ b/Engine/Plugins/Media/PixelStreaming/Source/PixelStreaming/Private/Settings.h
@@ -50,6 +50,7 @@ namespace UE::PixelStreaming::Settings
 	extern TAutoConsoleVariable<bool> CVarPixelStreamingWebRTCDisableStats;
 	extern TAutoConsoleVariable<float> CVarPixelStreamingWebRTCAudioGain;
 	extern TAutoConsoleVariable<bool> CVarPixelStreamingWebRTCNegotiateCodecs;
+	extern TAutoConsoleVariable<FString> CVarPixelStreamingWebRTCFieldTrials;
 	// End WebRTC CVars
 
 	// Begin Pixel Streaming Plugin CVars
