Merging CL24768117 from Dev-TensorWorks:

Updating the backbuffer input to also handle the case of PIE sessions.

diff --git a/Engine/Plugins/Media/PixelStreaming/Source/PixelStreaming/Private/PixelStreamingVideoInputPIEViewport.cpp b/Engine/Plugins/Media/PixelStreaming/Source/PixelStreaming/Private/PixelStreamingVideoInputPIEViewport.cpp
new file mode 100644
--- /dev/null
+++ b/Engine/Plugins/Media/PixelStreaming/Source/PixelStreaming/Private/PixelStreamingVideoInputPIEViewport.cpp
@@ -0,0 +1,55 @@
+// Copyright Epic Games, Inc. All Rights Reserved.
+
+#include "PixelStreamingVideoInputPIEViewport.h"
+#include "PixelCaptureInputFrameRHI.h"
+#include "Utils.h"
+#include "UnrealClient.h"
+#include "Engine/GameViewportClient.h"
+#include "RenderingThread.h"
+
+TSharedPtr<FPixelStreamingVideoInputPIEViewport> FPixelStreamingVideoInputPIEViewport::Create()
+{
+	TSharedPtr<FPixelStreamingVideoInputPIEViewport> NewInput = TSharedPtr<FPixelStreamingVideoInputPIEViewport>(new FPixelStreamingVideoInputPIEViewport());
+	TWeakPtr<FPixelStreamingVideoInputPIEViewport> WeakInput = NewInput;
+
+	UE::PixelStreaming::DoOnGameThread([WeakInput]() {
+		if (TSharedPtr<FPixelStreamingVideoInputPIEViewport> Input = WeakInput.Pin())
+		{
+			Input->DelegateHandle = UGameViewportClient::OnViewportRendered().AddSP(Input.ToSharedRef(), &FPixelStreamingVideoInputPIEViewport::OnViewportRendered);
+		}
+	});
+	
+	return NewInput;
+}
+
+FPixelStreamingVideoInputPIEViewport::~FPixelStreamingVideoInputPIEViewport()
+{
+	if (!IsEngineExitRequested())
+	{
+		UGameViewportClient::OnViewportRendered().Remove(DelegateHandle);
+	}
+}
+
+void FPixelStreamingVideoInputPIEViewport::OnViewportRendered(FViewport* InViewport)
+{
+	if (!InViewport->IsPlayInEditorViewport())
+	{
+		return;
+	}
+
+	const FTextureRHIRef& FrameBuffer = InViewport->GetRenderTargetTexture();
+	if (!FrameBuffer)
+	{
+		return;
+	}
+
+	ENQUEUE_RENDER_COMMAND(StreamViewportTextureCommand)
+	([&, FrameBuffer](FRHICommandList& RHICmdList) {
+		OnFrame(FPixelCaptureInputFrameRHI(FrameBuffer));
+	});
+}
+
+FString FPixelStreamingVideoInputPIEViewport::ToString()
+{
+	return TEXT("the PIE Viewport");
+}
diff --git a/Engine/Plugins/Media/PixelStreaming/Source/PixelStreaming/Public/PixelStreamingVideoInputPIEViewport.h b/Engine/Plugins/Media/PixelStreaming/Source/PixelStreaming/Public/PixelStreamingVideoInputPIEViewport.h
new file mode 100644
--- /dev/null
+++ b/Engine/Plugins/Media/PixelStreaming/Source/PixelStreaming/Public/PixelStreamingVideoInputPIEViewport.h
@@ -0,0 +1,26 @@
+// Copyright Epic Games, Inc. All Rights Reserved.
+
+#pragma once
+
+#include "PixelStreamingVideoInputRHI.h"
+
+class FViewport;
+
+/*
+ * An extension of the back buffer input that can handle PIE sessions. Primarily to be used in blueprints
+ */
+class PIXELSTREAMING_API FPixelStreamingVideoInputPIEViewport : public FPixelStreamingVideoInputRHI
+{
+public:
+	static TSharedPtr<FPixelStreamingVideoInputPIEViewport> Create();
+	virtual ~FPixelStreamingVideoInputPIEViewport();
+
+	virtual FString ToString() override;
+
+private:
+	FPixelStreamingVideoInputPIEViewport() = default;
+
+	void OnViewportRendered(FViewport* InViewport);
+
+	FDelegateHandle DelegateHandle;
+};
diff --git a/Engine/Plugins/Media/PixelStreaming/Source/PixelStreamingBlueprint/Private/PixelStreamingStreamerVideoInputBackBuffer.cpp b/Engine/Plugins/Media/PixelStreaming/Source/PixelStreamingBlueprint/Private/PixelStreamingStreamerVideoInputBackBuffer.cpp
--- a/Engine/Plugins/Media/PixelStreaming/Source/PixelStreamingBlueprint/Private/PixelStreamingStreamerVideoInputBackBuffer.cpp
+++ b/Engine/Plugins/Media/PixelStreaming/Source/PixelStreamingBlueprint/Private/PixelStreamingStreamerVideoInputBackBuffer.cpp
@@ -1,10 +1,21 @@
 // Copyright Epic Games, Inc. All Rights Reserved.
 
 #include "PixelStreamingStreamerVideoInputBackBuffer.h"
+#include "PixelStreamingVideoInputPIEViewport.h"
 #include "PixelStreamingVideoInputBackBuffer.h"
 
 UPixelStreamingStreamerVideoInputBackBuffer::UPixelStreamingStreamerVideoInputBackBuffer(const FObjectInitializer& ObjectInitializer)
 	: Super(ObjectInitializer)
 {
-	VideoInput = FPixelStreamingVideoInputBackBuffer::Create();
+	// detect if we're in PIE mode or not
+	bool IsGame = false;
+	FParse::Bool(FCommandLine::Get(), TEXT("game"), IsGame);
+	if (GIsEditor && !IsGame)
+	{
+		VideoInput = FPixelStreamingVideoInputPIEViewport::Create();
+	}
+	else
+	{
+		VideoInput = FPixelStreamingVideoInputBackBuffer::Create();
+	}
 }
