Merging CL25098655 from Dev-TensorWorks:

Original CL description:
Fixing failing pixel streaming automated tests.
- Some tests were waiting for an ICE connection status change and assuming that meant that a data channel was available but this was not always true.
- Changed tests using data channels to wait for a data channel to be available.
- Additionally added a check to make sure the data channel exists before using it to prevent null pointer usage.

This CL specifically adds (other work already integrated):
- Additionally added a check to make sure the data channel exists before using it to prevent null pointer usage.

diff --git a/Engine/Plugins/Media/PixelStreaming/Source/PixelStreaming/Private/Tests/TestUtils.cpp b/Engine/Plugins/Media/PixelStreaming/Source/PixelStreaming/Private/Tests/TestUtils.cpp
--- a/Engine/Plugins/Media/PixelStreaming/Source/PixelStreaming/Private/Tests/TestUtils.cpp
+++ b/Engine/Plugins/Media/PixelStreaming/Source/PixelStreaming/Private/Tests/TestUtils.cpp
@@ -49,14 +49,28 @@ namespace UE::PixelStreaming
 	bool FSendDataChannelMessageToStreamer::Update()
 	{
 		UE_LOG(LogPixelStreaming, Log, TEXT("SendDataChannelMessageToStreamer: %d, %s"), Id, *Body);
-		Player->DataChannel->SendMessage(Id, Body);
+		if (!Player->DataChannel)
+		{
+			UE_LOG(LogPixelStreaming, Error, TEXT("No DataChannel on player."));
+		}
+		else
+		{
+			Player->DataChannel->SendMessage(Id, Body);
+		}
 		return true;
 	}
 
 	bool FSendCustomMessageToStreamer::Update()
 	{
 		UE_LOG(LogPixelStreaming, Log, TEXT("SendDataChannelMessageToStreamer: %d, %d"), Id, Body);
-		Player->DataChannel->SendMessage(Id, Body);
+		if (!Player->DataChannel)
+		{
+			UE_LOG(LogPixelStreaming, Error, TEXT("No DataChannel on player."));
+		}
+		else
+		{
+			Player->DataChannel->SendMessage(Id, Body);
+		}
 		return true;
 	}
 
