Fix: RDG Graph would be destructed at odd times leading to crashes.

diff --git a/Engine/Plugins/Media/PixelStreaming/Source/PixelStreamingEditor/Private/PixelStreamingVideoInputBackBufferComposited.cpp b/Engine/Plugins/Media/PixelStreaming/Source/PixelStreamingEditor/Private/PixelStreamingVideoInputBackBufferComposited.cpp
--- a/Engine/Plugins/Media/PixelStreaming/Source/PixelStreamingEditor/Private/PixelStreamingVideoInputBackBufferComposited.cpp
+++ b/Engine/Plugins/Media/PixelStreaming/Source/PixelStreamingEditor/Private/PixelStreamingVideoInputBackBufferComposited.cpp
@@ -130,40 +130,44 @@ void FPixelStreamingVideoInputBackBufferComposited::OnBackBufferReady(SWindow& S
 			return;
 		}
 
-		FRHICommandListImmediate& RHICmdList = FRHICommandListExecutor::GetImmediateCommandList();
-		FRDGBuilder GraphBuilder(RHICmdList);
-
-		// Register an external RDG texture from the provided frame buffer
-		FRDGTextureRef InputTexture = GraphBuilder.RegisterExternalTexture(CreateRenderTarget(FrameBuffer, *SlateWindow.GetTitle().ToString()));
-
-		// Create an internal RDG texture with the same extent and format as the source
-		FRDGTextureRef OutputTexture = GraphBuilder.CreateTexture(FRDGTextureDesc::Create2D(FrameBuffer->GetDesc().Extent, FrameBuffer->GetDesc().Format, FClearValueBinding::None, ETextureCreateFlags::Shared | ETextureCreateFlags::RenderTargetable), TEXT("VideoInputBackBufferCompositedStaging"));
-		// Bit cheeky, but when attempting to create two textures with the same description, RDG was just re-allocating which would lead to flickering. By converting to external, we force immediate allocation of the underlying pooled resource
-		TopLevelWindows[Index].SetTexture(GraphBuilder.ConvertToExternalTexture(OutputTexture));
-
-		// Initialise our copy paramaters
-		FCopyToStagingParameters* PassParameters = GraphBuilder.AllocParameters<FCopyToStagingParameters>();
-		PassParameters->Input = InputTexture;
-		PassParameters->Output = OutputTexture;
-
-		GraphBuilder.AddPass(
-			RDG_EVENT_NAME("VideoInputBackBufferCopyToStaging (%s)", *SlateWindow.GetTitle().ToString()),
-			PassParameters,
-			ERDGPassFlags::Readback,
-			[InputTexture, OutputTexture, Index, this](FRHICommandList& RHICmdList) {
-				// Simple texture copy as we know pixel format is the same
-				RHICmdList.CopyTexture(InputTexture->GetRHI(), OutputTexture->GetRHI(), {});
-			});
-
-		GraphBuilder.Execute();
+		{
+			FRHICommandListImmediate& RHICmdList = FRHICommandListExecutor::GetImmediateCommandList();
+			FRDGBuilder GraphBuilder(RHICmdList);
+
+			// Register an external RDG texture from the provided frame buffer
+			FRDGTextureRef InputTexture = GraphBuilder.RegisterExternalTexture(CreateRenderTarget(FrameBuffer, *SlateWindow.GetTitle().ToString()));
+
+			// Create an internal RDG texture with the same extent and format as the source
+			FRDGTextureRef OutputTexture = GraphBuilder.CreateTexture(FRDGTextureDesc::Create2D(FrameBuffer->GetDesc().Extent, FrameBuffer->GetDesc().Format, FClearValueBinding::None, ETextureCreateFlags::Shared | ETextureCreateFlags::RenderTargetable), TEXT("VideoInputBackBufferCompositedStaging"));
+			// Bit cheeky, but when attempting to create two textures with the same description, RDG was just re-allocating which would lead to flickering. By converting to external, we force immediate allocation of the underlying pooled resource
+			TopLevelWindows[Index].SetTexture(GraphBuilder.ConvertToExternalTexture(OutputTexture));
+
+			// Initialise our copy paramaters
+			FCopyToStagingParameters* PassParameters = GraphBuilder.AllocParameters<FCopyToStagingParameters>();
+			PassParameters->Input = InputTexture;
+			PassParameters->Output = OutputTexture;
+
+			GraphBuilder.AddPass(
+				RDG_EVENT_NAME("VideoInputBackBufferCopyToStaging (%s)", *SlateWindow.GetTitle().ToString()),
+				PassParameters,
+				ERDGPassFlags::Readback,
+				[InputTexture, OutputTexture, Index, this](FRHICommandList& RHICmdList) {
+					// Simple texture copy as we know pixel format is the same
+					RHICmdList.CopyTexture(InputTexture->GetRHI(), OutputTexture->GetRHI(), {});
+				});
 
-		// Check all of our windows have a texture
-		bool bSkipComposite = TopLevelWindows.ContainsByPredicate([](FTexturedWindow Window) { return !Window.GetTexture(); });
+			GraphBuilder.Execute();
+		}
 
-		if (!bSkipComposite)
 		{
-			// All windows have a texture so we can composite
-			CompositeWindows();
+			// Check all of our windows have a texture
+			bool bSkipComposite = TopLevelWindows.ContainsByPredicate([](FTexturedWindow Window) { return !Window.GetTexture(); });
+
+			if (!bSkipComposite)
+			{
+				// All windows have a texture so we can composite
+				CompositeWindows();
+			}
 		}
 	}
 }
