Add: `EditorPixelStreamingSource` launch arg to set editor streaming source. This can be used in conjunction with `EditorPixelStreamingStartOnLaunch` to immediately start editor streaming the target source as soon as the editor opens.

diff --git a/Engine/Plugins/Media/PixelStreaming/Source/PixelStreamingEditor/Private/PixelStreamingEditorModule.cpp b/Engine/Plugins/Media/PixelStreaming/Source/PixelStreamingEditor/Private/PixelStreamingEditorModule.cpp
--- a/Engine/Plugins/Media/PixelStreaming/Source/PixelStreamingEditor/Private/PixelStreamingEditorModule.cpp
+++ b/Engine/Plugins/Media/PixelStreaming/Source/PixelStreamingEditor/Private/PixelStreamingEditorModule.cpp
@@ -101,7 +101,20 @@ void FPixelStreamingEditorModule::InitEditorStreaming(IPixelStreamingModule& Mod
 
 		if (UE::EditorPixelStreaming::Settings::CVarEditorPixelStreamingStartOnLaunch.GetValueOnAnyThread())
 		{
-			StartStreaming(UE::EditorPixelStreaming::EStreamTypes::Editor);
+            // Default our source to the full editor
+            UE::EditorPixelStreaming::EStreamTypes Source = UE::EditorPixelStreaming::EStreamTypes::Editor;
+
+            FString SourceStr = UE::EditorPixelStreaming::Settings::CVarEditorPixelStreamingSource.GetValueOnAnyThread();
+            if(SourceStr == TEXT("Editor") || SourceStr == TEXT("editor"))
+            {
+                Source = UE::EditorPixelStreaming::EStreamTypes::Editor;
+            }
+            else if(SourceStr == TEXT("LevelEditor") || SourceStr == TEXT("leveleditor"))
+            {
+                Source = UE::EditorPixelStreaming::EStreamTypes::LevelEditorViewport;
+            }
+
+			StartStreaming(Source);
 		}
 	});
 
diff --git a/Engine/Plugins/Media/PixelStreaming/Source/PixelStreamingEditor/Private/Settings.cpp b/Engine/Plugins/Media/PixelStreaming/Source/PixelStreamingEditor/Private/Settings.cpp
--- a/Engine/Plugins/Media/PixelStreaming/Source/PixelStreamingEditor/Private/Settings.cpp
+++ b/Engine/Plugins/Media/PixelStreaming/Source/PixelStreamingEditor/Private/Settings.cpp
@@ -17,11 +17,19 @@ namespace UE::EditorPixelStreaming::Settings
 		TEXT("Use a remote signalling server. Default: false"),
 		ECVF_Default);
 
+    TAutoConsoleVariable<FString> CVarEditorPixelStreamingSource(
+		TEXT("PixelStreaming.Editor.Source"),
+		TEXT("Editor"),
+		TEXT("Editor PixelStreaming source. Supported values are `Editor`, `LevelEditor`"),
+		ECVF_Default);
+
     void InitialiseSettings()
     {
 		using namespace UE::PixelStreaming;
         // Options parse (if these exist they are set to true)
 		CommandLineParseOption(TEXT("EditorPixelStreamingStartOnLaunch"), CVarEditorPixelStreamingStartOnLaunch);
 		CommandLineParseOption(TEXT("EditorPixelStreamingUseRemoteSignallingServer"), CVarEditorPixelStreamingUseRemoteSignallingServer);
+
+        CommandLineParseValue(TEXT("EditorPixelStreamingSource="), CVarEditorPixelStreamingSource);
     }
 }
\ No newline at end of file
diff --git a/Engine/Plugins/Media/PixelStreaming/Source/PixelStreamingEditor/Private/Settings.h b/Engine/Plugins/Media/PixelStreaming/Source/PixelStreamingEditor/Private/Settings.h
--- a/Engine/Plugins/Media/PixelStreaming/Source/PixelStreamingEditor/Private/Settings.h
+++ b/Engine/Plugins/Media/PixelStreaming/Source/PixelStreamingEditor/Private/Settings.h
@@ -13,5 +13,6 @@ namespace UE::EditorPixelStreaming::Settings
 	// Begin Editor CVars
 	extern TAutoConsoleVariable<bool> CVarEditorPixelStreamingStartOnLaunch;
 	extern TAutoConsoleVariable<bool> CVarEditorPixelStreamingUseRemoteSignallingServer;
+    extern TAutoConsoleVariable<FString> CVarEditorPixelStreamingSource;
     // End Editor CVars
 }
\ No newline at end of file
