Add: PixelStreaming messages now support 'String's as valid data types for a message structure.

diff --git a/Engine/Plugins/Media/PixelStreaming/Source/PixelStreamingInput/Private/PixelStreamingInputMessage.cpp b/Engine/Plugins/Media/PixelStreaming/Source/PixelStreamingInput/Private/PixelStreamingInputMessage.cpp
--- a/Engine/Plugins/Media/PixelStreaming/Source/PixelStreamingInput/Private/PixelStreamingInputMessage.cpp
+++ b/Engine/Plugins/Media/PixelStreaming/Source/PixelStreamingInput/Private/PixelStreamingInputMessage.cpp
@@ -9,7 +9,6 @@ FPixelStreamingInputMessage::FPixelStreamingInputMessage()
 	: Id(++CurrentId)
 	, Structure({})
 {
-	ByteLength = 0;
 }
 
 FPixelStreamingInputMessage::FPixelStreamingInputMessage(uint8 InId)
@@ -22,15 +21,12 @@ FPixelStreamingInputMessage::FPixelStreamingInputMessage(uint8 InId)
 	{
 		CurrentId = InId; // Ensure that any new message id starts after the highest specified Id
 	}
-
-	ByteLength = 0;
 }
 
 FPixelStreamingInputMessage::FPixelStreamingInputMessage(TArray<EPixelStreamingMessageTypes> InStructure)
 	: Id(++CurrentId)
 	, Structure(InStructure)
 {
-	ByteLength = CalculateByteLength();
 }
 
 FPixelStreamingInputMessage::FPixelStreamingInputMessage(uint8 InId, TArray<EPixelStreamingMessageTypes> InStructure)
@@ -43,36 +39,4 @@ FPixelStreamingInputMessage::FPixelStreamingInputMessage(uint8 InId, TArray<EPix
 	{
 		CurrentId = InId; // Ensure that any new message id starts after the highest specified Id
 	}
-
-	ByteLength = CalculateByteLength();
-}
-
-uint8 FPixelStreamingInputMessage::CalculateByteLength()
-{
-	uint8 Length = 0;
-	for (EPixelStreamingMessageTypes Type : Structure)
-	{
-		switch (Type)
-		{
-			case EPixelStreamingMessageTypes::Uint8:
-				Length += 1;
-				break;
-			case EPixelStreamingMessageTypes::Uint16:
-			case EPixelStreamingMessageTypes::Int16:
-				Length += 2;
-				break;
-			case EPixelStreamingMessageTypes::Float:
-				Length += 4;
-				break;
-			case EPixelStreamingMessageTypes::Double:
-				Length += 8;
-				break;
-			case EPixelStreamingMessageTypes::Variable:
-			case EPixelStreamingMessageTypes::Undefined:
-			default:
-				// These types of have no defined length
-				break;
-		}
-	}
-	return Length;
 }
\ No newline at end of file
diff --git a/Engine/Plugins/Media/PixelStreaming/Source/PixelStreamingInput/Private/PixelStreamingInputModule.cpp b/Engine/Plugins/Media/PixelStreaming/Source/PixelStreamingInput/Private/PixelStreamingInputModule.cpp
--- a/Engine/Plugins/Media/PixelStreaming/Source/PixelStreamingInput/Private/PixelStreamingInputModule.cpp
+++ b/Engine/Plugins/Media/PixelStreaming/Source/PixelStreamingInput/Private/PixelStreamingInputModule.cpp
@@ -64,17 +64,17 @@ namespace UE::PixelStreamingInput
 		FPixelStreamingInputProtocol::ToStreamerProtocol.Add("AverageBitrateRequest", FPixelStreamingInputMessage(3));
 		FPixelStreamingInputProtocol::ToStreamerProtocol.Add("StartStreaming", FPixelStreamingInputMessage(4));
 		FPixelStreamingInputProtocol::ToStreamerProtocol.Add("StopStreaming", FPixelStreamingInputMessage(5));
-		FPixelStreamingInputProtocol::ToStreamerProtocol.Add("LatencyTest", FPixelStreamingInputMessage(6));
+		FPixelStreamingInputProtocol::ToStreamerProtocol.Add("LatencyTest", FPixelStreamingInputMessage(6, { EType::String }));
 		FPixelStreamingInputProtocol::ToStreamerProtocol.Add("RequestInitialSettings", FPixelStreamingInputMessage(7));
-		FPixelStreamingInputProtocol::ToStreamerProtocol.Add("TestEcho", FPixelStreamingInputMessage(8));
+		FPixelStreamingInputProtocol::ToStreamerProtocol.Add("TestEcho", FPixelStreamingInputMessage(8, { EType::String }));
 
 		/*
 		 * Input Messages.
 		 */
 		// Generic Input Messages.
-		FPixelStreamingInputProtocol::ToStreamerProtocol.Add("UIInteraction", FPixelStreamingInputMessage(50));
-		FPixelStreamingInputProtocol::ToStreamerProtocol.Add("Command", FPixelStreamingInputMessage(51));
-		FPixelStreamingInputProtocol::ToStreamerProtocol.Add("TextboxEntry", FPixelStreamingInputMessage(52));
+		FPixelStreamingInputProtocol::ToStreamerProtocol.Add("UIInteraction", FPixelStreamingInputMessage(50, { EType::String }));
+		FPixelStreamingInputProtocol::ToStreamerProtocol.Add("Command", FPixelStreamingInputMessage(51, { EType::String }));
+		FPixelStreamingInputProtocol::ToStreamerProtocol.Add("TextboxEntry", FPixelStreamingInputMessage(52, { EType::String }));
 
 		// Keyboard Input Message.
 		// Complex command with payload, therefore we specify the length of the payload (bytes) as well as the structure of the payload
diff --git a/Engine/Plugins/Media/PixelStreaming/Source/PixelStreamingInput/Private/PixelStreamingInputProtocol.cpp b/Engine/Plugins/Media/PixelStreaming/Source/PixelStreamingInput/Private/PixelStreamingInputProtocol.cpp
--- a/Engine/Plugins/Media/PixelStreaming/Source/PixelStreamingInput/Private/PixelStreamingInputProtocol.cpp
+++ b/Engine/Plugins/Media/PixelStreaming/Source/PixelStreamingInput/Private/PixelStreamingInputProtocol.cpp
@@ -20,40 +20,41 @@ TSharedPtr<FJsonObject> FPixelStreamingInputProtocol::ToJson(EPixelStreamingMess
 		TSharedPtr<FJsonObject> MessageJson = MakeShareable(new FJsonObject());
 
 		MessageJson->SetField("id", MakeShared<FJsonValueNumber>(Value.GetID()));
-		MessageJson->SetField("byteLength", MakeShared<FJsonValueNumber>(Value.GetByteLength()));
 
-		if (Value.GetByteLength() > 0)
+        TArray<EPixelStreamingMessageTypes> Structure = Value.GetStructure();
+		TArray<TSharedPtr<FJsonValue>> StructureJson;
+		for (auto It = Structure.CreateIterator(); It; ++It)
 		{
-			TArray<TSharedPtr<FJsonValue>> StructureJson;
-			TArray<EPixelStreamingMessageTypes> Structure = Value.GetStructure();
-			for (auto It = Structure.CreateIterator(); It; ++It)
+			FString Text;
+			switch (*It)
 			{
-				FString Text;
-				switch (*It)
-				{
-					case EPixelStreamingMessageTypes::Uint8:
-						Text = "uint8";
-						break;
-					case EPixelStreamingMessageTypes::Uint16:
-						Text = "uint16";
-						break;
-					case EPixelStreamingMessageTypes::Int16:
-						Text = "int16";
-						break;
-					case EPixelStreamingMessageTypes::Float:
-						Text = "float";
-						break;
-					case EPixelStreamingMessageTypes::Double:
-						Text = "double";
-						break;
-					default:
-						Text = "";
-				}
-				TSharedRef<FJsonValueString> JsonValue = MakeShareable(new FJsonValueString(*Text));
-				StructureJson.Add(JsonValue);
+				case EPixelStreamingMessageTypes::Uint8:
+					Text = "uint8";
+					break;
+				case EPixelStreamingMessageTypes::Uint16:
+					Text = "uint16";
+					break;
+				case EPixelStreamingMessageTypes::Int16:
+					Text = "int16";
+					break;
+				case EPixelStreamingMessageTypes::Float:
+					Text = "float";
+					break;
+				case EPixelStreamingMessageTypes::Double:
+					Text = "double";
+					break;
+            	case EPixelStreamingMessageTypes::String:
+					Text = "string";
+					break;
+				case EPixelStreamingMessageTypes::Undefined:
+				case EPixelStreamingMessageTypes::Variable:
+				default:
+					Text = "";
 			}
-			MessageJson->SetArrayField("structure", StructureJson);
+			TSharedRef<FJsonValueString> JsonValue = MakeShareable(new FJsonValueString(*Text));
+			StructureJson.Add(JsonValue);
 		}
+		MessageJson->SetArrayField("structure", StructureJson);
 
 		ProtocolJson->SetField(*Key, MakeShared<FJsonValueObject>(MessageJson));
 	});
diff --git a/Engine/Plugins/Media/PixelStreaming/Source/PixelStreamingInput/Public/PixelStreamingInputEnums.h b/Engine/Plugins/Media/PixelStreaming/Source/PixelStreamingInput/Public/PixelStreamingInputEnums.h
--- a/Engine/Plugins/Media/PixelStreaming/Source/PixelStreamingInput/Public/PixelStreamingInputEnums.h
+++ b/Engine/Plugins/Media/PixelStreaming/Source/PixelStreamingInput/Public/PixelStreamingInputEnums.h
@@ -15,7 +15,8 @@ enum class PIXELSTREAMINGINPUT_API EPixelStreamingMessageTypes
 	Int16,
 	Float,
 	Double,
-	Variable,
+    Variable UE_DEPRECATED(5.3, "EPixelStreamingMessageTypes::Variable has been deprecated. Please use either String or Undefined instead"),
+    String,
 	Undefined
 };
 
diff --git a/Engine/Plugins/Media/PixelStreaming/Source/PixelStreamingInput/Public/PixelStreamingInputMessage.h b/Engine/Plugins/Media/PixelStreaming/Source/PixelStreamingInput/Public/PixelStreamingInputMessage.h
--- a/Engine/Plugins/Media/PixelStreaming/Source/PixelStreamingInput/Public/PixelStreamingInputMessage.h
+++ b/Engine/Plugins/Media/PixelStreaming/Source/PixelStreamingInput/Public/PixelStreamingInputMessage.h
@@ -10,28 +10,26 @@
 class PIXELSTREAMINGINPUT_API FPixelStreamingInputMessage
 {
 public:
-	// Emtpy constructor. Sets ID to next available ID, byteLength to 0, and an empty Structure
+	// Emtpy constructor. Sets ID to next available ID, and Structure to []
 	FPixelStreamingInputMessage();
 
-	// Constructor taking an ID. Sets byteLength to 0, and an empty Structure
+	// Constructor taking an ID. Sets Structure to []
 	FPixelStreamingInputMessage(uint8 InId);
 
-	// Constructor taking a Structure. Sets ID to next available ID and calculates bytelength from the structure
+	// Constructor taking a Structure. Sets ID to next available ID
 	FPixelStreamingInputMessage(TArray<EPixelStreamingMessageTypes> InStructure);
 
-	// Constructor taking an ID and a Structure. Calculates bytelength from the structure
+	// Constructor taking an ID and a Structure
 	FPixelStreamingInputMessage(uint8 InId, TArray<EPixelStreamingMessageTypes> InStructure);
 
-	uint8 GetByteLength() const { return ByteLength; }
+    UE_DEPRECATED(5.3, "ByteLength no longer exists on FPixelStreamingInputMessage.")
+    uint8 GetByteLength() const { return 0; }
 	uint8 GetID() const { return Id; }
 	TArray<EPixelStreamingMessageTypes> GetStructure() const { return Structure; }
 
 private:
-	uint8 CalculateByteLength();
-
 	uint8 Id;
 	TArray<EPixelStreamingMessageTypes> Structure;
-	uint8 ByteLength;
 
 	static uint8 CurrentId;
 };
\ No newline at end of file
