Fix: Disable warnings about EPixelStreamingMessageTypes::Variable being deprecated

diff --git a/Engine/Plugins/Media/PixelStreaming/Source/PixelStreamingInput/Private/PixelStreamingInputProtocol.cpp b/Engine/Plugins/Media/PixelStreaming/Source/PixelStreamingInput/Private/PixelStreamingInputProtocol.cpp
--- a/Engine/Plugins/Media/PixelStreaming/Source/PixelStreamingInput/Private/PixelStreamingInputProtocol.cpp
+++ b/Engine/Plugins/Media/PixelStreaming/Source/PixelStreamingInput/Private/PixelStreamingInputProtocol.cpp
@@ -7,6 +7,7 @@
 FInputProtocolMap FPixelStreamingInputProtocol::ToStreamerProtocol;
 FInputProtocolMap FPixelStreamingInputProtocol::FromStreamerProtocol;
 
+PRAGMA_DISABLE_DEPRECATION_WARNINGS
 TSharedPtr<FJsonObject> FPixelStreamingInputProtocol::ToJson(EPixelStreamingMessageDirection Direction)
 {
 	TSharedPtr<FJsonObject> ProtocolJson = MakeShareable(new FJsonObject());
@@ -60,4 +61,5 @@ TSharedPtr<FJsonObject> FPixelStreamingInputProtocol::ToJson(EPixelStreamingMess
 	});
 
 	return ProtocolJson;
-}
\ No newline at end of file
+}
+PRAGMA_ENABLE_DEPRECATION_WARNINGS
\ No newline at end of file
