Updating server fetch scripts to point to the correct version.

diff --git a/Engine/Plugins/Media/PixelStreaming/Resources/WebServers/get_ps_servers.bat b/Engine/Plugins/Media/PixelStreaming/Resources/WebServers/get_ps_servers.bat
--- a/Engine/Plugins/Media/PixelStreaming/Resources/WebServers/get_ps_servers.bat
+++ b/Engine/Plugins/Media/PixelStreaming/Resources/WebServers/get_ps_servers.bat
@@ -80,11 +80,15 @@ if DEFINED UEVersion (
     SET PSInfraTagOrBranch=UE5.2
     SET IsTag=0
   )
+  if "%UEVersion%"=="5.3" (
+    SET PSInfraTagOrBranch=UE5.3
+    SET IsTag=0
+  )
 )
 
 @Rem If no arguments select a specific version, fetch the appropriate default
 if NOT DEFINED PSInfraTagOrBranch (
-    SET PSInfraTagOrBranch=UE5.2
+    SET PSInfraTagOrBranch=UE5.3
     SET IsTag=0
 )
 echo Tag or branch: !PSInfraTagOrBranch!
diff --git a/Engine/Plugins/Media/PixelStreaming/Resources/WebServers/get_ps_servers.sh b/Engine/Plugins/Media/PixelStreaming/Resources/WebServers/get_ps_servers.sh
--- a/Engine/Plugins/Media/PixelStreaming/Resources/WebServers/get_ps_servers.sh
+++ b/Engine/Plugins/Media/PixelStreaming/Resources/WebServers/get_ps_servers.sh
@@ -26,7 +26,7 @@ print_help() {
 # Set all default variables (e.g. # Name and version of ps-infra that we are downloading)
 PSInfraOrg=EpicGames
 PSInfraRepo=PixelStreamingInfrastructure
-PSInfraTagOrBranch=UE5.2
+PSInfraTagOrBranch=UE5.3
 RefType=heads
 IsTag=0
 ReleaseUrlBase=https://github.com/EpicGames/PixelStreamingInfrastructure/releases/download
@@ -76,12 +76,17 @@ then
     PSInfraTagOrBranch=UE5.2
     IsTag=0
   fi
+  if [ "$UEVersion" = "5.3" ]
+  then
+    PSInfraTagOrBranch=UE5.3
+    IsTag=0
+  fi
 fi
 
 # If no arguments select a specific version, fetch the appropriate default
 if [ -z "$PSInfraTagOrBranch" ]
 then
-  PSInfraTagOrBranch=UE5.2
+  PSInfraTagOrBranch=UE5.3
   IsTag=0
 fi
 echo "Tag or branch: $PSInfraTagOrBranch"
