Fix: No level audio when editor streaming

diff --git a/Engine/Plugins/Media/PixelStreaming/Source/PixelStreamingEditor/Private/EditorSubmixListener.cpp b/Engine/Plugins/Media/PixelStreaming/Source/PixelStreamingEditor/Private/EditorSubmixListener.cpp
new file mode 100644
--- /dev/null
+++ b/Engine/Plugins/Media/PixelStreaming/Source/PixelStreamingEditor/Private/EditorSubmixListener.cpp
@@ -0,0 +1,23 @@
+// Copyright Epic Games, Inc. All Rights Reserved.
+
+#include "EditorSubmixListener.h"
+
+#include "AudioDevice.h"
+#include "PixelStreamingPeerConnection.h"
+
+namespace UE::EditorPixelStreaming
+{
+	FEditorSubmixListener::FEditorSubmixListener(FAudioDeviceHandle AudioDevice)
+		: AudioInput(FPixelStreamingPeerConnection::CreateAudioInput())
+	{
+		AudioDevice->RegisterSubmixBufferListener(this);
+	}
+
+	void FEditorSubmixListener::OnNewSubmixBuffer(const USoundSubmix* OwningSubmix, float* AudioData,
+		int32 NumSamples, int32 NumChannels,
+		const int32 SampleRate, double AudioClock)
+	{
+		AudioInput->PushAudio(AudioData, NumSamples, NumChannels, SampleRate);
+	}
+
+} // namespace UE::EditorPixelStreaming
diff --git a/Engine/Plugins/Media/PixelStreaming/Source/PixelStreamingEditor/Private/EditorSubmixListener.h b/Engine/Plugins/Media/PixelStreaming/Source/PixelStreamingEditor/Private/EditorSubmixListener.h
new file mode 100644
--- /dev/null
+++ b/Engine/Plugins/Media/PixelStreaming/Source/PixelStreamingEditor/Private/EditorSubmixListener.h
@@ -0,0 +1,25 @@
+// Copyright Epic Games, Inc. All Rights Reserved.
+
+#pragma once
+
+#include "AudioDeviceHandle.h"
+#include "ISubmixBufferListener.h"
+#include "IPixelStreamingAudioInput.h"
+
+namespace UE::EditorPixelStreaming
+{
+	class FEditorSubmixListener : public ISubmixBufferListener
+	{
+	public:
+		FEditorSubmixListener(FAudioDeviceHandle AudioDevice);
+		virtual ~FEditorSubmixListener() = default;
+
+		// ISubmixBufferListener interface
+		void OnNewSubmixBuffer(const USoundSubmix* OwningSubmix, float* AudioData,
+			int32 NumSamples, int32 NumChannels,
+			const int32 SampleRate, double AudioClock) override;
+
+	private:
+		TSharedPtr<IPixelStreamingAudioInput> AudioInput;
+	};
+} // namespace UE::EditorPixelStreaming
\ No newline at end of file
diff --git a/Engine/Plugins/Media/PixelStreaming/Source/PixelStreamingEditor/Private/PixelStreamingEditorModule.cpp b/Engine/Plugins/Media/PixelStreaming/Source/PixelStreamingEditor/Private/PixelStreamingEditorModule.cpp
--- a/Engine/Plugins/Media/PixelStreaming/Source/PixelStreamingEditor/Private/PixelStreamingEditorModule.cpp
+++ b/Engine/Plugins/Media/PixelStreaming/Source/PixelStreamingEditor/Private/PixelStreamingEditorModule.cpp
@@ -25,6 +25,8 @@
 #include "Framework/Notifications/NotificationManager.h"
 #include "Widgets/Notifications/SNotificationList.h"
 
+#include "AudioDeviceManager.h"
+
 #define LOCTEXT_NAMESPACE "PixelStreamingEditorModule"
 
 DECLARE_LOG_CATEGORY_EXTERN(LogPixelStreamingEditor, Log, All);
@@ -49,6 +51,19 @@ void FPixelStreamingEditorModule::StartupModule()
 	Settings::InitialiseSettings();
 	bUseExternalSignallingServer = Settings::CVarEditorPixelStreamingUseRemoteSignallingServer.GetValueOnAnyThread();
 
+    // Enable experimental audio so we can mix the different editor audio outputs
+    IConsoleVariable* CVar = IConsoleManager::Get().FindConsoleVariable(TEXT("PixelStreaming.ExperimentalAudioInput")); 
+    CVar->Set(true);
+
+    FAudioDeviceManagerDelegates::OnAudioDeviceCreated.AddLambda([&](Audio::FDeviceId AudioDeviceId) {
+	    FAudioDeviceManager* DeviceManager = GEngine->GetAudioDeviceManager();
+	    FAudioDeviceHandle Device = DeviceManager->GetAudioDevice(AudioDeviceId);
+        AudioInputs.Add(AudioDeviceId, MakeShared<FEditorSubmixListener>(Device));
+    });
+    FAudioDeviceManagerDelegates::OnAudioDeviceDestroyed.AddLambda([&](Audio::FDeviceId AudioDeviceId) {
+        AudioInputs.Remove(AudioDeviceId);
+    });
+
 	IPixelStreamingModule& Module = IPixelStreamingModule::Get();
 	Module.OnReady().AddRaw(this, &FPixelStreamingEditorModule::InitEditorStreaming);
 }
diff --git a/Engine/Plugins/Media/PixelStreaming/Source/PixelStreamingEditor/Private/PixelStreamingEditorModule.h b/Engine/Plugins/Media/PixelStreaming/Source/PixelStreamingEditor/Private/PixelStreamingEditorModule.h
--- a/Engine/Plugins/Media/PixelStreaming/Source/PixelStreamingEditor/Private/PixelStreamingEditorModule.h
+++ b/Engine/Plugins/Media/PixelStreaming/Source/PixelStreamingEditor/Private/PixelStreamingEditorModule.h
@@ -5,14 +5,17 @@
 #include "IPixelStreamingModule.h"
 #include "PixelStreamingServers.h"
 #include "PixelStreamingEditorUtils.h"
+#include "IPixelStreamingAudioInput.h"
 #include "IPixelStreamingEditorModule.h"
+#include "EditorSubmixListener.h"
+#include "AudioDeviceHandle.h"
 
 namespace UE::EditorPixelStreaming
 {
 	class FPixelStreamingToolbar;
 }
 
-class PIXELSTREAMINGEDITOR_API FPixelStreamingEditorModule : public IPixelStreamingEditorModule
+class FPixelStreamingEditorModule : public IPixelStreamingEditorModule
 {
 public:
 	/** IModuleInterface implementation */
@@ -64,4 +67,6 @@ private:
 	TSharedPtr<IPixelStreamingStreamer> EditorStreamer;
 
 	bool bUseExternalSignallingServer = false;
+
+	TMap<Audio::FDeviceId, TSharedPtr<UE::EditorPixelStreaming::FEditorSubmixListener>> AudioInputs;
 };
