Fix: Remove verbose disabling of deprecation warnings and instead use UEs inbuilt definition

diff --git a/Engine/Plugins/Media/PixelStreaming/Source/PixelStreaming/Private/Streamer.cpp b/Engine/Plugins/Media/PixelStreaming/Source/PixelStreaming/Private/Streamer.cpp
--- a/Engine/Plugins/Media/PixelStreaming/Source/PixelStreaming/Private/Streamer.cpp
+++ b/Engine/Plugins/Media/PixelStreaming/Source/PixelStreaming/Private/Streamer.cpp
@@ -149,19 +149,9 @@ namespace UE::PixelStreaming
 	{
 		// This method is marked as deprecated but still calls the deprecated method on the input handler. As such, we disable
 		// the warnings that arise from using the input handlers method
-#if PLATFORM_WINDOWS
-	#pragma warning(push)
-	#pragma warning(disable : 4996)
-#elif PLATFORM_LINUX
-	#pragma clang diagnostic push
-	#pragma clang diagnostic ignored "-Wdeprecated-declarations"
-#endif
+		PRAGMA_DISABLE_DEPRECATION_WARNINGS
 		InputHandler->SetTargetScreenSize(InTargetScreenSize);
-#if PLATFORM_WINDOWS
-	#pragma warning(pop)
-#elif PLATFORM_LINUX
-	#pragma clang diagnostic pop
-#endif
+		PRAGMA_ENABLE_DEPRECATION_WARNINGS
 	}
 
 	TWeakPtr<FIntPoint> FStreamer::GetTargetScreenSize()
