Fix Editor Streaming memory leak

diff --git a/Engine/Plugins/Media/PixelStreaming/Source/PixelStreaming/Private/UtilsRender.h b/Engine/Plugins/Media/PixelStreaming/Source/PixelStreaming/Private/UtilsRender.h
--- a/Engine/Plugins/Media/PixelStreaming/Source/PixelStreaming/Private/UtilsRender.h
+++ b/Engine/Plugins/Media/PixelStreaming/Source/PixelStreaming/Private/UtilsRender.h
@@ -151,7 +151,7 @@ namespace UE::PixelStreaming
 			const FIntRect ViewRect(FIntPoint(0, 0), InputTexture->Desc.Extent);
 
 			//Dummy ViewFamily/ViewInfo created to use built in Draw Screen/Texture Pass
-			FSceneViewFamily ViewFamily(FSceneViewFamily::ConstructionValues(nullptr, nullptr, FEngineShowFlags(ESFIM_Game))
+			FSceneViewFamilyContext ViewFamily(FSceneViewFamily::ConstructionValues(nullptr, nullptr, FEngineShowFlags(ESFIM_Game))
 				.SetTime(FGameTime())
 				.SetGammaCorrection(1.0f));
 			FSceneViewInitOptions ViewInitOptions;
diff --git a/Engine/Plugins/Media/PixelStreaming/Source/PixelStreamingEditor/Private/PixelStreamingVideoInputBackBufferComposited.cpp b/Engine/Plugins/Media/PixelStreaming/Source/PixelStreamingEditor/Private/PixelStreamingVideoInputBackBufferComposited.cpp
--- a/Engine/Plugins/Media/PixelStreaming/Source/PixelStreamingEditor/Private/PixelStreamingVideoInputBackBufferComposited.cpp
+++ b/Engine/Plugins/Media/PixelStreaming/Source/PixelStreamingEditor/Private/PixelStreamingVideoInputBackBufferComposited.cpp
@@ -230,7 +230,7 @@ void FPixelStreamingVideoInputBackBufferComposited::CompositeWindows()
 		const FIntRect ViewRect(FIntPoint(0, 0), Extent);
 
 		// Dummy ViewFamily/ViewInfo created to use built in Draw Screen/Texture Pass
-		FSceneViewFamily ViewFamily(FSceneViewFamily::ConstructionValues(nullptr, nullptr, FEngineShowFlags(ESFIM_Game))
+		FSceneViewFamilyContext ViewFamily(FSceneViewFamily::ConstructionValues(nullptr, nullptr, FEngineShowFlags(ESFIM_Game))
 										.SetTime(FGameTime())
 										.SetGammaCorrection(1.0f));
 		FSceneViewInitOptions ViewInitOptions;
