Pixel Streaming: Fix crash caused by routing input while auto saving.

diff --git a/Engine/Plugins/Media/PixelStreaming/Source/PixelStreamingInput/Private/PixelStreamingInputHandler.cpp b/Engine/Plugins/Media/PixelStreaming/Source/PixelStreamingInput/Private/PixelStreamingInputHandler.cpp
--- a/Engine/Plugins/Media/PixelStreaming/Source/PixelStreamingInput/Private/PixelStreamingInputHandler.cpp
+++ b/Engine/Plugins/Media/PixelStreaming/Source/PixelStreamingInput/Private/PixelStreamingInputHandler.cpp
@@ -1,6 +1,7 @@
 // Copyright Epic Games, Inc. All Rights Reserved.
 
 #include "PixelStreamingInputHandler.h"
+
 #include "InputStructures.h"
 #include "PixelStreamingInputProtocol.h"
 #include "JavaScriptKeyCodes.inl"
@@ -25,6 +26,10 @@
 #include "Widgets/Input/SMultiLineEditableTextBox.h"
 #include "PixelStreamingInputDevice.h"
 
+#if WITH_EDITOR
+#include "UObject/UObjectGlobals.h"
+#endif
+
 DECLARE_LOG_CATEGORY_EXTERN(LogPixelStreamingInputHandler, Log, VeryVerbose);
 DEFINE_LOG_CATEGORY(LogPixelStreamingInputHandler);
 
@@ -238,6 +243,29 @@ namespace UE::PixelStreamingInput
 
 	void FPixelStreamingInputHandler::Tick(const float InDeltaTime)
 	{
+#if WITH_EDITOR
+		/* No routing input while saving ... this is relevant for auto-save and can cause an incredibly rare crash...
+		 * 
+		 * The gist is that the auto-save system calls FSlateApplication::Tick(), which executes its OnPreTick() containing
+		 * our FPixelStreamingInputHandler::Tick. Routing any input executes Slate delegates. Again, the gist is that
+		 * the delegates can do anything including calling StaticConstructObject(), which will crash the editor
+		 * ("Illegal call to StaticConstructObject() while serializing object data!").
+		 * An example of a StaticConstructObject call is a UMG widget calling CreateWidget in response to a button's OnClick (which we routed!).
+		 *
+		 * If you're curious why our Tick gets called by auto-save:
+		 * The auto save starts in FPackageAutoSaver::AttemptAutoSave, which calls FEditorFileUtils::AutosaveMapEx.
+		 * This causes the world package to be saved (UEditorEngine::SavePackage) with a FSlowTask.
+		 * The slow task calls FFeedbackContextEditor::ProgressReported... which ticks slate so the progres bar modal window updates.
+		 * Consult with FPixelStreamingInputDevice::FPixelStreamingInputDevice, which explicitly wants to tick when a modal window is open.
+		 *
+		 * TLDR: if we're auto-saving, we'll postbone routing input until the auto save is done.
+		 */
+		if (GIsSavingPackage)
+		{
+			return;
+		}
+#endif
+		
 		TouchIndicesProcessedThisFrame.Reset();
 
 		FMessage Message;
