Fix for asan issue caused by missing UnregisterSubmixBufferListener call on destruction of pixel stream editor listener.
#rnx

diff --git a/Engine/Plugins/Media/PixelStreaming/Source/PixelStreamingEditor/Private/EditorSubmixListener.cpp b/Engine/Plugins/Media/PixelStreaming/Source/PixelStreamingEditor/Private/EditorSubmixListener.cpp
--- a/Engine/Plugins/Media/PixelStreaming/Source/PixelStreamingEditor/Private/EditorSubmixListener.cpp
+++ b/Engine/Plugins/Media/PixelStreaming/Source/PixelStreamingEditor/Private/EditorSubmixListener.cpp
@@ -8,9 +8,21 @@
 namespace UE::EditorPixelStreaming
 {
 	FEditorSubmixListener::FEditorSubmixListener(FAudioDeviceHandle AudioDevice)
-		: AudioInput(FPixelStreamingPeerConnection::CreateAudioInput())
+		: AudioDeviceId(AudioDevice.GetDeviceID())
+		, AudioInput(FPixelStreamingPeerConnection::CreateAudioInput())
 	{
-		AudioDevice->RegisterSubmixBufferListener(this);
+		if (AudioDevice.IsValid())
+		{
+			AudioDevice->RegisterSubmixBufferListener(this);
+		}
+	}
+
+	FEditorSubmixListener::~FEditorSubmixListener()
+	{
+		if (FAudioDevice* AudioDevice = FAudioDeviceManager::Get()->GetAudioDeviceRaw(AudioDeviceId))
+		{
+			AudioDevice->UnregisterSubmixBufferListener(this);
+		}
 	}
 
 	void FEditorSubmixListener::OnNewSubmixBuffer(const USoundSubmix* OwningSubmix, float* AudioData,
diff --git a/Engine/Plugins/Media/PixelStreaming/Source/PixelStreamingEditor/Private/EditorSubmixListener.h b/Engine/Plugins/Media/PixelStreaming/Source/PixelStreamingEditor/Private/EditorSubmixListener.h
--- a/Engine/Plugins/Media/PixelStreaming/Source/PixelStreamingEditor/Private/EditorSubmixListener.h
+++ b/Engine/Plugins/Media/PixelStreaming/Source/PixelStreamingEditor/Private/EditorSubmixListener.h
@@ -12,7 +12,7 @@ namespace UE::EditorPixelStreaming
 	{
 	public:
 		FEditorSubmixListener(FAudioDeviceHandle AudioDevice);
-		virtual ~FEditorSubmixListener() = default;
+		virtual ~FEditorSubmixListener();
 
 		// ISubmixBufferListener interface
 		void OnNewSubmixBuffer(const USoundSubmix* OwningSubmix, float* AudioData,
@@ -20,6 +20,7 @@ namespace UE::EditorPixelStreaming
 			const int32 SampleRate, double AudioClock) override;
 
 	private:
+		Audio::FDeviceId AudioDeviceId;
 		TSharedPtr<IPixelStreamingAudioInput> AudioInput;
 	};
 } // namespace UE::EditorPixelStreaming
\ No newline at end of file
