Fix for editor crash in pixel streaming due to AsShared being disallowed in ctor/dtor
#rnx

diff --git a/Engine/Plugins/Media/PixelStreaming/Source/PixelStreamingEditor/Private/EditorSubmixListener.cpp b/Engine/Plugins/Media/PixelStreaming/Source/PixelStreamingEditor/Private/EditorSubmixListener.cpp
--- a/Engine/Plugins/Media/PixelStreaming/Source/PixelStreamingEditor/Private/EditorSubmixListener.cpp
+++ b/Engine/Plugins/Media/PixelStreaming/Source/PixelStreamingEditor/Private/EditorSubmixListener.cpp
@@ -7,17 +7,28 @@
 
 namespace UE::EditorPixelStreaming
 {
-	FEditorSubmixListener::FEditorSubmixListener(FAudioDeviceHandle AudioDevice)
-		: AudioDeviceId(AudioDevice.GetDeviceID())
-		, AudioInput(FPixelStreamingPeerConnection::CreateAudioInput())
+	TSharedRef<FEditorSubmixListener, ESPMode::ThreadSafe> FEditorSubmixListener::FEditorSubmixListener::Create(const FAudioDeviceHandle& AudioDevice)
 	{
+		TSharedRef<FEditorSubmixListener, ESPMode::ThreadSafe> NewListener = MakeShared<FEditorSubmixListener, ESPMode::ThreadSafe>(AudioDevice);
 		if (AudioDevice.IsValid())
 		{
-			AudioDevice->RegisterSubmixBufferListener(AsShared(), AudioDevice->GetMainSubmixObject());
+			AudioDevice.GetAudioDevice()->RegisterSubmixBufferListener(NewListener, AudioDevice->GetMainSubmixObject());
 		}
+
+		return NewListener;
+	}
+
+	FEditorSubmixListener::FEditorSubmixListener(const FAudioDeviceHandle& AudioDevice)
+		: AudioDeviceId(AudioDevice.GetDeviceID())
+		, AudioInput(FPixelStreamingPeerConnection::CreateAudioInput())
+	{
 	}
 
 	FEditorSubmixListener::~FEditorSubmixListener()
+	{
+	}
+
+	void FEditorSubmixListener::Shutdown()
 	{
 		if (FAudioDevice* AudioDevice = FAudioDeviceManager::Get()->GetAudioDeviceRaw(AudioDeviceId))
 		{
diff --git a/Engine/Plugins/Media/PixelStreaming/Source/PixelStreamingEditor/Private/EditorSubmixListener.h b/Engine/Plugins/Media/PixelStreaming/Source/PixelStreamingEditor/Private/EditorSubmixListener.h
--- a/Engine/Plugins/Media/PixelStreaming/Source/PixelStreamingEditor/Private/EditorSubmixListener.h
+++ b/Engine/Plugins/Media/PixelStreaming/Source/PixelStreamingEditor/Private/EditorSubmixListener.h
@@ -11,7 +11,7 @@ namespace UE::EditorPixelStreaming
 	class FEditorSubmixListener : public ISubmixBufferListener
 	{
 	public:
-		FEditorSubmixListener(FAudioDeviceHandle AudioDevice);
+		static TSharedRef<FEditorSubmixListener> Create(const FAudioDeviceHandle& DeviceHandle);
 		virtual ~FEditorSubmixListener();
 
 		// ISubmixBufferListener interface
@@ -19,6 +19,14 @@ namespace UE::EditorPixelStreaming
 			int32 NumSamples, int32 NumChannels,
 			const int32 SampleRate, double AudioClock) override;
 
+		void Shutdown();
+
+	protected:
+		template <typename ObjectType, ESPMode Mode>
+		friend class SharedPointerInternals::TIntrusiveReferenceController;
+
+		FEditorSubmixListener(const FAudioDeviceHandle& AudioDevice);
+
 	private:
 		Audio::FDeviceId AudioDeviceId;
 		TSharedPtr<IPixelStreamingAudioInput> AudioInput;
diff --git a/Engine/Plugins/Media/PixelStreaming/Source/PixelStreamingEditor/Private/PixelStreamingEditorModule.cpp b/Engine/Plugins/Media/PixelStreaming/Source/PixelStreamingEditor/Private/PixelStreamingEditorModule.cpp
--- a/Engine/Plugins/Media/PixelStreaming/Source/PixelStreamingEditor/Private/PixelStreamingEditorModule.cpp
+++ b/Engine/Plugins/Media/PixelStreaming/Source/PixelStreamingEditor/Private/PixelStreamingEditorModule.cpp
@@ -59,11 +59,16 @@ void FPixelStreamingEditorModule::StartupModule()
 	{
 		FAudioDeviceManager* DeviceManager = GEngine->GetAudioDeviceManager();
 		FAudioDeviceHandle Device = DeviceManager->GetAudioDevice(AudioDeviceId);
-		AudioInputs.Add(AudioDeviceId, MakeShared<FEditorSubmixListener, ESPMode::ThreadSafe>(Device));
+		AudioInputs.Add(AudioDeviceId, FEditorSubmixListener::Create(Device));
 	});
 	FAudioDeviceManagerDelegates::OnAudioDeviceDestroyed.AddLambda([&](Audio::FDeviceId AudioDeviceId)
 	{
-		AudioInputs.Remove(AudioDeviceId);
+		TSharedPtr<UE::EditorPixelStreaming::FEditorSubmixListener, ESPMode::ThreadSafe> Input = AudioInputs.FindRef(AudioDeviceId);
+		if (Input.IsValid())
+		{
+			Input->Shutdown();
+			AudioInputs.Remove(AudioDeviceId);
+		}
 	});
 
 	IPixelStreamingModule& Module = IPixelStreamingModule::Get();
