Missing reshelf to fix build
#rnx

diff --git a/Engine/Plugins/Media/PixelStreaming/Source/PixelStreamingEditor/Private/EditorSubmixListener.h b/Engine/Plugins/Media/PixelStreaming/Source/PixelStreamingEditor/Private/EditorSubmixListener.h
--- a/Engine/Plugins/Media/PixelStreaming/Source/PixelStreamingEditor/Private/EditorSubmixListener.h
+++ b/Engine/Plugins/Media/PixelStreaming/Source/PixelStreamingEditor/Private/EditorSubmixListener.h
@@ -11,7 +11,7 @@ namespace UE::EditorPixelStreaming
 	class FEditorSubmixListener : public ISubmixBufferListener
 	{
 	public:
-		static TSharedRef<FEditorSubmixListener> Create(const FAudioDeviceHandle& DeviceHandle);
+		static TSharedRef<FEditorSubmixListener, ESPMode::ThreadSafe> Create(const FAudioDeviceHandle& DeviceHandle);
 		virtual ~FEditorSubmixListener();
 
 		// ISubmixBufferListener interface
