Move UniquePtr gen to SharedPtr gen to support ISubmixBufferListener API update
#rnx

diff --git a/Engine/Plugins/Media/PixelStreaming/Source/PixelStreaming/Private/PixelStreamingAudioDeviceModule.cpp b/Engine/Plugins/Media/PixelStreaming/Source/PixelStreaming/Private/PixelStreamingAudioDeviceModule.cpp
--- a/Engine/Plugins/Media/PixelStreaming/Source/PixelStreaming/Private/PixelStreamingAudioDeviceModule.cpp
+++ b/Engine/Plugins/Media/PixelStreaming/Source/PixelStreaming/Private/PixelStreamingAudioDeviceModule.cpp
@@ -33,7 +33,7 @@ namespace UE::PixelStreaming
 
 	FAudioDeviceModule::FAudioDeviceModule(TSharedPtr<UE::PixelStreaming::FAudioInputMixer> InMixer)
 		: bInitialized(false)
-		, Capturer(MakeUnique<UE::PixelStreaming::FAudioSubmixCapturer>())
+		, Capturer(MakeShared<UE::PixelStreaming::FAudioSubmixCapturer, ESPMode::ThreadSafe>())
 		, Requester(MakeUnique<UE::PixelStreaming::FAudioPlayoutRequester>())
 		, InputMixer(InMixer)
 	{
diff --git a/Engine/Plugins/Media/PixelStreaming/Source/PixelStreaming/Private/PixelStreamingAudioDeviceModule.h b/Engine/Plugins/Media/PixelStreaming/Source/PixelStreaming/Private/PixelStreamingAudioDeviceModule.h
--- a/Engine/Plugins/Media/PixelStreaming/Source/PixelStreaming/Private/PixelStreamingAudioDeviceModule.h
+++ b/Engine/Plugins/Media/PixelStreaming/Source/PixelStreaming/Private/PixelStreamingAudioDeviceModule.h
@@ -26,7 +26,7 @@ namespace UE::PixelStreaming
 
 	private:
 		FThreadSafeBool bInitialized; // True when we setup capturer/playout requester.
-		TUniquePtr<UE::PixelStreaming::FAudioSubmixCapturer> Capturer;
+		TSharedPtr<UE::PixelStreaming::FAudioSubmixCapturer, ESPMode::ThreadSafe> Capturer;
 		TUniquePtr<UE::PixelStreaming::FAudioPlayoutRequester> Requester;
 		TSharedPtr<UE::PixelStreaming::FAudioInputMixer> InputMixer;
 
