Remove private include paths from pixel streaming editor.

diff --git a/Engine/Plugins/Media/PixelStreaming/Source/PixelStreamingEditor/PixelStreamingEditor.build.cs b/Engine/Plugins/Media/PixelStreaming/Source/PixelStreamingEditor/PixelStreamingEditor.build.cs
--- a/Engine/Plugins/Media/PixelStreaming/Source/PixelStreamingEditor/PixelStreamingEditor.build.cs
+++ b/Engine/Plugins/Media/PixelStreaming/Source/PixelStreamingEditor/PixelStreamingEditor.build.cs
@@ -10,12 +10,6 @@ namespace UnrealBuildTool.Rules
 		{
 			var EngineDir = Path.GetFullPath(Target.RelativeEnginePath);
 
-			// NOTE: General rule is not to access the private folder of another module
-			PrivateIncludePaths.AddRange(new string[]
-			{
-				Path.Combine(EngineDir, "Plugins/Media/PixelStreaming/Source/PixelStreaming/Private"),
-			});
-
 			PrivateDependencyModuleNames.AddRange(new string[]
 			{
 				"Core",
diff --git a/Engine/Plugins/Media/PixelStreaming/Source/PixelStreamingEditor/Private/PixelStreamingVideoInputBackBufferComposited.cpp b/Engine/Plugins/Media/PixelStreaming/Source/PixelStreamingEditor/Private/PixelStreamingVideoInputBackBufferComposited.cpp
--- a/Engine/Plugins/Media/PixelStreaming/Source/PixelStreamingEditor/Private/PixelStreamingVideoInputBackBufferComposited.cpp
+++ b/Engine/Plugins/Media/PixelStreaming/Source/PixelStreamingEditor/Private/PixelStreamingVideoInputBackBufferComposited.cpp
@@ -2,6 +2,7 @@
 
 #include "PixelStreamingVideoInputBackBufferComposited.h"
 
+#include "Async/Async.h"
 #include "EngineModule.h"
 #include "Framework/Application/SlateApplication.h"
 #include "MediaShaders.h"
@@ -32,7 +33,7 @@ TSharedPtr<FPixelStreamingVideoInputBackBufferComposited> FPixelStreamingVideoIn
 	TSharedPtr<FPixelStreamingVideoInputBackBufferComposited> NewInput = TSharedPtr<FPixelStreamingVideoInputBackBufferComposited>(new FPixelStreamingVideoInputBackBufferComposited());
 	TWeakPtr<FPixelStreamingVideoInputBackBufferComposited> WeakInput = NewInput;
 	// Set up the callback on the game thread since FSlateApplication::Get() can only be used there
-	UE::PixelStreaming::DoOnGameThread([WeakInput]() {
+	AsyncTask(ENamedThreads::GameThread, [WeakInput]() {
 		if (TSharedPtr<FPixelStreamingVideoInputBackBufferComposited> Input = WeakInput.Pin())
 		{
 			FSlateRenderer* Renderer = FSlateApplication::Get().GetRenderer();
@@ -45,7 +46,7 @@ TSharedPtr<FPixelStreamingVideoInputBackBufferComposited> FPixelStreamingVideoIn
 FPixelStreamingVideoInputBackBufferComposited::FPixelStreamingVideoInputBackBufferComposited()
 {
 	SharedFrameRect = MakeShared<FIntRect>();
-	UE::PixelStreaming::DoOnGameThread([this]() { FSlateApplication::Get().OnPreTick().AddLambda([this](float DeltaTime) {
+	AsyncTask(ENamedThreads::GameThread, [this]() { FSlateApplication::Get().OnPreTick().AddLambda([this](float DeltaTime) {
 			FScopeLock Lock(&TopLevelWindowsCriticalSection);
 			TArray<TSharedRef<SWindow>> TopLevelSlateWindows;
 			FSlateApplication::Get().GetAllVisibleWindowsOrdered(TopLevelSlateWindows);
@@ -92,7 +93,7 @@ FPixelStreamingVideoInputBackBufferComposited::~FPixelStreamingVideoInputBackBuf
 {
 	if (!IsEngineExitRequested())
 	{
-		UE::PixelStreaming::DoOnGameThread([HandleCopy = DelegateHandle]() { FSlateApplication::Get().GetRenderer()->OnBackBufferReadyToPresent().Remove(HandleCopy); });
+		AsyncTask(ENamedThreads::GameThread, [HandleCopy = DelegateHandle]() { FSlateApplication::Get().GetRenderer()->OnBackBufferReadyToPresent().Remove(HandleCopy); });
 	}
 }
 
diff --git a/Engine/Plugins/Media/PixelStreaming/Source/PixelStreamingEditor/Private/PixelStreamingVideoInputViewport.cpp b/Engine/Plugins/Media/PixelStreaming/Source/PixelStreamingEditor/Private/PixelStreamingVideoInputViewport.cpp
--- a/Engine/Plugins/Media/PixelStreaming/Source/PixelStreamingEditor/Private/PixelStreamingVideoInputViewport.cpp
+++ b/Engine/Plugins/Media/PixelStreaming/Source/PixelStreamingEditor/Private/PixelStreamingVideoInputViewport.cpp
@@ -1,6 +1,8 @@
 // Copyright Epic Games, Inc. All Rights Reserved.
 
 #include "PixelStreamingVideoInputViewport.h"
+
+#include "Async/Async.h"
 #include "Settings.h"
 #include "Utils.h"
 #include "PixelCaptureInputFrameRHI.h"
@@ -22,7 +24,7 @@ TSharedPtr<FPixelStreamingVideoInputViewport> FPixelStreamingVideoInputViewport:
 	TWeakPtr<FPixelStreamingVideoInputViewport> WeakInput = NewInput;
 
 	// Set up the callback on the game thread since FSlateApplication::Get() can only be used there
-	UE::PixelStreaming::DoOnGameThread([WeakInput]() {
+	AsyncTask(ENamedThreads::GameThread, [WeakInput]() {
 		if (TSharedPtr<FPixelStreamingVideoInputViewport> Input = WeakInput.Pin())
 		{
 			Input->DelegateHandle = UGameViewportClient::OnViewportRendered().AddSP(Input.ToSharedRef(), &FPixelStreamingVideoInputViewport::OnViewportRendered);
@@ -36,7 +38,7 @@ FPixelStreamingVideoInputViewport::~FPixelStreamingVideoInputViewport()
 {
 	if (!IsEngineExitRequested())
 	{
-		UE::PixelStreaming::DoOnGameThread([HandleCopy = DelegateHandle]() {
+		AsyncTask(ENamedThreads::GameThread, [HandleCopy = DelegateHandle]() {
 			UGameViewportClient::OnViewportRendered().Remove(HandleCopy);
 		});
 	}
diff --git a/Engine/Plugins/Media/PixelStreaming/Source/PixelStreamingEditor/Private/Settings.cpp b/Engine/Plugins/Media/PixelStreaming/Source/PixelStreamingEditor/Private/Settings.cpp
--- a/Engine/Plugins/Media/PixelStreaming/Source/PixelStreamingEditor/Private/Settings.cpp
+++ b/Engine/Plugins/Media/PixelStreaming/Source/PixelStreamingEditor/Private/Settings.cpp
@@ -1,11 +1,51 @@
 // Copyright Epic Games, Inc. All Rights Reserved.
 
 #include "Settings.h"
-#include "Utils.h"
 
 namespace UE::EditorPixelStreaming::Settings
 {
-    TAutoConsoleVariable<bool> CVarEditorPixelStreamingStartOnLaunch(
+    template <typename T>
+	inline void CommandLineParseValue(const TCHAR* Match, TAutoConsoleVariable<T>& CVar)
+	{
+		T Value;
+		if (FParse::Value(FCommandLine::Get(), Match, Value))
+		{
+			CVar->Set(Value, ECVF_SetByCommandline);
+		}
+	};
+
+	inline void CommandLineParseValue(const TCHAR* Match, TAutoConsoleVariable<FString>& CVar, bool bStopOnSeparator = false)
+	{
+		FString Value;
+		if (FParse::Value(FCommandLine::Get(), Match, Value, bStopOnSeparator))
+		{
+			CVar->Set(*Value, ECVF_SetByCommandline);
+		}
+	};
+
+	inline void CommandLineParseOption(const TCHAR* Match, TAutoConsoleVariable<bool>& CVar)
+	{
+		FString ValueMatch(Match);
+		ValueMatch.Append(TEXT("="));
+		FString Value;
+		if (FParse::Value(FCommandLine::Get(), *ValueMatch, Value))
+		{
+			if (Value.Equals(FString(TEXT("true")), ESearchCase::IgnoreCase))
+			{
+				CVar->Set(true, ECVF_SetByCommandline);
+			}
+			else if (Value.Equals(FString(TEXT("false")), ESearchCase::IgnoreCase))
+			{
+				CVar->Set(false, ECVF_SetByCommandline);
+			}
+		}
+		else if (FParse::Param(FCommandLine::Get(), Match))
+		{
+			CVar->Set(true, ECVF_SetByCommandline);
+		}
+	}
+	
+	TAutoConsoleVariable<bool> CVarEditorPixelStreamingStartOnLaunch(
 		TEXT("PixelStreaming.Editor.StartOnLaunch"),
 		false,
 		TEXT("Start streaming the Editor as soon as it launches. Default: false"),
@@ -25,7 +65,6 @@ namespace UE::EditorPixelStreaming::Settings
 
     void InitialiseSettings()
     {
-		using namespace UE::PixelStreaming;
         // Options parse (if these exist they are set to true)
 		CommandLineParseOption(TEXT("EditorPixelStreamingStartOnLaunch"), CVarEditorPixelStreamingStartOnLaunch);
 		CommandLineParseOption(TEXT("EditorPixelStreamingUseRemoteSignallingServer"), CVarEditorPixelStreamingUseRemoteSignallingServer);
