Fix: Prevent race conditions in Editor compositing by constraining GraphBuilder lifetime

diff --git a/Engine/Plugins/Media/PixelStreaming/Source/PixelStreamingEditor/Private/PixelStreamingVideoInputBackBufferComposited.cpp b/Engine/Plugins/Media/PixelStreaming/Source/PixelStreamingEditor/Private/PixelStreamingVideoInputBackBufferComposited.cpp
--- a/Engine/Plugins/Media/PixelStreaming/Source/PixelStreamingEditor/Private/PixelStreamingVideoInputBackBufferComposited.cpp
+++ b/Engine/Plugins/Media/PixelStreaming/Source/PixelStreamingEditor/Private/PixelStreamingVideoInputBackBufferComposited.cpp
@@ -176,7 +176,6 @@ void FPixelStreamingVideoInputBackBufferComposited::OnBackBufferReady(SWindow& S
 void FPixelStreamingVideoInputBackBufferComposited::CompositeWindows()
 {
 	FRHICommandListImmediate& RHICmdList = FRHICommandListExecutor::GetImmediateCommandList();
-	FRDGBuilder GraphBuilder(RHICmdList);
 
 	// Process all of the windows we will need to render. This processing step finds the extents of the
 	// composited texture as well as the top-left point
@@ -192,10 +191,6 @@ void FPixelStreamingVideoInputBackBufferComposited::CompositeWindows()
 		BottomRight = VectorMax(BottomRight, (WindowPosition + TextureExtent));
 	}
 
-	// Clamp the texture dimensions to ensure no RHI crashes
-	// Create an RDG texture that is the size of our extent for use as the composited frame
-	FRDGTextureRef CompositedTexture = GraphBuilder.CreateTexture(FRDGTextureDesc::Create2D(VectorMin(BottomRight - TopLeft, FIntPoint(16384, 16384)), EPixelFormat::PF_B8G8R8A8, FClearValueBinding::None, ETextureCreateFlags::Shared | ETextureCreateFlags::RenderTargetable), TEXT("VideoInputBackBufferCompositedCompositedTexture"));
-
 	// Shader globals used in the conversion pass
 	FGlobalShaderMap* GlobalShaderMap = GetGlobalShaderMap(GMaxRHIFeatureLevel);
 	TShaderMapRef<FScreenPassVS> VertexShader(GlobalShaderMap);
@@ -205,66 +200,75 @@ void FPixelStreamingVideoInputBackBufferComposited::CompositeWindows()
 	FModifyAlphaSwizzleRgbaPS::FPermutationDomain PermutationVector;
 	PermutationVector.Set<FModifyAlphaSwizzleRgbaPS::FConversionOp>(ConversionOperation);
 
-	for (FTexturedWindow& CurrentWindow : TopLevelWindows)
-	{
-		FIntPoint WindowPosition = FIntPoint(CurrentWindow.GetPositionInScreen().X, CurrentWindow.GetPositionInScreen().Y) - TopLeft;
-
-		FRHITexture* CurrentTexture = CurrentWindow.GetTexture()->GetRHI();
+	FTextureRHIRef OutTexture = nullptr;
+	{	// FRDGBuilder uses a global allocator which can cause race conditions
+		// To prevent issues its lifetime needs to end as soon as it has executed
+		FRDGBuilder GraphBuilder(RHICmdList);
 
-		FRDGTextureRef InputTexture = GraphBuilder.RegisterExternalTexture(CreateRenderTarget(CurrentTexture, TEXT("VideoInputBackBufferCompositedStaging")));
-		// There is only ever one tooltip and as such UE keeps the same texture for each and just rerenders the
-		// content this can lead to small tooltips having a large texture from a previously displayed long tooltip
-		// so we use the tooltips window size which is guaranteed to be correct
-		FIntPoint Extent = VectorMin(CurrentTexture->GetDesc().Extent, CurrentWindow.GetSizeInScreen().IntPoint());
+		// Clamp the texture dimensions to ensure no RHI crashes
+		// Create an RDG texture that is the size of our extent for use as the composited frame
+		FRDGTextureRef CompositedTexture = GraphBuilder.CreateTexture(FRDGTextureDesc::Create2D(VectorMin(BottomRight - TopLeft, FIntPoint(16384, 16384)), EPixelFormat::PF_B8G8R8A8, FClearValueBinding::None, ETextureCreateFlags::Shared | ETextureCreateFlags::RenderTargetable), TEXT("VideoInputBackBufferCompositedCompositedTexture"));
 
-		// Ensure we have a valid extent (texture or window > 0,0)
-		if (Extent.X == 0 || Extent.Y == 0)
+		for (FTexturedWindow& CurrentWindow : TopLevelWindows)
 		{
-			continue;
+			FIntPoint WindowPosition = FIntPoint(CurrentWindow.GetPositionInScreen().X, CurrentWindow.GetPositionInScreen().Y) - TopLeft;
+
+			FRHITexture* CurrentTexture = CurrentWindow.GetTexture()->GetRHI();
+
+			FRDGTextureRef InputTexture = GraphBuilder.RegisterExternalTexture(CreateRenderTarget(CurrentTexture, TEXT("VideoInputBackBufferCompositedStaging")));
+			// There is only ever one tooltip and as such UE keeps the same texture for each and just rerenders the
+			// content this can lead to small tooltips having a large texture from a previously displayed long tooltip
+			// so we use the tooltips window size which is guaranteed to be correct
+			FIntPoint Extent = VectorMin(CurrentTexture->GetDesc().Extent, CurrentWindow.GetSizeInScreen().IntPoint());
+
+			// Ensure we have a valid extent (texture or window > 0,0)
+			if (Extent.X == 0 || Extent.Y == 0)
+			{
+				continue;
+			}
+
+			// Configure our viewports appropriately
+			FScreenPassTextureViewport InputViewport(InputTexture, FIntRect(FIntPoint(0, 0), Extent));
+			FScreenPassTextureViewport OutputViewport(CompositedTexture, FIntRect(WindowPosition, WindowPosition + Extent));
+
+			// Rectangle area to use from the source texture
+			const FIntRect ViewRect(FIntPoint(0, 0), Extent);
+
+			// Dummy ViewFamily/ViewInfo created to use built in Draw Screen/Texture Pass
+			FSceneViewFamilyContext ViewFamily(FSceneViewFamily::ConstructionValues(nullptr, nullptr, FEngineShowFlags(ESFIM_Game))
+											.SetTime(FGameTime())
+											.SetGammaCorrection(1.0f));
+			FSceneViewInitOptions ViewInitOptions;
+			ViewInitOptions.ViewFamily = &ViewFamily;
+			ViewInitOptions.SetViewRectangle(ViewRect);
+			ViewInitOptions.ViewOrigin = FVector::ZeroVector;
+			ViewInitOptions.ViewRotationMatrix = FMatrix::Identity;
+			ViewInitOptions.ProjectionMatrix = FMatrix::Identity;
+
+			GetRendererModule().CreateAndInitSingleView(RHICmdList, &ViewFamily, &ViewInitOptions);
+			const FSceneView& View = *ViewFamily.Views[0];
+
+			TShaderMapRef<FModifyAlphaSwizzleRgbaPS> PixelShader(GlobalShaderMap, PermutationVector);
+			FModifyAlphaSwizzleRgbaPS::FParameters* PixelShaderParameters = PixelShader->AllocateAndSetParameters(GraphBuilder, InputTexture, CompositedTexture);
+			// Add screen pass to convert whatever format the editor produces to BGRA8
+			AddDrawScreenPass(GraphBuilder, RDG_EVENT_NAME("VideoInputBackBufferCompositedSwizzle"), View, OutputViewport, InputViewport, VertexShader, PixelShader, PixelShaderParameters);
 		}
 
-		// Configure our viewports appropriately
-		FScreenPassTextureViewport InputViewport(InputTexture, FIntRect(FIntPoint(0, 0), Extent));
-		FScreenPassTextureViewport OutputViewport(CompositedTexture, FIntRect(WindowPosition, WindowPosition + Extent));
-
-		// Rectangle area to use from the source texture
-		const FIntRect ViewRect(FIntPoint(0, 0), Extent);
-
-		// Dummy ViewFamily/ViewInfo created to use built in Draw Screen/Texture Pass
-		FSceneViewFamilyContext ViewFamily(FSceneViewFamily::ConstructionValues(nullptr, nullptr, FEngineShowFlags(ESFIM_Game))
-										.SetTime(FGameTime())
-										.SetGammaCorrection(1.0f));
-		FSceneViewInitOptions ViewInitOptions;
-		ViewInitOptions.ViewFamily = &ViewFamily;
-		ViewInitOptions.SetViewRectangle(ViewRect);
-		ViewInitOptions.ViewOrigin = FVector::ZeroVector;
-		ViewInitOptions.ViewRotationMatrix = FMatrix::Identity;
-		ViewInitOptions.ProjectionMatrix = FMatrix::Identity;
-
-		GetRendererModule().CreateAndInitSingleView(RHICmdList, &ViewFamily, &ViewInitOptions);
-		const FSceneView& View = *ViewFamily.Views[0];
-
-		TShaderMapRef<FModifyAlphaSwizzleRgbaPS> PixelShader(GlobalShaderMap, PermutationVector);
-		FModifyAlphaSwizzleRgbaPS::FParameters* PixelShaderParameters = PixelShader->AllocateAndSetParameters(GraphBuilder, InputTexture, CompositedTexture);
-		// Add screen pass to convert whatever format the editor produces to BGRA8
-		AddDrawScreenPass(GraphBuilder, RDG_EVENT_NAME("VideoInputBackBufferCompositedSwizzle"), View, OutputViewport, InputViewport, VertexShader, PixelShader, PixelShaderParameters);
+		// Final pass to extract the composited frames underlying RHI resource for passing to the rest of the pixel streaming pipeline
+		FExtractCompositedTextureParameters* PassParameters = GraphBuilder.AllocParameters<FExtractCompositedTextureParameters>();
+		PassParameters->Input = CompositedTexture;
+		GraphBuilder.AddPass(
+			RDG_EVENT_NAME("PushCompositedFrame"),
+			PassParameters,
+			ERDGPassFlags::Readback,
+			[CompositedTexture, &OutTexture, this](FRHICommandList& RHICmdList) {
+				// Our composition is complete out the underlying rhi resource
+				OutTexture = CompositedTexture->GetRHI();
+			});
+
+		GraphBuilder.Execute();
 	}
 
-	// Final pass to extract the composited frames underlying RHI resource for passing to the rest of the pixel streaming pipeline
-	FTextureRHIRef OutTexture = nullptr;
-	FExtractCompositedTextureParameters* PassParameters = GraphBuilder.AllocParameters<FExtractCompositedTextureParameters>();
-	PassParameters->Input = CompositedTexture;
-	GraphBuilder.AddPass(
-		RDG_EVENT_NAME("PushCompositedFrame"),
-		PassParameters,
-		ERDGPassFlags::Readback,
-		[CompositedTexture, &OutTexture, this](FRHICommandList& RHICmdList) {
-			// Our composition is complete out the underlying rhi resource
-			OutTexture = CompositedTexture->GetRHI();
-		});
-
-	GraphBuilder.Execute();
-
 	OnFrame(FPixelCaptureInputFrameRHI(OutTexture));
 
 	// Update any subscribed streamers to let them know our composited frame size and position. This way it can correctly scale input from the browser
