Fix level viewport streaming due to internal viewport rendering changes.
Also added support for continuing stream when PIE is activated mid stream.

diff --git a/Engine/Plugins/Media/PixelStreaming/Source/PixelStreamingEditor/Private/PixelStreamingVideoInputViewport.cpp b/Engine/Plugins/Media/PixelStreaming/Source/PixelStreamingEditor/Private/PixelStreamingVideoInputViewport.cpp
--- a/Engine/Plugins/Media/PixelStreaming/Source/PixelStreamingEditor/Private/PixelStreamingVideoInputViewport.cpp
+++ b/Engine/Plugins/Media/PixelStreaming/Source/PixelStreamingEditor/Private/PixelStreamingVideoInputViewport.cpp
@@ -1,22 +1,26 @@
 // Copyright Epic Games, Inc. All Rights Reserved.
 
 #include "PixelStreamingVideoInputViewport.h"
-
 #include "Async/Async.h"
 #include "Settings.h"
 #include "Utils.h"
 #include "PixelCaptureInputFrameRHI.h"
 #include "PixelCaptureBufferFormat.h"
 #include "PixelCaptureCapturerRHI.h"
-#include "PixelCaptureCapturerRHIToI420CPU.h"
-#include "PixelCaptureCapturerRHIToI420Compute.h"
 #include "UnrealClient.h"
 #include "Engine/GameViewportClient.h"
 #include "RenderingThread.h"
 #include "Framework/Application/SlateApplication.h"
-#include "Slate/SceneViewport.h"
+#include "LevelEditor.h"
+#include "LevelEditorViewport.h"
+#include "SLevelViewport.h"
 #include "IPixelStreamingModule.h"
 
+bool IsPIESessionRunning()
+{
+	return GEditor && GEditor->PlayWorld && GEditor->PlayWorld->WorldType == EWorldType::PIE;
+}
+
 TSharedPtr<FPixelStreamingVideoInputViewport> FPixelStreamingVideoInputViewport::Create(TSharedPtr<IPixelStreamingStreamer> InAssociatedStreamer)
 {
 	TSharedPtr<FPixelStreamingVideoInputViewport> NewInput = TSharedPtr<FPixelStreamingVideoInputViewport>(new FPixelStreamingVideoInputViewport());
@@ -27,7 +31,19 @@ TSharedPtr<FPixelStreamingVideoInputViewport> FPixelStreamingVideoInputViewport:
 	AsyncTask(ENamedThreads::GameThread, [WeakInput]() {
 		if (TSharedPtr<FPixelStreamingVideoInputViewport> Input = WeakInput.Pin())
 		{
-			Input->DelegateHandle = UGameViewportClient::OnViewportRendered().AddSP(Input.ToSharedRef(), &FPixelStreamingVideoInputViewport::OnViewportRendered);
+			// Bind to GameViewport delegate for PIE sessions
+			Input->PIEDelegateHandle = UGameViewportClient::OnViewportRendered().AddSP(Input.ToSharedRef(), 
+			&FPixelStreamingVideoInputViewport::OnPIEViewportRendered);
+
+			// For non-PIE cases (just using level editor) then bind to slate window rendered
+			FSlateRenderer* Renderer = FSlateApplication::Get().GetRenderer();
+			if(!Renderer)
+			{
+				return;
+			}
+
+			Input->DelegateHandle = Renderer->OnSlateWindowRendered().AddSP(Input.ToSharedRef(), 
+			&FPixelStreamingVideoInputViewport::OnWindowRendered);
 		}
 	});
 
@@ -38,53 +54,87 @@ FPixelStreamingVideoInputViewport::~FPixelStreamingVideoInputViewport()
 {
 	if (!IsEngineExitRequested())
 	{
-		AsyncTask(ENamedThreads::GameThread, [HandleCopy = DelegateHandle]() {
-			UGameViewportClient::OnViewportRendered().Remove(HandleCopy);
+		AsyncTask(ENamedThreads::GameThread, [HandleCopy = DelegateHandle, PIEHandleCopy = PIEDelegateHandle]() {
+			UGameViewportClient::OnViewportRendered().Remove(PIEHandleCopy);
+			FSlateRenderer* Renderer = FSlateApplication::Get().GetRenderer();
+			if(Renderer)
+			{
+				Renderer->OnSlateWindowRendered().Remove(HandleCopy);
+			}
 		});
 	}
 }
 
-bool FPixelStreamingVideoInputViewport::FilterViewport(const FViewport* InViewport)
+bool FPixelStreamingVideoInputViewport::FilterWindow(SWindow& InWindow)
 {
 	TSharedPtr<IPixelStreamingStreamer> Streamer = AssociatedStreamer.Pin();
-	if (!Streamer.IsValid() || !Streamer->IsStreaming())
+	if (!Streamer || !Streamer->IsStreaming())
 	{
 		return false;
 	}
 
-	TSharedPtr<SViewport> TargetScene = Streamer->GetTargetViewport().Pin();
-	if (!TargetScene.IsValid())
+	TSharedPtr<SWindow> TargetWindow = Streamer->GetTargetWindow().Pin();
+
+	if (!TargetWindow)
 	{
 		return false;
 	}
 
-	if (InViewport == nullptr)
+	// Check if the window we were passed is our target window.
+	if (TargetWindow.Get() != &InWindow)
 	{
 		return false;
 	}
 
-	if (InViewport->GetViewportType() != TargetViewportType)
+	return true;
+}
+
+void FPixelStreamingVideoInputViewport::OnWindowRendered(SWindow& InWindow, void* InResource)
+{
+	// If there is a PIE session running then don't render level viewport
+	if (IsPIESessionRunning())
+	{
+		return;
+	}
+
+	if (!FilterWindow(InWindow))
 	{
-		return false;
+		return;
 	}
 
-	// Bit dirty to do a static cast here, but we check viewport type just above so it is somewhat "safe".
-	TSharedPtr<SViewport> InScene = StaticCast<const FSceneViewport*>(InViewport)->GetViewportWidget().Pin();
+	TSharedPtr<IPixelStreamingStreamer> Streamer = AssociatedStreamer.Pin();
+	if (!Streamer || !Streamer->IsStreaming())
+	{
+		return;
+	}
 
-	// If the viewport we were passed is not our target viewport we are not interested in getting its texture.
-	if (TargetScene != InScene)
+	FLevelEditorModule& LevelEditorModule = FModuleManager::GetModuleChecked<FLevelEditorModule>("LevelEditor");
+	TSharedPtr<SLevelViewport> ActiveLevelViewport = LevelEditorModule.GetFirstActiveLevelViewport();
+	if (!ActiveLevelViewport.IsValid())
 	{
-		return false;
+		return;
 	}
 
+	FLevelEditorViewportClient& LevelViewportClient = ActiveLevelViewport->GetLevelViewportClient();
+	FViewport* TargetViewport = LevelViewportClient.Viewport;
 
+	SubmitViewport(TargetViewport);
+}
 
-	return true;
+void FPixelStreamingVideoInputViewport::OnPIEViewportRendered(FViewport* InViewport)
+{
+	// If gameclient viewport is rendered and it is not PIE we are not interested
+	if (!IsPIESessionRunning())
+	{
+		return;
+	}
+
+	SubmitViewport(InViewport);
 }
 
-void FPixelStreamingVideoInputViewport::OnViewportRendered(FViewport* InViewport)
+void FPixelStreamingVideoInputViewport::SubmitViewport(FViewport* InViewport)
 {
-	if (!FilterViewport(InViewport))
+	if(!InViewport)
 	{
 		return;
 	}
@@ -96,7 +146,7 @@ void FPixelStreamingVideoInputViewport::OnViewportRendered(FViewport* InViewport
 	}
 
 	ENQUEUE_RENDER_COMMAND(StreamViewportTextureCommand)
-	([&, FrameBuffer](FRHICommandList& RHICmdList) {
+	([this, FrameBuffer](FRHICommandList& RHICmdList) {
 		OnFrame(FPixelCaptureInputFrameRHI(FrameBuffer));
 	});
 }
diff --git a/Engine/Plugins/Media/PixelStreaming/Source/PixelStreamingEditor/Public/PixelStreamingVideoInputViewport.h b/Engine/Plugins/Media/PixelStreaming/Source/PixelStreamingEditor/Public/PixelStreamingVideoInputViewport.h
--- a/Engine/Plugins/Media/PixelStreaming/Source/PixelStreamingEditor/Public/PixelStreamingVideoInputViewport.h
+++ b/Engine/Plugins/Media/PixelStreaming/Source/PixelStreamingEditor/Public/PixelStreamingVideoInputViewport.h
@@ -21,10 +21,13 @@ public:
 private:
 	FPixelStreamingVideoInputViewport() = default;
 
-	bool FilterViewport(const FViewport* InViewport);
-	void OnViewportRendered(FViewport* Viewport);
+	bool FilterWindow(SWindow& Window);
+	void OnWindowRendered(SWindow& Window, void* Resource);
+	void OnPIEViewportRendered(FViewport* InViewport);
+	void SubmitViewport(FViewport* InViewport);
 
 	FDelegateHandle DelegateHandle;
+	FDelegateHandle PIEDelegateHandle;
 
 	FName TargetViewportType = FName(FString(TEXT("SceneViewport")));
 
