Expose new PixelStreaming/WebRTC settings to mitigate frame drops in LAN scenarios:

1) Disable frame dropping in WebRTC using the WebRTC-FrameDropper/Disabled/ field trial.

2) Control the video pacer inside WebRTC to send much larger bitrates without dropping frames in the congestion window.

The combination of these two settings allows sending large fixed bitrates without frames being dropped on the UE side - this is useful for LAN streaming scenarios.

diff --git a/Engine/Plugins/Media/PixelStreaming/Source/PixelStreaming/Private/PixelStreamingPeerConnection.cpp b/Engine/Plugins/Media/PixelStreaming/Source/PixelStreaming/Private/PixelStreamingPeerConnection.cpp
--- a/Engine/Plugins/Media/PixelStreaming/Source/PixelStreaming/Private/PixelStreamingPeerConnection.cpp
+++ b/Engine/Plugins/Media/PixelStreaming/Source/PixelStreaming/Private/PixelStreamingPeerConnection.cpp
@@ -868,6 +868,52 @@ void FPixelStreamingPeerConnection::RemoveAudioInput(TSharedPtr<IPixelStreamingA
 	AudioMixer->DisconnectInput(StaticCastSharedPtr<FAudioInput>(AudioInput));
 }
 
+
+void InitializeFieldTrials()
+{
+	FString FieldTrials = Settings::CVarPixelStreamingWebRTCFieldTrials.GetValueOnAnyThread();
+
+	// Set the WebRTC-FrameDropper/Disabled/ if the CVar is set
+	if(Settings::CVarPixelStreamingWebRTCDisableFrameDropper.GetValueOnAnyThread())
+	{
+		FieldTrials += TEXT("WebRTC-FrameDropper/Disabled/");
+	}
+
+	// Parse "WebRTC-Video-Pacing/" field trial
+	{
+		float OutPacingFactor = -1.0f;
+		float OutPacingMaxDelayMs = -1.0f;
+		bool bVideoPacingFieldTrial = Settings::GetVideoPacing(OutPacingFactor, OutPacingMaxDelayMs);
+		if(bVideoPacingFieldTrial)
+		{
+			FString VideoPacingFieldTrialStr = TEXT("WebRTC-Video-Pacing/");
+			bool bHasPacingFactor = OutPacingFactor >= 0.0f;
+			if(bHasPacingFactor)
+			{
+				VideoPacingFieldTrialStr += FString::Printf(TEXT("factor:%.1f"), OutPacingFactor);
+			}
+			bool bHasMaxDelay = OutPacingMaxDelayMs >= 0.0f;
+			if(bHasMaxDelay)
+			{
+				VideoPacingFieldTrialStr += bHasPacingFactor ? TEXT(",") : TEXT("");
+				VideoPacingFieldTrialStr += FString::Printf(TEXT("max_delay:%.0f"), OutPacingMaxDelayMs);
+			}
+			VideoPacingFieldTrialStr += TEXT("/");
+			FieldTrials += VideoPacingFieldTrialStr;
+		}
+	}
+
+	if (!FieldTrials.IsEmpty()) {
+		//Pass the field trials string to WebRTC. String must never be destroyed.
+		TStringConversion<TStringConvert<TCHAR, ANSICHAR>> Str = StringCast<ANSICHAR>(*FieldTrials);
+		int length = Str.Length() + 1;
+		char* WRTCFieldTrials = (char*)FMemory::SystemMalloc(length);
+		FMemory::Memcpy(WRTCFieldTrials, Str.Get(), Str.Length());
+		WRTCFieldTrials[length - 1] = '\0';
+		webrtc::field_trial::InitFieldTrialsFromString(WRTCFieldTrials);
+	}
+}
+
 void FPixelStreamingPeerConnection::CreatePeerConnectionFactory()
 {
 	using namespace UE::PixelStreaming;
@@ -925,17 +971,9 @@ void FPixelStreamingPeerConnection::CreatePeerConnectionFactory()
 
 	std::unique_ptr<FVideoEncoderFactoryLayered> VideoEncoderFactory = std::make_unique<FVideoEncoderFactoryLayered>();
 	GVideoEncoderFactory = VideoEncoderFactory.get();
-	 
-	FString FieldTrials = Settings::CVarPixelStreamingWebRTCFieldTrials.GetValueOnAnyThread();
-	if (!FieldTrials.IsEmpty()) {
-		//Pass the field trials string to WebRTC. String must never be destroyed.
-		TStringConversion<TStringConvert<TCHAR, ANSICHAR>> Str = StringCast<ANSICHAR>(*FieldTrials);
-		int length = Str.Length() + 1;
-		char* WRTCFieldTrials = (char*)FMemory::SystemMalloc(length);
-		FMemory::Memcpy(WRTCFieldTrials, Str.Get(), Str.Length());
-		WRTCFieldTrials[length - 1] = '\0';
-		webrtc::field_trial::InitFieldTrialsFromString(WRTCFieldTrials);
-	}
+
+	// Set up the field trials, read things from CVars etc
+	InitializeFieldTrials();
 
 	PeerConnectionFactory = webrtc::CreatePeerConnectionFactory(
 		nullptr,													   // network_thread
diff --git a/Engine/Plugins/Media/PixelStreaming/Source/PixelStreaming/Private/Settings.cpp b/Engine/Plugins/Media/PixelStreaming/Source/PixelStreaming/Private/Settings.cpp
--- a/Engine/Plugins/Media/PixelStreaming/Source/PixelStreaming/Private/Settings.cpp
+++ b/Engine/Plugins/Media/PixelStreaming/Source/PixelStreaming/Private/Settings.cpp
@@ -197,6 +197,24 @@ namespace UE::PixelStreaming::Settings
 		TEXT("Sets the amount of gain to apply to audio. Default: 1.0"),
 		ECVF_Default);
 
+	TAutoConsoleVariable<bool> CVarPixelStreamingWebRTCDisableFrameDropper(
+		TEXT("PixelStreaming.WebRTC.DisableFrameDropper"),
+		false,
+		TEXT("Disables the WebRTC internal frame dropper using the field trial WebRTC-FrameDropper/Disabled/"),
+		ECVF_Default);
+
+	TAutoConsoleVariable<float> CVarPixelStreamingWebRTCVideoPacingMaxDelay(
+		TEXT("PixelStreaming.WebRTC.VideoPacing.MaxDelay"),
+		-1.0f,
+		TEXT("Enables the WebRTC-Video-Pacing field trial and sets the max delay (ms) parameter. Default: -1.0f (values below zero are discarded.)"),
+		ECVF_Default);
+
+	TAutoConsoleVariable<float> CVarPixelStreamingWebRTCVideoPacingFactor(
+		TEXT("PixelStreaming.WebRTC.VideoPacing.Factor"),
+		-1.0f,
+		TEXT("Enables the WebRTC-Video-Pacing field trial and sets the video pacing factor parameter. Larger values are more lenient on larger bitrates. Default: -1.0f (values below zero are discarded.)"),
+		ECVF_Default);
+
 	TAutoConsoleVariable<FString> CVarPixelStreamingWebRTCFieldTrials(
 		TEXT("PixelStreaming.WebRTC.FieldTrials"),
 		TEXT(""),
@@ -548,6 +566,14 @@ namespace UE::PixelStreaming::Settings
 		return FParse::Value(FCommandLine::Get(), TEXT("PixelStreamingPort="), OutSignallingServerPort);
 	}
 
+	bool GetVideoPacing(float& OutPacingFactor, float& OutMaxDelayMs)
+	{
+		OutPacingFactor = CVarPixelStreamingWebRTCVideoPacingFactor.GetValueOnAnyThread();
+		OutMaxDelayMs = CVarPixelStreamingWebRTCVideoPacingMaxDelay.GetValueOnAnyThread();
+		// Note: Use of OR operator here, field trial is enable if either of these is non-zero and positive.
+		return OutPacingFactor >= 0.0 || OutMaxDelayMs >= 0.0;
+	}
+
 	EInputControllerMode GetInputControllerMode()
 	{
 		// Convert the current value to all lowercase and remove any whitespace.
@@ -646,10 +672,12 @@ namespace UE::PixelStreaming::Settings
 		CommandLineParseValue(TEXT("PixelStreamingWebRTCMaxBitrate="), CVarPixelStreamingWebRTCMaxBitrate);
 		CommandLineParseValue(TEXT("PixelStreamingWebRTCLowQpThreshold="), CVarPixelStreamingWebRTCLowQpThreshold);
 		CommandLineParseValue(TEXT("PixelStreamingWebRTCHighQpThreshold="), CVarPixelStreamingWebRTCHighQpThreshold);
+		CommandLineParseValue(TEXT("PixelStreamingWebRTCVideoPacingFactor="), CVarPixelStreamingWebRTCVideoPacingFactor);
+		CommandLineParseValue(TEXT("PixelStreamingWebRTCVideoPacingMaxDelay="), CVarPixelStreamingWebRTCVideoPacingMaxDelay);
+		CommandLineParseValue(TEXT("PixelStreamingWebRTCFieldTrials="), CVarPixelStreamingWebRTCFieldTrials);
 		CommandLineParseValue(TEXT("PixelStreamingFreezeFrameQuality"), CVarPixelStreamingFreezeFrameQuality);
 		CommandLineParseValue(TEXT("PixelStreamingInputController="), CVarPixelStreamingInputController);
 		CommandLineParseValue(TEXT("PixelStreamingSignalingReconnectInterval="), CVarPixelStreamingSignalingReconnectInterval);
-		CommandLineParseValue(TEXT("PixelStreamingWebRTCFieldTrials="), CVarPixelStreamingWebRTCFieldTrials);
 
 		// Options parse (if these exist they are set to true)
 		CommandLineParseOption(TEXT("PixelStreamingOnScreenStats"), CVarPixelStreamingOnScreenStats);
@@ -660,6 +688,7 @@ namespace UE::PixelStreaming::Settings
 		CommandLineParseOption(TEXT("PixelStreamingWebRTCDisableReceiveAudio"), CVarPixelStreamingWebRTCDisableReceiveAudio);
 		CommandLineParseOption(TEXT("PixelStreamingWebRTCDisableTransmitAudio"), CVarPixelStreamingWebRTCDisableTransmitAudio);
 		CommandLineParseOption(TEXT("PixelStreamingWebRTCDisableAudioSync"), CVarPixelStreamingWebRTCDisableAudioSync);
+		CommandLineParseOption(TEXT("PixelStreamingWebRTCDisableFrameDropper"), CVarPixelStreamingWebRTCDisableFrameDropper);
 		CommandLineParseOption(TEXT("PixelStreamingSendPlayerIdAsInteger"), CVarSendPlayerIdAsInteger);
 		CommandLineParseOption(TEXT("PixelStreamingWebRTCUseLegacyAudioDevice"), CVarPixelStreamingWebRTCUseLegacyAudioDevice);
 		CommandLineParseOption(TEXT("PixelStreamingDisableLatencyTester"), CVarPixelStreamingDisableLatencyTester);
diff --git a/Engine/Plugins/Media/PixelStreaming/Source/PixelStreaming/Private/Settings.h b/Engine/Plugins/Media/PixelStreaming/Source/PixelStreaming/Private/Settings.h
--- a/Engine/Plugins/Media/PixelStreaming/Source/PixelStreaming/Private/Settings.h
+++ b/Engine/Plugins/Media/PixelStreaming/Source/PixelStreaming/Private/Settings.h
@@ -51,6 +51,9 @@ namespace UE::PixelStreaming::Settings
 	extern TAutoConsoleVariable<float> CVarPixelStreamingWebRTCAudioGain;
 	extern TAutoConsoleVariable<bool> CVarPixelStreamingWebRTCNegotiateCodecs;
 	extern TAutoConsoleVariable<FString> CVarPixelStreamingWebRTCFieldTrials;
+	extern TAutoConsoleVariable<bool> CVarPixelStreamingWebRTCDisableFrameDropper;
+	extern TAutoConsoleVariable<float> CVarPixelStreamingWebRTCVideoPacingFactor;
+	extern TAutoConsoleVariable<float> CVarPixelStreamingWebRTCVideoPacingMaxDelay;
 	// End WebRTC CVars
 
 	// Begin Pixel Streaming Plugin CVars
@@ -96,6 +99,13 @@ namespace UE::PixelStreaming::Settings
 	bool GetSignallingServerUrl(FString& OutSignallingServerURL);
 	bool GetSignallingServerIP(FString& OutSignallingServerIP);
 	bool GetSignallingServerPort(uint16& OutSignallingServerPort);
+
+	/**
+	* @param OutPacingFactor The pacing factor to use for the field trial.
+	* @param OutMaxDelayMs The max delay to use for the field trial.
+	* @return True if WebRTC-Video-Pacing field trial is enabled.
+	*/
+	bool GetVideoPacing(float& OutPacingFactor, float& OutMaxDelayMs);
 	// End utility functions etc.
 
 	extern FPixelStreamingSimulcastParameters SimulcastParameters;
