Fix peers not getting any frames when stream sharing without quality controller.

diff --git a/Engine/Plugins/Media/PixelStreaming/Source/PixelStreaming/Private/VideoSource.cpp b/Engine/Plugins/Media/PixelStreaming/Source/PixelStreaming/Private/VideoSource.cpp
--- a/Engine/Plugins/Media/PixelStreaming/Source/PixelStreaming/Private/VideoSource.cpp
+++ b/Engine/Plugins/Media/PixelStreaming/Source/PixelStreaming/Private/VideoSource.cpp
@@ -16,9 +16,14 @@ namespace UE::PixelStreaming
 
 	void FVideoSource::MaybePushFrame()
 	{
-		if (VideoInput->IsReady() && ShouldGenerateFramesCheck())
+		if (VideoInput->IsReady())
 		{
-			PushFrame();
+			CurrentState = webrtc::MediaSourceInterface::SourceState::kLive;
+
+			if(ShouldGenerateFramesCheck())
+			{
+				PushFrame();
+			}
 		}
 	}
 
@@ -27,8 +32,6 @@ namespace UE::PixelStreaming
 	    TRACE_CPUPROFILER_EVENT_SCOPE_ON_CHANNEL_STR("PixelStreaming Push Video Frame", PixelStreamingChannel);
 		static int32 FrameId = 1;
 
-		CurrentState = webrtc::MediaSourceInterface::SourceState::kLive;
-
 		rtc::scoped_refptr<webrtc::VideoFrameBuffer> FrameBuffer = VideoInput->GetFrameBuffer();
 		check(FrameBuffer->width() != 0 && FrameBuffer->height() != 0);
 		webrtc::VideoFrame Frame = webrtc::VideoFrame::Builder()
