Moving frame timings to rtc::TimeMillis for better consistency with webrtc.

diff --git a/Engine/Plugins/Media/PixelCapture/Source/PixelCapture/Private/PixelCaptureCapturer.cpp b/Engine/Plugins/Media/PixelCapture/Source/PixelCapture/Private/PixelCaptureCapturer.cpp
--- a/Engine/Plugins/Media/PixelCapture/Source/PixelCapture/Private/PixelCaptureCapturer.cpp
+++ b/Engine/Plugins/Media/PixelCapture/Source/PixelCapture/Private/PixelCaptureCapturer.cpp
@@ -44,7 +44,7 @@ void FPixelCaptureCapturer::Capture(const IPixelCaptureInputFrame& InputFrame)
 	}
 
 	InitMetadata(InputFrame.Metadata.Copy());
-	StartTime = FPlatformTime::Cycles64();
+	StartTime = rtc::TimeMillis();
 	CPUStartTime = 0;
 	GPUEnqueueTime = 0;
 	GPUStartTime = 0;
@@ -68,16 +68,16 @@ void FPixelCaptureCapturer::MarkCPUWorkStart()
 	{
 		MarkCPUWorkEnd();
 	}
-	CPUStartTime = FPlatformTime::Cycles64();
+	CPUStartTime = rtc::TimeMillis();
 }
 
 void FPixelCaptureCapturer::MarkCPUWorkEnd()
 {
 	check(CurrentOutputBuffer != nullptr);
-	CurrentOutputBuffer->Metadata.CaptureProcessCPUTime += FPlatformTime::Cycles64() - CPUStartTime;
+	CurrentOutputBuffer->Metadata.CaptureProcessCPUTime += rtc::TimeMillis() - CPUStartTime;
 	CPUStartTime = 0;
 
-	GPUEnqueueTime = FPlatformTime::Cycles64();
+	GPUEnqueueTime = rtc::TimeMillis();
 }
 
 void FPixelCaptureCapturer::MarkGPUWorkStart()
@@ -86,7 +86,7 @@ void FPixelCaptureCapturer::MarkGPUWorkStart()
 	{
 		MarkGPUWorkEnd();
 	}
-	GPUStartTime = FPlatformTime::Cycles64();
+	GPUStartTime = rtc::TimeMillis();
 	
 	CurrentOutputBuffer->Metadata.CaptureProcessGPUDelay += GPUStartTime - GPUEnqueueTime;
 	GPUEnqueueTime = 0;
@@ -94,7 +94,7 @@ void FPixelCaptureCapturer::MarkGPUWorkStart()
 
 void FPixelCaptureCapturer::MarkGPUWorkEnd()
 {
-	CurrentOutputBuffer->Metadata.CaptureProcessGPUTime += FPlatformTime::Cycles64() - GPUStartTime;
+	CurrentOutputBuffer->Metadata.CaptureProcessGPUTime += rtc::TimeMillis() - GPUStartTime;
 	GPUStartTime = 0;
 }
 
@@ -119,7 +119,7 @@ void FPixelCaptureCapturer::FinalizeMetadata()
 		MarkGPUWorkEnd();
 	}
 
-	CurrentOutputBuffer->Metadata.CaptureTime += FPlatformTime::Cycles64() - StartTime;
+	CurrentOutputBuffer->Metadata.CaptureTime += rtc::TimeMillis() - StartTime;
 }
 
 void FPixelCaptureCapturer::EndProcess()
diff --git a/Engine/Plugins/Media/PixelCapture/Source/PixelCapture/Private/PixelCaptureInputFrameI420.cpp b/Engine/Plugins/Media/PixelCapture/Source/PixelCapture/Private/PixelCaptureInputFrameI420.cpp
--- a/Engine/Plugins/Media/PixelCapture/Source/PixelCapture/Private/PixelCaptureInputFrameI420.cpp
+++ b/Engine/Plugins/Media/PixelCapture/Source/PixelCapture/Private/PixelCaptureInputFrameI420.cpp
@@ -6,7 +6,7 @@
 FPixelCaptureInputFrameI420::FPixelCaptureInputFrameI420(TSharedPtr<FPixelCaptureBufferI420> Buffer)
 	: I420Buffer(Buffer)
 {
-	Metadata.SourceTime = FPlatformTime::Cycles64();
+	Metadata.SourceTime = rtc::TimeMillis();
 }
 
 int32 FPixelCaptureInputFrameI420::GetType() const
diff --git a/Engine/Plugins/Media/PixelCapture/Source/PixelCapture/Private/PixelCaptureInputFrameNV12.cpp b/Engine/Plugins/Media/PixelCapture/Source/PixelCapture/Private/PixelCaptureInputFrameNV12.cpp
--- a/Engine/Plugins/Media/PixelCapture/Source/PixelCapture/Private/PixelCaptureInputFrameNV12.cpp
+++ b/Engine/Plugins/Media/PixelCapture/Source/PixelCapture/Private/PixelCaptureInputFrameNV12.cpp
@@ -6,7 +6,7 @@
 FPixelCaptureInputFrameNV12::FPixelCaptureInputFrameNV12(TSharedPtr<FPixelCaptureBufferNV12> Buffer)
 	: NV12Buffer(Buffer)
 {
-	Metadata.SourceTime = FPlatformTime::Cycles64();
+	Metadata.SourceTime = rtc::TimeMillis();
 }
 
 int32 FPixelCaptureInputFrameNV12::GetType() const
diff --git a/Engine/Plugins/Media/PixelCapture/Source/PixelCapture/Private/PixelCaptureInputFrameRHI.cpp b/Engine/Plugins/Media/PixelCapture/Source/PixelCapture/Private/PixelCaptureInputFrameRHI.cpp
--- a/Engine/Plugins/Media/PixelCapture/Source/PixelCapture/Private/PixelCaptureInputFrameRHI.cpp
+++ b/Engine/Plugins/Media/PixelCapture/Source/PixelCapture/Private/PixelCaptureInputFrameRHI.cpp
@@ -6,7 +6,7 @@
 FPixelCaptureInputFrameRHI::FPixelCaptureInputFrameRHI(FTexture2DRHIRef InFrameTexture)
 	: FrameTexture(InFrameTexture)
 {
-	Metadata.SourceTime = FPlatformTime::Cycles64();
+	Metadata.SourceTime = rtc::TimeMillis();
 }
 
 int32 FPixelCaptureInputFrameRHI::GetType() const
diff --git a/Engine/Plugins/Media/PixelCapture/Source/PixelCapture/Public/IPixelCaptureInputFrame.h b/Engine/Plugins/Media/PixelCapture/Source/PixelCapture/Public/IPixelCaptureInputFrame.h
--- a/Engine/Plugins/Media/PixelCapture/Source/PixelCapture/Public/IPixelCaptureInputFrame.h
+++ b/Engine/Plugins/Media/PixelCapture/Source/PixelCapture/Public/IPixelCaptureInputFrame.h
@@ -4,6 +4,7 @@
 
 #include "PixelCaptureFrameMetadata.h"
 #include "HAL/PlatformTime.h"
+#include "rtc_base/time_utils.h"
 
 /**
  * The base interface that is fed into the capture system. This wrapper is fed into the
@@ -13,7 +14,7 @@
 class PIXELCAPTURE_API IPixelCaptureInputFrame
 {
 public:
-	IPixelCaptureInputFrame() { Metadata.SourceTime = FPlatformTime::Cycles64(); }
+	IPixelCaptureInputFrame() { Metadata.SourceTime = rtc::TimeMillis(); }
 	virtual ~IPixelCaptureInputFrame() = default;
 
 	/**
diff --git a/Engine/Plugins/Media/PixelStreaming/Source/PixelStreaming/Private/Stats.cpp b/Engine/Plugins/Media/PixelStreaming/Source/PixelStreaming/Private/Stats.cpp
--- a/Engine/Plugins/Media/PixelStreaming/Source/PixelStreaming/Private/Stats.cpp
+++ b/Engine/Plugins/Media/PixelStreaming/Source/PixelStreaming/Private/Stats.cpp
@@ -492,11 +492,11 @@ namespace UE::PixelStreaming
 		return DeltaMs;
 	}
 
-	double FStats::AddTimeDeltaStat(uint64 Cycles1, uint64 Cycles2, const FString& Label)
+	double FStats::AddTimeDeltaStat(uint64 Millis1, uint64 Millis2, const FString& Label)
 	{
-		const uint64 MaxCycles = FGenericPlatformMath::Max(Cycles1, Cycles2);
-		const uint64 MinCycles = FGenericPlatformMath::Min(Cycles1, Cycles2);
-		const double DeltaMs = FPlatformTime::ToMilliseconds64(MaxCycles - MinCycles) * ((Cycles1 > Cycles2) ? 1.0 : -1.0);
+		const uint64 MaxMillis = FGenericPlatformMath::Max(Millis1, Millis2);
+		const uint64 MinMillis = FGenericPlatformMath::Min(Millis1, Millis2);
+		const double DeltaMs = (MaxMillis - MinMillis) * ((Millis1 > Millis2) ? 1.0 : -1.0);
 		const FStatData TimeData{ FName(*Label), DeltaMs, 2, true };
 		StoreApplicationStat(TimeData);
 		return DeltaMs;
diff --git a/Engine/Plugins/Media/PixelStreaming/Source/PixelStreaming/Private/VideoEncoderSingleLayerHardware.cpp b/Engine/Plugins/Media/PixelStreaming/Source/PixelStreaming/Private/VideoEncoderSingleLayerHardware.cpp
--- a/Engine/Plugins/Media/PixelStreaming/Source/PixelStreaming/Private/VideoEncoderSingleLayerHardware.cpp
+++ b/Engine/Plugins/Media/PixelStreaming/Source/PixelStreaming/Private/VideoEncoderSingleLayerHardware.cpp
@@ -98,7 +98,7 @@ namespace UE::PixelStreaming
 	{
 		FPixelCaptureFrameMetadata& FrameMetadata = Frame.Metadata;
 		FrameMetadata.UseCount++;
-		FrameMetadata.LastEncodeStartTime = FPlatformTime::Cycles64();
+		FrameMetadata.LastEncodeStartTime = rtc::TimeMillis();
 		if (FrameMetadata.UseCount == 1)
 		{
 			FrameMetadata.FirstEncodeStartTime = FrameMetadata.LastEncodeStartTime;
@@ -108,7 +108,7 @@ namespace UE::PixelStreaming
 	void FVideoEncoderSingleLayerHardware::UpdateFrameMetadataPostEncode(IPixelCaptureOutputFrame& Frame)
 	{
 		FPixelCaptureFrameMetadata& FrameMetadata = Frame.Metadata;
-		FrameMetadata.LastEncodeEndTime = FPlatformTime::Cycles64();
+		FrameMetadata.LastEncodeEndTime = rtc::TimeMillis();
 
 		FStats::Get()->AddFrameTimingStats(FrameMetadata);
 	}
@@ -116,7 +116,7 @@ namespace UE::PixelStreaming
 	void FVideoEncoderSingleLayerHardware::UpdateFrameMetadataPrePacketization(IPixelCaptureOutputFrame& Frame)
 	{
 		FPixelCaptureFrameMetadata& FrameMetadata = Frame.Metadata;
-		FrameMetadata.LastPacketizationStartTime = FPlatformTime::Cycles64();
+		FrameMetadata.LastPacketizationStartTime = rtc::TimeMillis();
 		if (FrameMetadata.UseCount == 1)
 		{
 			FrameMetadata.FirstPacketizationStartTime = FrameMetadata.LastPacketizationStartTime;
@@ -126,7 +126,7 @@ namespace UE::PixelStreaming
 	void FVideoEncoderSingleLayerHardware::UpdateFrameMetadataPostPacketization(IPixelCaptureOutputFrame& Frame)
 	{
 		FPixelCaptureFrameMetadata& FrameMetadata = Frame.Metadata;
-		FrameMetadata.LastPacketizationEndTime = FPlatformTime::Cycles64();
+		FrameMetadata.LastPacketizationEndTime = rtc::TimeMillis();
 
 		FStats::Get()->AddFrameTimingStats(FrameMetadata);
 	}
@@ -243,7 +243,7 @@ namespace UE::PixelStreaming
 				Image.SetSpatialIndex(0);
 				Image.rotation_ = webrtc::VideoRotation::kVideoRotation_0;
 				Image.SetTimestamp(frame.timestamp());
-				Image.capture_time_ms_ = Packet.Timestamp / 1000.0;
+				Image.capture_time_ms_ = frame.timestamp_us() / 1000.0;
 
 				webrtc::CodecSpecificInfo CodecInfo;
 				switch (Codec)
diff --git a/Engine/Plugins/Media/PixelStreaming/Source/PixelStreaming/Private/VideoEncoderSingleLayerVPX.cpp b/Engine/Plugins/Media/PixelStreaming/Source/PixelStreaming/Private/VideoEncoderSingleLayerVPX.cpp
--- a/Engine/Plugins/Media/PixelStreaming/Source/PixelStreaming/Private/VideoEncoderSingleLayerVPX.cpp
+++ b/Engine/Plugins/Media/PixelStreaming/Source/PixelStreaming/Private/VideoEncoderSingleLayerVPX.cpp
@@ -42,7 +42,7 @@ namespace UE::PixelStreaming
 	{
 		FPixelCaptureFrameMetadata& FrameMetadata = Frame.Metadata;
 		FrameMetadata.UseCount++;
-		FrameMetadata.LastEncodeStartTime = FPlatformTime::Cycles64();
+		FrameMetadata.LastEncodeStartTime = rtc::TimeMillis();
 		if (FrameMetadata.UseCount == 1)
 		{
 			FrameMetadata.FirstEncodeStartTime = FrameMetadata.LastEncodeStartTime;
@@ -52,7 +52,7 @@ namespace UE::PixelStreaming
 	void FVideoEncoderSingleLayerVPX::UpdateFrameMetadataPostEncode(IPixelCaptureOutputFrame& Frame)
 	{
 		FPixelCaptureFrameMetadata& FrameMetadata = Frame.Metadata;
-		FrameMetadata.LastEncodeEndTime = FPlatformTime::Cycles64();
+		FrameMetadata.LastEncodeEndTime = rtc::TimeMillis();
 
 		FStats::Get()->AddFrameTimingStats(FrameMetadata);
 	}
