Add: Debug information to log when same frame is encoded or frames are missed

diff --git a/Engine/Plugins/Media/PixelCapture/Source/PixelCapture/Private/PixelCaptureCapturer.cpp b/Engine/Plugins/Media/PixelCapture/Source/PixelCapture/Private/PixelCaptureCapturer.cpp
--- a/Engine/Plugins/Media/PixelCapture/Source/PixelCapture/Private/PixelCaptureCapturer.cpp
+++ b/Engine/Plugins/Media/PixelCapture/Source/PixelCapture/Private/PixelCaptureCapturer.cpp
@@ -25,7 +25,6 @@ void FPixelCaptureCapturer::Capture(const IPixelCaptureInputFrame& InputFrame)
 	}
 
 	bBusy = true;
-
 	const int32 InputWidth = InputFrame.GetWidth();
 	const int32 InputHeight = InputFrame.GetHeight();
 
@@ -49,11 +48,13 @@ void FPixelCaptureCapturer::Capture(const IPixelCaptureInputFrame& InputFrame)
 	GPUEnqueueTime = 0;
 	GPUStartTime = 0;
 
+	// Todo (Luke.Bermingham) - Converting output buffer to raw ptr here seems error prone considering how Buffer->LockProduceBuffer() works.
 	BeginProcess(InputFrame, CurrentOutputBuffer.Get());
 }
 
 void FPixelCaptureCapturer::Initialize(int32 InputWidth, int32 InputHeight)
 {
+	checkf(InputWidth > 0 && InputHeight > 0, TEXT("Capture should be initialized with non-zero resolution."));
 	Buffer = MakeUnique<UE::PixelCapture::FOutputFrameBuffer>();
 	Buffer->Reset(3, 10, [this, InputWidth, InputHeight]() { return TSharedPtr<IPixelCaptureOutputFrame>(CreateOutputBuffer(InputWidth, InputHeight)); });
 	ExpectedInputWidth = InputWidth;
@@ -87,7 +88,7 @@ void FPixelCaptureCapturer::MarkGPUWorkStart()
 		MarkGPUWorkEnd();
 	}
 	GPUStartTime = rtc::TimeMillis();
-	
+
 	CurrentOutputBuffer->Metadata.CaptureProcessGPUDelay += GPUStartTime - GPUEnqueueTime;
 	GPUEnqueueTime = 0;
 }
@@ -100,6 +101,7 @@ void FPixelCaptureCapturer::MarkGPUWorkEnd()
 
 void FPixelCaptureCapturer::InitMetadata(FPixelCaptureFrameMetadata Metadata)
 {
+	Metadata.Id = FrameId.Increment();
 	Metadata.ProcessName = GetCapturerName();
 	Metadata.CaptureTime = 0;
 	Metadata.CaptureProcessCPUTime = 0;
@@ -127,7 +129,7 @@ void FPixelCaptureCapturer::EndProcess()
 	checkf(bBusy, TEXT("Capture process EndProcess called but we're not busy. Maybe double called?"));
 
 	FinalizeMetadata();
-	
+
 	CurrentOutputBuffer = nullptr;
 	Buffer->ReleaseProduceBuffer();
 	bBusy = false;
diff --git a/Engine/Plugins/Media/PixelCapture/Source/PixelCapture/Public/PixelCaptureCapturer.h b/Engine/Plugins/Media/PixelCapture/Source/PixelCapture/Public/PixelCaptureCapturer.h
--- a/Engine/Plugins/Media/PixelCapture/Source/PixelCapture/Public/PixelCaptureCapturer.h
+++ b/Engine/Plugins/Media/PixelCapture/Source/PixelCapture/Public/PixelCaptureCapturer.h
@@ -4,6 +4,7 @@
 
 #include "IPixelCaptureInputFrame.h"
 #include "IPixelCaptureOutputFrame.h"
+#include "HAL/ThreadSafeCounter64.h"
 
 namespace UE::PixelCapture
 {
@@ -130,6 +131,7 @@ private:
 	int32 ExpectedInputWidth = 0;
 	int32 ExpectedInputHeight = 0;
 
+	FThreadSafeCounter64 FrameId;
 	uint64 StartTime;
 	uint64 CPUStartTime;
 	uint64 GPUEnqueueTime;
diff --git a/Engine/Plugins/Media/PixelCapture/Source/PixelCapture/Public/PixelCaptureFrameMetadata.h b/Engine/Plugins/Media/PixelCapture/Source/PixelCapture/Public/PixelCaptureFrameMetadata.h
--- a/Engine/Plugins/Media/PixelCapture/Source/PixelCapture/Public/PixelCaptureFrameMetadata.h
+++ b/Engine/Plugins/Media/PixelCapture/Source/PixelCapture/Public/PixelCaptureFrameMetadata.h
@@ -9,6 +9,9 @@ struct FPixelCaptureFrameMetadata
 	// Identifier for the capture pipeline/process this frame took
 	FString ProcessName = "Unknown";
 
+	// Identifier for the frame
+	uint64 Id = 0;
+
 	// Which layer this specific frame is associated with
 	int32 Layer = 0;
 
diff --git a/Engine/Plugins/Media/PixelStreaming/Source/PixelStreaming/Private/VideoEncoderSingleLayerHardware.cpp b/Engine/Plugins/Media/PixelStreaming/Source/PixelStreaming/Private/VideoEncoderSingleLayerHardware.cpp
--- a/Engine/Plugins/Media/PixelStreaming/Source/PixelStreaming/Private/VideoEncoderSingleLayerHardware.cpp
+++ b/Engine/Plugins/Media/PixelStreaming/Source/PixelStreaming/Private/VideoEncoderSingleLayerHardware.cpp
@@ -30,8 +30,9 @@ namespace UE::PixelStreaming
 	{
 		VideoConfig.Preset = PixelStreaming::Settings::GetEncoderPreset();
 		VideoConfig.LatencyMode = EAVLatencyMode::UltraLowLatency;
-		VideoConfig.Width = InCodecSettings->width;
-		VideoConfig.Height = InCodecSettings->height;
+		// We set width and height to zero here because we initialize encoder from the first frame dimensions, not this config.
+		VideoConfig.Width = 0;
+		VideoConfig.Height = 0;
 		VideoConfig.TargetFramerate = InCodecSettings->maxFramerate;
 		VideoConfig.TargetBitrate = InCodecSettings->startBitrate;
 		VideoConfig.MaxBitrate = InCodecSettings->maxBitrate;
@@ -137,9 +138,9 @@ namespace UE::PixelStreaming
 		const FPixelCaptureOutputFrameRHI& RHILayer = StaticCast<const FPixelCaptureOutputFrameRHI&>(AdaptedLayer);
 		// TODO-TE
 #if WEBRTC_5414
-		rtc::scoped_refptr<FFrameBufferRHI> RHIBuffer = rtc::make_ref_counted<FFrameBufferRHI>(MakeShared<FVideoResourceRHI>(HardwareEncoder.Pin()->GetDevice().ToSharedRef(), FVideoResourceRHI::FRawData{ RHILayer.GetFrameTexture(), nullptr, 0 }));
+		rtc::scoped_refptr<FFrameBufferRHI> RHIBuffer = rtc::make_ref_counted<FFrameBufferRHI>(MakeShared<FVideoResourceRHI>(HardwareEncoder->GetDevice().ToSharedRef(), FVideoResourceRHI::FRawData{ RHILayer.GetFrameTexture(), nullptr, 0 }));
 #else
-		rtc::scoped_refptr<FFrameBufferRHI> RHIBuffer = new rtc::RefCountedObject<FFrameBufferRHI>(MakeShared<FVideoResourceRHI>(HardwareEncoder.Pin()->GetDevice().ToSharedRef(), FVideoResourceRHI::FRawData{ RHILayer.GetFrameTexture(), nullptr, 0 }));
+		rtc::scoped_refptr<FFrameBufferRHI> RHIBuffer = new rtc::RefCountedObject<FFrameBufferRHI>(MakeShared<FVideoResourceRHI>(HardwareEncoder->GetDevice().ToSharedRef(), FVideoResourceRHI::FRawData{ RHILayer.GetFrameTexture(), nullptr, 0 }));
 #endif
 		NewFrame.set_video_frame_buffer(RHIBuffer);
 		return NewFrame;
@@ -147,18 +148,20 @@ namespace UE::PixelStreaming
 
 	void FVideoEncoderSingleLayerHardware::LateInitHardwareEncoder(uint32 StreamId)
 	{
+		checkf(InitialVideoConfig->Width > 0 && InitialVideoConfig->Height > 0, TEXT("Encoder config must have non-zero width and height."));
+
 		switch (Codec)
 		{
 			case EPixelStreamingCodec::H264:
 			{
 				FVideoEncoderConfigH264& VideoConfig = *StaticCast<FVideoEncoderConfigH264*>(InitialVideoConfig.Get());
-				HardwareEncoder = Factory.GetOrCreateHardwareEncoder(StreamId, VideoConfig);
+				HardwareEncoder = Factory.GetOrCreateHardwareEncoder(StreamId, VideoConfig).Pin();
 				break;
 			}
 			case EPixelStreamingCodec::H265:
 			{
 				FVideoEncoderConfigH265& VideoConfig = *StaticCast<FVideoEncoderConfigH265*>(InitialVideoConfig.Get());
-				HardwareEncoder = Factory.GetOrCreateHardwareEncoder(StreamId, VideoConfig);
+				HardwareEncoder = Factory.GetOrCreateHardwareEncoder(StreamId, VideoConfig).Pin();
 				break;
 			}
 		}
@@ -167,28 +170,52 @@ namespace UE::PixelStreaming
 	int32 FVideoEncoderSingleLayerHardware::Encode(webrtc::VideoFrame const& frame, std::vector<webrtc::VideoFrameType> const* frame_types)
 	{
 		const FFrameBufferMultiFormat* FrameBuffer = StaticCast<FFrameBufferMultiFormat*>(frame.video_frame_buffer().get());
+		EncodingStreamId = FrameBuffer->GetSourceStreamId();
+
+		IPixelCaptureOutputFrame* AdaptedLayer = FrameBuffer->RequestFormat(PixelCaptureBufferFormat::FORMAT_RHI);
+		// Check whether the output frame is valid because null frames are passed to stream sharing encoders.
+		if (AdaptedLayer == nullptr)
+		{
+			// probably the first request which starts the adapt pipeline for this format
+			return WEBRTC_VIDEO_CODEC_OK;
+		}
+
+		const int32 FrameWidth = AdaptedLayer->GetWidth();
+		const int32 FrameHeight = AdaptedLayer->GetHeight();
 
 		// We late init here so we can pull the stream ID off the incoming frames and pull the correct encoder for the stream
 		// earlier locations do not have this information.
-		if (!HardwareEncoder.IsValid())
+		if (!HardwareEncoder)
 		{
-			EncodingStreamId = FrameBuffer->GetSourceStreamId();
+			InitialVideoConfig->Width = FrameWidth;
+			InitialVideoConfig->Height = FrameHeight;
 			LateInitHardwareEncoder(EncodingStreamId);
 		}
 
-		if (TSharedPtr<FVideoEncoderHardware> const& PinnedHardwareEncoder = HardwareEncoder.Pin())
+		// Update the encoding config using the incoming frame resolution
+		UpdateConfig(FrameWidth, FrameHeight);
+
+		if (HardwareEncoder)
 		{
 			TRACE_CPUPROFILER_EVENT_SCOPE_ON_CHANNEL_STR("PixelStreaming Hardware Encoding", PixelStreamingChannel);
 
-			IPixelCaptureOutputFrame* AdaptedLayer = FrameBuffer->RequestFormat(PixelCaptureBufferFormat::FORMAT_RHI);
-
-			if (AdaptedLayer == nullptr)
+			uint64 CurrentFrameId = AdaptedLayer->Metadata.Id;
+			if(LastEncodedFrameId > 0)
 			{
-				// probably the first request which starts the adapt pipeline for this format
-				return WEBRTC_VIDEO_CODEC_OK;
+				if(CurrentFrameId < LastEncodedFrameId)
+				{
+					UE_LOG(LogPixelStreaming, Verbose, TEXT("Out of order frame entered hardware encoder."))
+				}
+				else if(CurrentFrameId == LastEncodedFrameId)
+				{
+					UE_LOG(LogPixelStreaming, Verbose, TEXT("Encoding the same frame."))
+				}
+				else if(CurrentFrameId - LastEncodedFrameId > 1)
+				{
+					UE_LOG(LogPixelStreaming, Verbose, TEXT("Missed encoding %d frames"), CurrentFrameId - LastEncodedFrameId - 1);
+				}
 			}
-
-			UpdateConfig(AdaptedLayer->GetWidth(), AdaptedLayer->GetHeight());
+			LastEncodedFrameId = CurrentFrameId;
 
 			UpdateFrameMetadataPreEncode(*AdaptedLayer);
 			const FPixelCaptureOutputFrameRHI& RHILayer = StaticCast<const FPixelCaptureOutputFrameRHI&>(*AdaptedLayer);
@@ -199,35 +226,33 @@ namespace UE::PixelStreaming
 			{
 				// If we have a wrapper format we need to also provide a descriptor of the raw format the data
 				// is transported in
-				FVideoDescriptor RawDescriptor = FVideoResourceRHI::GetDescriptorFrom(PinnedHardwareEncoder->GetDevice().ToSharedRef(), RHILayer.GetFrameTexture());
+				FVideoDescriptor RawDescriptor = FVideoResourceRHI::GetDescriptorFrom(HardwareEncoder->GetDevice().ToSharedRef(), RHILayer.GetFrameTexture());
 
 				// If we have a wrapper format for our resource, make sure the video descriptor specifies it
 				RHIBuffer = new rtc::RefCountedObject<FFrameBufferRHI>(MakeShared<FVideoResourceRHI>(
-					PinnedHardwareEncoder->GetDevice().ToSharedRef(),
-					// TODO-TE
+					HardwareEncoder->GetDevice().ToSharedRef(),
 					FVideoResourceRHI::FRawData{ RHILayer.GetFrameTexture(), nullptr, 0 },
 					FVideoDescriptor(static_cast<EVideoFormat>(WrapperData.Format), WrapperData.Width, WrapperData.Height, RawDescriptor)));
 			}
 			else
 			{
 				RHIBuffer = new rtc::RefCountedObject<FFrameBufferRHI>(MakeShared<FVideoResourceRHI>(
-					PinnedHardwareEncoder->GetDevice().ToSharedRef(),
-					// TODO-TE
+					HardwareEncoder->GetDevice().ToSharedRef(),
 					FVideoResourceRHI::FRawData{ RHILayer.GetFrameTexture(), nullptr, 0 }));
 			}
 
 			const bool bKeyframe = (frame_types && (*frame_types)[0] == webrtc::VideoFrameType::kVideoFrameKey) || Factory.ShouldForceKeyframe();
 
-			PinnedHardwareEncoder->SendFrame(RHIBuffer->GetVideoResource(), frame.timestamp_us(), bKeyframe);
+			HardwareEncoder->SendFrame(RHIBuffer->GetVideoResource(), frame.timestamp_us(), bKeyframe);
 
 			Factory.UnforceKeyFrame();
 
 			UpdateFrameMetadataPostEncode(*AdaptedLayer);
 
 			FVideoPacket Packet;
-			while (PinnedHardwareEncoder->ReceivePacket(Packet))
+			while (HardwareEncoder->ReceivePacket(Packet))
 			{
-				// TODO (Andrew) This works fine for 1:1 synchronous encoding, but will need to relate frames to packets for async or gop
+				// Note: This works fine for 1:1 synchronous encoding, but if we ever need to relate frames to packets for async or gop this will need thought.
 				webrtc::EncodedImage Image;
 
 				Image.timing_.packetization_finish_ms = FTimespan::FromSeconds(FPlatformTime::Seconds()).GetTotalMilliseconds();
@@ -313,17 +338,20 @@ namespace UE::PixelStreaming
 
 	void FVideoEncoderSingleLayerHardware::UpdateConfig(uint32 width, uint32 height)
 	{
-		if (TSharedPtr<FVideoEncoderHardware> const& PinnedHardwareEncoder = HardwareEncoder.Pin())
+		InitialVideoConfig->Width = width;
+		InitialVideoConfig->Height = height;
+
+		if (HardwareEncoder)
 		{
-			FVideoEncoderConfig VideoConfigMinimal = PinnedHardwareEncoder->GetMinimalConfig();
+			FVideoEncoderConfig VideoConfigMinimal = HardwareEncoder->GetMinimalConfig();
 			FVideoEncoderConfig* VideoConfig = &VideoConfigMinimal;
 
 			switch (Codec)
 			{
 				case EPixelStreamingCodec::H264:
-					if (PinnedHardwareEncoder->GetInstance()->Has<FVideoEncoderConfigH264>())
+					if (HardwareEncoder->GetInstance()->Has<FVideoEncoderConfigH264>())
 					{
-						FVideoEncoderConfigH264& VideoConfigH264 = PinnedHardwareEncoder->GetInstance()->Edit<FVideoEncoderConfigH264>();
+						FVideoEncoderConfigH264& VideoConfigH264 = HardwareEncoder->GetInstance()->Edit<FVideoEncoderConfigH264>();
 						VideoConfig = &VideoConfigH264;
 
 						VideoConfigH264.Profile = UE::PixelStreaming::Settings::GetH264Profile();
@@ -331,9 +359,9 @@ namespace UE::PixelStreaming
 
 					break;
 				case EPixelStreamingCodec::H265:
-					if (PinnedHardwareEncoder->GetInstance()->Has<FVideoEncoderConfigH265>())
+					if (HardwareEncoder->GetInstance()->Has<FVideoEncoderConfigH265>())
 					{
-						FVideoEncoderConfigH265& VideoConfigH265 = PinnedHardwareEncoder->GetInstance()->Edit<FVideoEncoderConfigH265>();
+						FVideoEncoderConfigH265& VideoConfigH265 = HardwareEncoder->GetInstance()->Edit<FVideoEncoderConfigH265>();
 						VideoConfig = &VideoConfigH265;
 
 						VideoConfigH265.Profile = UE::PixelStreaming::Settings::GetH265Profile();
@@ -375,7 +403,7 @@ namespace UE::PixelStreaming
 			VideoConfig->Width = width;
 			VideoConfig->Height = height;
 
-			PinnedHardwareEncoder->SetMinimalConfig(*VideoConfig);
+			HardwareEncoder->SetMinimalConfig(*VideoConfig);
 		}
 	}
 
@@ -393,7 +421,7 @@ namespace UE::PixelStreaming
 
 	void FVideoEncoderSingleLayerHardware::MaybeDumpFrame(webrtc::EncodedImage const& encoded_image)
 	{
-		// Dump H265 frames to file for debugging if CVar is turned on.
+		// Dump encoded frames to file for debugging if CVar is turned on.
 		if (UE::PixelStreaming::Settings::CVarPixelStreamingDebugDumpFrame.GetValueOnAnyThread())
 		{
 			static IFileHandle* FileHandle = nullptr;
@@ -403,6 +431,7 @@ namespace UE::PixelStreaming
 				FString TempFilePath = FPaths::CreateTempFilename(*FPaths::ProjectSavedDir(), TEXT("encoded_frame"), TEXT(".raw"));
 				FileHandle = PlatformFile.OpenWrite(*TempFilePath);
 				check(FileHandle);
+				// Note: To examine individual frames from this dump: ffmpeg -i "encoded_frame78134A5047638BB99AE1D88471E5E513.raw" "frames/out-%04d.jpg"
 			}
 
 			FileHandle->Write(encoded_image.data(), encoded_image.size());
diff --git a/Engine/Plugins/Media/PixelStreaming/Source/PixelStreaming/Private/VideoEncoderSingleLayerHardware.h b/Engine/Plugins/Media/PixelStreaming/Source/PixelStreaming/Private/VideoEncoderSingleLayerHardware.h
--- a/Engine/Plugins/Media/PixelStreaming/Source/PixelStreaming/Private/VideoEncoderSingleLayerHardware.h
+++ b/Engine/Plugins/Media/PixelStreaming/Source/PixelStreaming/Private/VideoEncoderSingleLayerHardware.h
@@ -51,7 +51,7 @@ namespace UE::PixelStreaming
 		EPixelStreamingCodec Codec;
 
 		TUniquePtr<FVideoEncoderConfig> InitialVideoConfig;
-		TWeakPtr<FVideoEncoderHardware> HardwareEncoder;
+		TSharedPtr<FVideoEncoderHardware> HardwareEncoder;
 
 		// We store this so we can restore back to it if the user decides to use then stop using the PixelStreaming.Encoder.TargetBitrate CVar.
 		int32 WebRtcProposedTargetBitrate = 5000000;
@@ -64,5 +64,8 @@ namespace UE::PixelStreaming
 
 		// used to key into active hardware encoders and pull the correct encoder for the stream.
 		uint32 EncodingStreamId;
+
+		// Used to track how often we encode the same frame
+		uint64 LastEncodedFrameId = 0;
 	};
 } // namespace UE::PixelStreaming
