Fix streaming sharing not working when quality control request is not permitted.

diff --git a/Engine/Plugins/Media/PixelStreaming/Source/PixelStreaming/Private/FrameBufferMultiFormat.cpp b/Engine/Plugins/Media/PixelStreaming/Source/PixelStreaming/Private/FrameBufferMultiFormat.cpp
--- a/Engine/Plugins/Media/PixelStreaming/Source/PixelStreaming/Private/FrameBufferMultiFormat.cpp
+++ b/Engine/Plugins/Media/PixelStreaming/Source/PixelStreaming/Private/FrameBufferMultiFormat.cpp
@@ -17,17 +17,17 @@ namespace UE::PixelStreaming
 
 	int FFrameBufferMultiFormatLayered::width() const
 	{
-		return FrameCapturer->GetWidth(GetNumLayers() - 1);
+		return FrameCapturer ? FrameCapturer->GetWidth(GetNumLayers() - 1) : -1;
 	}
 
 	int FFrameBufferMultiFormatLayered::height() const
 	{
-		return FrameCapturer->GetHeight(GetNumLayers() - 1);
+		return FrameCapturer ? FrameCapturer->GetHeight(GetNumLayers() - 1) : -1;
 	}
 
 	int FFrameBufferMultiFormatLayered::GetNumLayers() const
 	{
-		return FrameCapturer->GetNumLayers();
+		return FrameCapturer ? FrameCapturer->GetNumLayers() : -1;
 	}
 
 	rtc::scoped_refptr<FFrameBufferMultiFormat> FFrameBufferMultiFormatLayered::GetLayer(int LayerIndex) const
@@ -47,12 +47,12 @@ namespace UE::PixelStreaming
 
 	int FFrameBufferMultiFormat::width() const
 	{
-		return FrameCapturer->GetWidth(LayerIndex);
+		return FrameCapturer ? FrameCapturer->GetWidth(LayerIndex) : -1;
 	}
 
 	int FFrameBufferMultiFormat::height() const
 	{
-		return FrameCapturer->GetHeight(LayerIndex);
+		return FrameCapturer ? FrameCapturer->GetHeight(LayerIndex) : -1;
 	}
 
 	IPixelCaptureOutputFrame* FFrameBufferMultiFormat::RequestFormat(int32 Format) const
@@ -62,6 +62,12 @@ namespace UE::PixelStreaming
 		{
 			return CachedFrame->Get();
 		}
+
+		if(!FrameCapturer)
+		{
+			return nullptr;
+		}
+
 		TSharedPtr<IPixelCaptureOutputFrame> Frame = FrameCapturer->WaitForFormat(Format, LayerIndex);
 		CachedFormat.Add(Format, Frame);
 		return Frame.Get();
diff --git a/Engine/Plugins/Media/PixelStreaming/Source/PixelStreaming/Private/PixelStreamingVideoInput.cpp b/Engine/Plugins/Media/PixelStreaming/Source/PixelStreaming/Private/PixelStreamingVideoInput.cpp
--- a/Engine/Plugins/Media/PixelStreaming/Source/PixelStreaming/Private/PixelStreamingVideoInput.cpp
+++ b/Engine/Plugins/Media/PixelStreaming/Source/PixelStreaming/Private/PixelStreamingVideoInput.cpp
@@ -46,6 +46,15 @@ rtc::scoped_refptr<webrtc::VideoFrameBuffer> FPixelStreamingVideoInput::GetFrame
 #endif
 }
 
+rtc::scoped_refptr<webrtc::VideoFrameBuffer> FPixelStreamingVideoInput::GetEmptyFrameBuffer()
+{
+#if WEBRTC_5414
+	return rtc::make_ref_counted<UE::PixelStreaming::FFrameBufferMultiFormatLayered>(nullptr, StreamId);
+#else
+	return new rtc::RefCountedObject<UE::PixelStreaming::FFrameBufferMultiFormatLayered>(nullptr, StreamId);
+#endif
+}
+
 TSharedPtr<IPixelCaptureOutputFrame> FPixelStreamingVideoInput::RequestFormat(int32 Format, int32 LayerIndex)
 {
 	if (FrameCapturer != nullptr)
diff --git a/Engine/Plugins/Media/PixelStreaming/Source/PixelStreaming/Private/VideoSource.cpp b/Engine/Plugins/Media/PixelStreaming/Source/PixelStreaming/Private/VideoSource.cpp
--- a/Engine/Plugins/Media/PixelStreaming/Source/PixelStreaming/Private/VideoSource.cpp
+++ b/Engine/Plugins/Media/PixelStreaming/Source/PixelStreaming/Private/VideoSource.cpp
@@ -24,6 +24,11 @@ namespace UE::PixelStreaming
 			{
 				PushFrame();
 			}
+			// We always send a frame for WebRTC purposes, but empty frames early exit out of the encoder.
+			else
+			{
+				PushEmptyFrame();
+			}
 		}
 	}
 
@@ -42,4 +47,20 @@ namespace UE::PixelStreaming
 									   .build();
 		OnFrame(Frame);
 	}
+
+	void FVideoSource::PushEmptyFrame()
+	{
+		TRACE_CPUPROFILER_EVENT_SCOPE_ON_CHANNEL_STR("PixelStreaming Push Empty Video Frame", PixelStreamingChannel);
+		static int32 EmptyFrameId = 1;
+
+		rtc::scoped_refptr<webrtc::VideoFrameBuffer> FrameBuffer = VideoInput->GetEmptyFrameBuffer();
+		webrtc::VideoFrame Frame = webrtc::VideoFrame::Builder()
+									   .set_video_frame_buffer(FrameBuffer)
+									   .set_timestamp_us(rtc::TimeMicros())
+									   .set_rotation(webrtc::VideoRotation::kVideoRotation_0)
+									   .set_id(EmptyFrameId++)
+									   .build();
+		OnFrame(Frame);
+	}
+
 } // namespace UE::PixelStreaming
diff --git a/Engine/Plugins/Media/PixelStreaming/Source/PixelStreaming/Private/VideoSource.h b/Engine/Plugins/Media/PixelStreaming/Source/PixelStreaming/Private/VideoSource.h
--- a/Engine/Plugins/Media/PixelStreaming/Source/PixelStreaming/Private/VideoSource.h
+++ b/Engine/Plugins/Media/PixelStreaming/Source/PixelStreaming/Private/VideoSource.h
@@ -39,5 +39,6 @@ namespace UE::PixelStreaming
 		TFunction<bool()> ShouldGenerateFramesCheck;
 
 		void PushFrame();
+		void PushEmptyFrame();
 	};
 } // namespace UE::PixelStreaming
diff --git a/Engine/Plugins/Media/PixelStreaming/Source/PixelStreaming/Public/PixelStreamingVideoInput.h b/Engine/Plugins/Media/PixelStreaming/Source/PixelStreaming/Public/PixelStreamingVideoInput.h
--- a/Engine/Plugins/Media/PixelStreaming/Source/PixelStreaming/Public/PixelStreamingVideoInput.h
+++ b/Engine/Plugins/Media/PixelStreaming/Source/PixelStreaming/Public/PixelStreamingVideoInput.h
@@ -39,6 +39,12 @@ public:
 	 */
 	rtc::scoped_refptr<webrtc::VideoFrameBuffer> GetFrameBuffer();
 
+	/**
+	 * Used internally by the encode process.
+	 * @return A frame buffer that has a null capturer.
+	 */
+	rtc::scoped_refptr<webrtc::VideoFrameBuffer> GetEmptyFrameBuffer();
+
 	/**
 	 *
 	 */
