[NO NON-PII COMMIT MESSAGE]

diff --git a/Engine/Plugins/Media/PixelStreaming/Source/PixelStreaming/Private/Settings.cpp b/Engine/Plugins/Media/PixelStreaming/Source/PixelStreaming/Private/Settings.cpp
--- a/Engine/Plugins/Media/PixelStreaming/Source/PixelStreaming/Private/Settings.cpp
+++ b/Engine/Plugins/Media/PixelStreaming/Source/PixelStreaming/Private/Settings.cpp
@@ -291,6 +291,12 @@ namespace UE::PixelStreaming::Settings
 		TEXT("Whether we should only stream as fast as we render or at some fixed interval. Coupled means only stream what we render."),
 		ECVF_Default);
 
+	TAutoConsoleVariable<float> CVarPixelStreamingDecoupleWaitFactor(
+		TEXT("PixelStreaming.DecoupleWaitFactor"),
+		1.25f,
+		TEXT("Frame rate factor to wait for a captured frame when streaming in decoupled mode. Higher factor waits longer but may also result in higher latency."),
+		ECVF_Default);
+
 	TAutoConsoleVariable<float> CVarPixelStreamingSignalingReconnectInterval(
 		TEXT("PixelStreaming.SignalingReconnectInterval"),
 		2.0f,
@@ -678,6 +684,7 @@ namespace UE::PixelStreaming::Settings
 		CommandLineParseValue(TEXT("PixelStreamingFreezeFrameQuality"), CVarPixelStreamingFreezeFrameQuality);
 		CommandLineParseValue(TEXT("PixelStreamingInputController="), CVarPixelStreamingInputController);
 		CommandLineParseValue(TEXT("PixelStreamingSignalingReconnectInterval="), CVarPixelStreamingSignalingReconnectInterval);
+		CommandLineParseValue(TEXT("PixelStreamingDecoupleWaitFactor="), CVarPixelStreamingDecoupleWaitFactor);
 
 		// Options parse (if these exist they are set to true)
 		CommandLineParseOption(TEXT("PixelStreamingOnScreenStats"), CVarPixelStreamingOnScreenStats);
diff --git a/Engine/Plugins/Media/PixelStreaming/Source/PixelStreaming/Private/Settings.h b/Engine/Plugins/Media/PixelStreaming/Source/PixelStreaming/Private/Settings.h
--- a/Engine/Plugins/Media/PixelStreaming/Source/PixelStreaming/Private/Settings.h
+++ b/Engine/Plugins/Media/PixelStreaming/Source/PixelStreaming/Private/Settings.h
@@ -69,6 +69,7 @@ namespace UE::PixelStreaming::Settings
 	extern TAutoConsoleVariable<bool> CVarPixelStreamingExperimentalAudioInput;
 	extern TAutoConsoleVariable<bool> CVarPixelStreamingCaptureUseFence;
 	extern TAutoConsoleVariable<bool> CVarPixelStreamingDecoupleFramerate;
+	extern TAutoConsoleVariable<float> CVarPixelStreamingDecoupleWaitFactor;
 	// Ends Pixel Streaming Plugin CVars
 
 	/* Pixel Streaming can limit who can send input (keyboard, mouse, etc). */
diff --git a/Engine/Plugins/Media/PixelStreaming/Source/PixelStreaming/Private/VideoSourceGroup.cpp b/Engine/Plugins/Media/PixelStreaming/Source/PixelStreaming/Private/VideoSourceGroup.cpp
--- a/Engine/Plugins/Media/PixelStreaming/Source/PixelStreaming/Private/VideoSourceGroup.cpp
+++ b/Engine/Plugins/Media/PixelStreaming/Source/PixelStreaming/Private/VideoSourceGroup.cpp
@@ -194,10 +194,11 @@ namespace UE::PixelStreaming
 
 				// Decrease this value to make expected frame delivery more precise, however may result in more old frames being sent
 				const double PrecisionFactor = 0.1;
+				const double WaitFactor = UE::PixelStreaming::Settings::CVarPixelStreamingDecoupleWaitFactor.GetValueOnAnyThread();
 
 				// In "auto" mode vary this value based on historical average
 				const double TargetSubmitMs = 1000.0 / VideoSourceGroup->FramesPerSecond;
-				const double TargetSubmitMsWithPadding = TargetSubmitMs * (1.0 + PrecisionFactor);
+				const double TargetSubmitMsWithPadding = TargetSubmitMs * WaitFactor;
 				const double CloseEnoughMs = TargetSubmitMs * PrecisionFactor;
 				const bool bFrameOverdue = TimeSinceLastSubmitMs >= TargetSubmitMsWithPadding;
 
@@ -205,7 +206,7 @@ namespace UE::PixelStreaming
 				if(!bFrameOverdue)
 				{
 					// Frame arrived in a timely fashion, but is it too soon to maintain our target rate? If so, sleep.
-					double WaitTimeRemainingMs = TargetSubmitMs - TimeSinceLastSubmitMs;
+					double WaitTimeRemainingMs = TargetSubmitMsWithPadding - TimeSinceLastSubmitMs;
 					if(WaitTimeRemainingMs > CloseEnoughMs)
 					{
 						bool bGotNewFrame = FrameEvent.Get()->Wait(WaitTimeRemainingMs);
