Fix crash when pixel streaming shutting down. Cause was thread object (FrameThread) being destroyed while the thread was still running. Fix is just to wait for thread to finish first.

diff --git a/Engine/Plugins/Media/PixelStreaming/Source/PixelStreaming/Private/VideoSourceGroup.cpp b/Engine/Plugins/Media/PixelStreaming/Source/PixelStreaming/Private/VideoSourceGroup.cpp
--- a/Engine/Plugins/Media/PixelStreaming/Source/PixelStreaming/Private/VideoSourceGroup.cpp
+++ b/Engine/Plugins/Media/PixelStreaming/Source/PixelStreaming/Private/VideoSourceGroup.cpp
@@ -155,7 +155,7 @@ namespace UE::PixelStreaming
 	{
 		if (FrameThread != nullptr)
 		{
-			FrameThread->Kill(false);
+			FrameThread->Kill(true);
 		}
 		FrameThread = nullptr;
 		bThreadRunning = false;
