Added missing memzeros for audio generator to corner cases where they might fail to produce a valid audio data

diff --git a/Engine/Plugins/Media/PixelStreaming/Source/PixelStreaming/Private/PixelStreamingAudioComponent.cpp b/Engine/Plugins/Media/PixelStreaming/Source/PixelStreaming/Private/PixelStreamingAudioComponent.cpp
--- a/Engine/Plugins/Media/PixelStreaming/Source/PixelStreaming/Private/PixelStreamingAudioComponent.cpp
+++ b/Engine/Plugins/Media/PixelStreaming/Source/PixelStreaming/Private/PixelStreamingAudioComponent.cpp
@@ -244,18 +244,25 @@ int32 FWebRTCSoundGenerator::OnGenerateAudio(float* OutAudio, int32 NumSamples)
 		return NumSamples;
 	}
 
+	int32 NumSamplesToCopy = NumSamples;
+
 	// Critical section
 	{
 		FScopeLock Lock(&CriticalSection);
 
-		int32 NumSamplesToCopy = FGenericPlatformMath::Min(NumSamples, Buffer.Num());
+		NumSamplesToCopy = FGenericPlatformMath::Min(NumSamples, Buffer.Num());
 
 		// Copy from local buffer into OutAudio if we have enough samples
 		Audio::ArrayPcm16ToFloat(MakeArrayView(Buffer.GetData(), NumSamplesToCopy), MakeArrayView(OutAudio, NumSamplesToCopy));
 
 		// Remove front NumSamples from the local buffer
 		Buffer.RemoveAt(0, NumSamplesToCopy, false);
+	}
 
-		return NumSamplesToCopy;
+	if (NumSamplesToCopy < NumSamples)
+	{
+		FMemory::Memzero(&OutAudio[NumSamplesToCopy], (NumSamples - NumSamplesToCopy) * sizeof(float));
 	}
+
+	return NumSamplesToCopy;
 }
diff --git a/Engine/Source/Runtime/Engine/Private/SoundGenerator.cpp b/Engine/Source/Runtime/Engine/Private/SoundGenerator.cpp
--- a/Engine/Source/Runtime/Engine/Private/SoundGenerator.cpp
+++ b/Engine/Source/Runtime/Engine/Private/SoundGenerator.cpp
@@ -20,7 +20,12 @@ int32 ISoundGenerator::GetNextBuffer(float* OutAudio, int32 NumSamples, bool bRe
 	int32 NumSamplesToGenerate = NumSamples;
 	if (!bRequireNumberSamples)
 	{
-		NumSamplesToGenerate = FMath::Min(NumSamples, GetDesiredNumSamplesToRenderPerCallback());
+		const int32 NewNumSamplesToGenerate = FMath::Min(NumSamples, GetDesiredNumSamplesToRenderPerCallback());
+		if (const int32 Diff = NumSamples - NewNumSamplesToGenerate)
+		{
+			FMemory::Memzero(&OutAudio[NewNumSamplesToGenerate], Diff * sizeof(float));
+		}
+		NumSamplesToGenerate = NewNumSamplesToGenerate;
 	}
 
 	return OnGenerateAudio(OutAudio, NumSamplesToGenerate);
