Fix: Remove private includes in PixelStreamingHMD

diff --git a/Engine/Plugins/Media/PixelStreaming/Source/PixelStreamingHMD/PixelStreamingHMD.build.cs b/Engine/Plugins/Media/PixelStreaming/Source/PixelStreamingHMD/PixelStreamingHMD.build.cs
--- a/Engine/Plugins/Media/PixelStreaming/Source/PixelStreamingHMD/PixelStreamingHMD.build.cs
+++ b/Engine/Plugins/Media/PixelStreaming/Source/PixelStreamingHMD/PixelStreamingHMD.build.cs
@@ -6,12 +6,6 @@ namespace UnrealBuildTool.Rules
 	{
 		public PixelStreamingHMD(ReadOnlyTargetRules Target) : base(Target)
 		{
-			PrivateIncludePaths.AddRange(
-				new string[] {
-					System.IO.Path.Combine(GetModuleDirectory("Renderer"), "Private"),
-				}
-				);
-
 			PrivateDependencyModuleNames.AddRange(
 				new string[]
 				{
diff --git a/Engine/Plugins/Media/PixelStreaming/Source/PixelStreamingHMD/Private/PixelStreamingHMD.cpp b/Engine/Plugins/Media/PixelStreaming/Source/PixelStreamingHMD/Private/PixelStreamingHMD.cpp
--- a/Engine/Plugins/Media/PixelStreaming/Source/PixelStreamingHMD/Private/PixelStreamingHMD.cpp
+++ b/Engine/Plugins/Media/PixelStreaming/Source/PixelStreamingHMD/Private/PixelStreamingHMD.cpp
@@ -8,11 +8,11 @@
 #include "Engine/GameViewportClient.h"
 #include "GameFramework/PlayerController.h"
 #include "Engine/LocalPlayer.h"
-#include "SceneRendering.h"
 #include "PostProcess/PostProcessHMD.h"
 #include "GameFramework/WorldSettings.h"
 #include "Settings.h"
 #include "Widgets/SWindow.h"
+#include "SceneView.h"
 
 FPixelStreamingHMD::FPixelStreamingHMD(const FAutoRegister& AutoRegister)
 	: FHeadMountedDisplayBase(nullptr)
