Allowing bitrates to be set on connection.

diff --git a/Engine/Plugins/Media/PixelStreaming/Source/PixelStreaming/Private/PixelStreamingPeerConnection.cpp b/Engine/Plugins/Media/PixelStreaming/Source/PixelStreaming/Private/PixelStreamingPeerConnection.cpp
--- a/Engine/Plugins/Media/PixelStreaming/Source/PixelStreaming/Private/PixelStreamingPeerConnection.cpp
+++ b/Engine/Plugins/Media/PixelStreaming/Source/PixelStreaming/Private/PixelStreamingPeerConnection.cpp
@@ -348,6 +348,8 @@ void FPixelStreamingPeerConnection::CreateAnswer(EReceiveMediaOption ReceiveOpti
 
 void FPixelStreamingPeerConnection::ReceiveOffer(const FString& Sdp, const VoidCallback& SuccessCallback, const ErrorCallback& ErrorCallback)
 {
+	RefreshStreamBitrate();
+	
 	webrtc::SdpParseError Error;
 	std::unique_ptr<webrtc::SessionDescriptionInterface> SessionDesc = webrtc::CreateSessionDescription(webrtc::SdpType::kOffer, UE::PixelStreaming::ToString(Sdp), &Error);
 	if (SessionDesc)
@@ -366,6 +368,8 @@ void FPixelStreamingPeerConnection::ReceiveOffer(const FString& Sdp, const VoidC
 
 void FPixelStreamingPeerConnection::ReceiveAnswer(const FString& Sdp, const VoidCallback& SuccessCallback, const ErrorCallback& ErrorCallback)
 {
+	RefreshStreamBitrate();
+
 	webrtc::SdpParseError Error;
 	std::unique_ptr<webrtc::SessionDescriptionInterface> SessionDesc = webrtc::CreateSessionDescription(webrtc::SdpType::kAnswer, ToString(Sdp), &Error);
 	if (SessionDesc)
@@ -507,6 +511,14 @@ void FPixelStreamingPeerConnection::SetVideoSource(rtc::scoped_refptr<webrtc::Vi
 	}
 }
 
+void FPixelStreamingPeerConnection::RefreshStreamBitrate()
+{
+	webrtc::BitrateSettings bitrateSettings;
+	bitrateSettings.min_bitrate_bps = Settings::CVarPixelStreamingWebRTCMinBitrate.GetValueOnAnyThread();
+	bitrateSettings.max_bitrate_bps = Settings::CVarPixelStreamingWebRTCMaxBitrate.GetValueOnAnyThread();
+	PeerConnection->SetBitrate(bitrateSettings);
+}
+
 void FPixelStreamingPeerConnection::SetAudioSource(rtc::scoped_refptr<webrtc::AudioSourceInterface> InAudioSource)
 {
 	const static FString AudioTrackLabel = TEXT("pixelstreaming_audio_track_label");
diff --git a/Engine/Plugins/Media/PixelStreaming/Source/PixelStreaming/Private/PixelStreamingSignallingConnection.cpp b/Engine/Plugins/Media/PixelStreaming/Source/PixelStreaming/Private/PixelStreamingSignallingConnection.cpp
--- a/Engine/Plugins/Media/PixelStreaming/Source/PixelStreaming/Private/PixelStreamingSignallingConnection.cpp
+++ b/Engine/Plugins/Media/PixelStreaming/Source/PixelStreaming/Private/PixelStreamingSignallingConnection.cpp
@@ -12,6 +12,7 @@
 #include "Utils.h"
 #include "ToStringExtensions.h"
 #include "GenericPlatform/GenericPlatformHttp.h"
+#include "PixelStreamingModule.h"
 
 DECLARE_LOG_CATEGORY_EXTERN(LogPixelStreamingSS, Log, VeryVerbose);
 DEFINE_LOG_CATEGORY(LogPixelStreamingSS);
@@ -535,13 +536,26 @@ void FPixelStreamingSignallingConnection::OnSessionDescription(const FJsonObject
 
 	FPixelStreamingPlayerId PlayerId;
 	bool bGotPlayerId = GetPlayerIdJson(Json, PlayerId);
-	if (!bGotPlayerId)
+	if (bGotPlayerId)
 	{
-		Observer->OnSignallingSessionDescription(Type, Sdp);
+		int MinBitrate;
+		int MaxBitrate;
+		bool bGotMinBitrate;
+		bool bGotMaxBitrate;
+
+		bGotMinBitrate = Json->TryGetNumberField(TEXT("minBitrate"), MinBitrate);
+		bGotMaxBitrate = Json->TryGetNumberField(TEXT("maxBitrate"), MaxBitrate);
+
+		if (bGotMinBitrate && bGotMaxBitrate && MinBitrate > 0 && MaxBitrate > 0)
+		{
+			Observer->OnPlayerRequestsBitrate(PlayerId, MinBitrate, MaxBitrate);
+		}
+
+		Observer->OnSignallingSessionDescription(PlayerId, Type, Sdp);
 	}
 	else
 	{
-		Observer->OnSignallingSessionDescription(PlayerId, Type, Sdp);
+		Observer->OnSignallingSessionDescription(Type, Sdp);
 	}
 }
 
diff --git a/Engine/Plugins/Media/PixelStreaming/Source/PixelStreaming/Private/Settings.cpp b/Engine/Plugins/Media/PixelStreaming/Source/PixelStreaming/Private/Settings.cpp
--- a/Engine/Plugins/Media/PixelStreaming/Source/PixelStreaming/Private/Settings.cpp
+++ b/Engine/Plugins/Media/PixelStreaming/Source/PixelStreaming/Private/Settings.cpp
@@ -355,6 +355,14 @@ namespace UE::PixelStreaming::Settings
 			Streamer->SetStreamFPS(Var->GetInt());
 		});
 	}
+
+	void OnWebRTCBitrateRangeChanged(IConsoleVariable* Var)
+	{
+		IPixelStreamingModule::Get().ForEachStreamer([](TSharedPtr<IPixelStreamingStreamer> Streamer) {
+			Streamer->RefreshStreamBitrate();
+		});
+	}
+
 	// Ends Pixel Streaming Plugin CVars
 
 	// Begin utility functions etc.
@@ -661,6 +669,8 @@ namespace UE::PixelStreaming::Settings
 
 		CVarPixelStreamingOnScreenStats.AsVariable()->SetOnChangedCallback(FConsoleVariableDelegate::CreateStatic(&OnHudStatsToggled));
 		CVarPixelStreamingWebRTCFps.AsVariable()->SetOnChangedCallback(FConsoleVariableDelegate::CreateStatic(&OnWebRTCFpsChanged));
+		CVarPixelStreamingWebRTCMinBitrate.AsVariable()->SetOnChangedCallback(FConsoleVariableDelegate::CreateStatic(&OnWebRTCBitrateRangeChanged));
+		CVarPixelStreamingWebRTCMaxBitrate.AsVariable()->SetOnChangedCallback(FConsoleVariableDelegate::CreateStatic(&OnWebRTCBitrateRangeChanged));
 		CVarPixelStreamingEncoderKeyframeInterval.AsVariable()->SetOnChangedCallback(FConsoleVariableDelegate::CreateStatic(&OnKeyframeIntervalChanged));
 
 		// Values parse from commands line
diff --git a/Engine/Plugins/Media/PixelStreaming/Source/PixelStreaming/Private/SignallingConnectionObserver.cpp b/Engine/Plugins/Media/PixelStreaming/Source/PixelStreaming/Private/SignallingConnectionObserver.cpp
--- a/Engine/Plugins/Media/PixelStreaming/Source/PixelStreaming/Private/SignallingConnectionObserver.cpp
+++ b/Engine/Plugins/Media/PixelStreaming/Source/PixelStreaming/Private/SignallingConnectionObserver.cpp
@@ -86,6 +86,11 @@ namespace UE::PixelStreaming
 		Streamer.AddNewDataChannel(PlayerId, NewChannel);
 	}
 
+	void FPixelStreamingSignallingConnectionObserver::OnPlayerRequestsBitrate(FPixelStreamingPlayerId PlayerId, int MinBitrate, int MaxBitrate)
+	{
+		Streamer.PlayerRequestsBitrate(PlayerId, MinBitrate, MaxBitrate);
+	}
+
 	// These are player only and will only be relevant when on the receiving side of pixel streaming, such as the player plugin.
 	void FPixelStreamingSignallingConnectionObserver::OnSignallingStreamerList(const TArray<FString>& StreamerList)
 	{
diff --git a/Engine/Plugins/Media/PixelStreaming/Source/PixelStreaming/Private/SignallingConnectionObserver.h b/Engine/Plugins/Media/PixelStreaming/Source/PixelStreaming/Private/SignallingConnectionObserver.h
--- a/Engine/Plugins/Media/PixelStreaming/Source/PixelStreaming/Private/SignallingConnectionObserver.h
+++ b/Engine/Plugins/Media/PixelStreaming/Source/PixelStreaming/Private/SignallingConnectionObserver.h
@@ -24,6 +24,7 @@ namespace UE::PixelStreaming
 		virtual void OnSignallingPlayerConnected(FPixelStreamingPlayerId PlayerId, const FPixelStreamingPlayerConfig& PlayerConfig, bool bSendOffer) override;
 		virtual void OnSignallingPlayerDisconnected(FPixelStreamingPlayerId PlayerId) override;
 		virtual void OnSignallingSFUPeerDataChannels(FPixelStreamingPlayerId SFUId, FPixelStreamingPlayerId PlayerId, int32 SendStreamId, int32 RecvStreamId) override;
+		virtual void OnPlayerRequestsBitrate(FPixelStreamingPlayerId PlayerId, int MinBitrate, int MaxBitrate) override;
 
 		// Player-only
 		virtual void OnSignallingStreamerList(const TArray<FString>& StreamerList) override;
diff --git a/Engine/Plugins/Media/PixelStreaming/Source/PixelStreaming/Private/Streamer.cpp b/Engine/Plugins/Media/PixelStreaming/Source/PixelStreaming/Private/Streamer.cpp
--- a/Engine/Plugins/Media/PixelStreaming/Source/PixelStreaming/Private/Streamer.cpp
+++ b/Engine/Plugins/Media/PixelStreaming/Source/PixelStreaming/Private/Streamer.cpp
@@ -481,6 +481,29 @@ namespace UE::PixelStreaming
 		}
 	}
 
+	void FStreamer::PlayerRequestsBitrate(FPixelStreamingPlayerId PlayerId, int MinBitrate, int MaxBitrate)
+	{
+		if (QualityControllingId == INVALID_PLAYER_ID || PlayerId == QualityControllingId)
+		{
+			Settings::CVarPixelStreamingWebRTCMinBitrate->Set(MinBitrate, ECVF_SetByCode);
+			Settings::CVarPixelStreamingWebRTCMaxBitrate->Set(MaxBitrate, ECVF_SetByCode);
+		}
+		else
+		{
+			UE_LOG(LogPixelStreaming, Warning, TEXT("Unable to set bitrates from player %s since they are not the Quality Controller."), *PlayerId);
+		}
+	}
+
+	void FStreamer::RefreshStreamBitrate()
+	{
+		Players.Apply([this](FPixelStreamingPlayerId PlayerId, FPlayerContext& PlayerContext) {
+			if (PlayerContext.PeerConnection)
+			{
+				PlayerContext.PeerConnection->RefreshStreamBitrate();
+			}
+		});
+	}
+
 	bool FStreamer::CreateSession(FPixelStreamingPlayerId PlayerId)
 	{
 		if (FPlayerContext* PlayerContext = Players.Find(PlayerId))
diff --git a/Engine/Plugins/Media/PixelStreaming/Source/PixelStreaming/Private/Streamer.h b/Engine/Plugins/Media/PixelStreaming/Source/PixelStreaming/Private/Streamer.h
--- a/Engine/Plugins/Media/PixelStreaming/Source/PixelStreaming/Private/Streamer.h
+++ b/Engine/Plugins/Media/PixelStreaming/Source/PixelStreaming/Private/Streamer.h
@@ -82,6 +82,10 @@ namespace UE::PixelStreaming
 		// TODO(Luke) hook this back up so that the Engine can change how the interface is working browser side
 		void AddPlayerConfig(TSharedRef<FJsonObject>& JsonObject);
 
+		virtual void PlayerRequestsBitrate(FPixelStreamingPlayerId PlayerId, int MinBitrate, int MaxBitrate) override;
+
+		virtual void RefreshStreamBitrate() override;
+
 	private:
 		FStreamer(const FString& StreamerId);
 
diff --git a/Engine/Plugins/Media/PixelStreaming/Source/PixelStreaming/Private/Tests/MockPlayer.h b/Engine/Plugins/Media/PixelStreaming/Source/PixelStreaming/Private/Tests/MockPlayer.h
--- a/Engine/Plugins/Media/PixelStreaming/Source/PixelStreaming/Private/Tests/MockPlayer.h
+++ b/Engine/Plugins/Media/PixelStreaming/Source/PixelStreaming/Private/Tests/MockPlayer.h
@@ -201,6 +201,7 @@ namespace UE::PixelStreaming
 		virtual void OnSignallingPlayerDisconnected(FPixelStreamingPlayerId PlayerId) override {}
 		virtual void OnSignallingSFUPeerDataChannels(FPixelStreamingPlayerId SFUId, FPixelStreamingPlayerId PlayerId, int32 SendStreamId, int32 RecvStreamId) override {}
 		virtual void OnSignallingStreamerList(const TArray<FString>& StreamerList) override {}
+		virtual void OnPlayerRequestsBitrate(FPixelStreamingPlayerId PlayerId, int MinBitrate, int MaxBitrate) override {}
 
 		DECLARE_MULTICAST_DELEGATE(FOnConnectionEstablished);
 		FOnConnectionEstablished OnConnectionEstablished;
diff --git a/Engine/Plugins/Media/PixelStreaming/Source/PixelStreaming/Private/Tests/SignallingServerConnectionTests.cpp b/Engine/Plugins/Media/PixelStreaming/Source/PixelStreaming/Private/Tests/SignallingServerConnectionTests.cpp
--- a/Engine/Plugins/Media/PixelStreaming/Source/PixelStreaming/Private/Tests/SignallingServerConnectionTests.cpp
+++ b/Engine/Plugins/Media/PixelStreaming/Source/PixelStreaming/Private/Tests/SignallingServerConnectionTests.cpp
@@ -39,6 +39,7 @@ namespace
 		virtual void OnSignallingPlayerConnected(FPixelStreamingPlayerId, const FPixelStreamingPlayerConfig&, bool bSendOffer) override {}
 		virtual void OnSignallingPlayerDisconnected(FPixelStreamingPlayerId PlayerId) override {}
 		virtual void OnSignallingSFUPeerDataChannels(FPixelStreamingPlayerId SFUId, FPixelStreamingPlayerId PlayerId, int32 SendStreamId, int32 RecvStreamId) override {}
+		virtual void OnPlayerRequestsBitrate(FPixelStreamingPlayerId PlayerId, int MinBitrate, int MaxBitrate) override {}
 
 		//// Player-only
 		virtual void OnSignallingSessionDescription(webrtc::SdpType Type, const FString& Sdp) override
diff --git a/Engine/Plugins/Media/PixelStreaming/Source/PixelStreaming/Public/IPixelStreamingSignallingConnectionObserver.h b/Engine/Plugins/Media/PixelStreaming/Source/PixelStreaming/Public/IPixelStreamingSignallingConnectionObserver.h
--- a/Engine/Plugins/Media/PixelStreaming/Source/PixelStreaming/Public/IPixelStreamingSignallingConnectionObserver.h
+++ b/Engine/Plugins/Media/PixelStreaming/Source/PixelStreaming/Public/IPixelStreamingSignallingConnectionObserver.h
@@ -23,6 +23,7 @@ public:
 	virtual void OnSignallingPlayerConnected(FPixelStreamingPlayerId PlayerId, const FPixelStreamingPlayerConfig& PlayerConfig, bool bSendOffer) = 0;
 	virtual void OnSignallingPlayerDisconnected(FPixelStreamingPlayerId PlayerId) = 0;
 	virtual void OnSignallingSFUPeerDataChannels(FPixelStreamingPlayerId SFUId, FPixelStreamingPlayerId PlayerId, int32 SendStreamId, int32 RecvStreamId) = 0;
+	virtual void OnPlayerRequestsBitrate(FPixelStreamingPlayerId PlayerId, int MinBitrate, int MaxBitrate) = 0;
 
 	// Player-only
 	virtual void OnSignallingStreamerList(const TArray<FString>& StreamerList) = 0;
diff --git a/Engine/Plugins/Media/PixelStreaming/Source/PixelStreaming/Public/IPixelStreamingStreamer.h b/Engine/Plugins/Media/PixelStreaming/Source/PixelStreaming/Public/IPixelStreamingStreamer.h
--- a/Engine/Plugins/Media/PixelStreaming/Source/PixelStreaming/Public/IPixelStreamingStreamer.h
+++ b/Engine/Plugins/Media/PixelStreaming/Source/PixelStreaming/Public/IPixelStreamingStreamer.h
@@ -304,4 +304,10 @@ public:
 	 */
 	virtual void SetConfigOption(const FName& OptionName, const FString& Value) = 0;
 	virtual bool GetConfigOption(const FName& OptionName, FString& OutValue) = 0;
+
+	/**
+	*/
+	virtual void PlayerRequestsBitrate(FPixelStreamingPlayerId PlayerId, int MinBitrate, int MaxBitrate) = 0;
+
+	virtual void RefreshStreamBitrate() = 0;
 };
diff --git a/Engine/Plugins/Media/PixelStreaming/Source/PixelStreaming/Public/PixelStreamingPeerConnection.h b/Engine/Plugins/Media/PixelStreaming/Source/PixelStreaming/Public/PixelStreamingPeerConnection.h
--- a/Engine/Plugins/Media/PixelStreaming/Source/PixelStreaming/Public/PixelStreamingPeerConnection.h
+++ b/Engine/Plugins/Media/PixelStreaming/Source/PixelStreaming/Public/PixelStreamingPeerConnection.h
@@ -201,6 +201,8 @@ public:
 	 */
 	void SetWebRTCStatsCallback(rtc::scoped_refptr<webrtc::RTCStatsCollectorCallback> InCallback);
 
+	void RefreshStreamBitrate();
+
 	/**
 	 * @return The negotiated video codec as parsed from the SDP offer/answer.
 	 * Note: This changes as new offer/answers come in potentially.
