Add: Allow users to specify port range for WebRTC as well as flag for the WebRTC port allocator.

diff --git a/Engine/Plugins/Media/PixelStreaming/Source/PixelStreaming/Private/PixelStreamingSignallingConnection.cpp b/Engine/Plugins/Media/PixelStreaming/Source/PixelStreaming/Private/PixelStreamingSignallingConnection.cpp
--- a/Engine/Plugins/Media/PixelStreaming/Source/PixelStreaming/Private/PixelStreamingSignallingConnection.cpp
+++ b/Engine/Plugins/Media/PixelStreaming/Source/PixelStreaming/Private/PixelStreamingSignallingConnection.cpp
@@ -520,6 +520,42 @@ void FPixelStreamingSignallingConnection::OnConfig(const FJsonObjectPtr& Json)
 	// force `UnifiedPlan` as we control both ends of WebRTC streaming
 	RTCConfig.sdp_semantics = webrtc::SdpSemantics::kUnifiedPlan;
 
+#if WEBRTC_5414
+	int MinPort, MaxPort;
+	MinPort = UE::PixelStreaming::Settings::CVarPixelStreamingWebRTCMinPort.GetValueOnAnyThread();
+	MaxPort = UE::PixelStreaming::Settings::CVarPixelStreamingWebRTCMaxPort.GetValueOnAnyThread();
+
+	if(MinPort < 0 || MinPort > 65535)
+	{
+		UE_LOG(LogPixelStreamingSS, Warning, TEXT("Invalid PixelStreaming.WebRTC.MinPort specified. Value must be within 0 to 65535 inclusive"));
+		MinPort = 49152;
+		UE::PixelStreaming::Settings::CVarPixelStreamingWebRTCMinPort->Set(MinPort, ECVF_SetByCode);
+	}
+
+	if(MaxPort < 0 || MaxPort > 65535)
+	{
+		UE_LOG(LogPixelStreamingSS, Warning, TEXT("Invalid PixelStreaming.WebRTC.MaxPort specified. Value must be within 0 to 65535 inclusive"));
+		MaxPort = 65535;
+		UE::PixelStreaming::Settings::CVarPixelStreamingWebRTCMaxPort->Set(MaxPort, ECVF_SetByCode);
+	}
+
+	if (MinPort > MaxPort)
+	{
+		int OldMax = MaxPort;
+		MaxPort = MinPort;
+		MinPort = OldMax;
+
+		// To try to not be misleading with debug texts etc, we reset these sanitised settings here
+		UE::PixelStreaming::Settings::CVarPixelStreamingWebRTCMinPort->Set(MinPort, ECVF_SetByCode);
+		UE::PixelStreaming::Settings::CVarPixelStreamingWebRTCMaxPort->Set(MaxPort, ECVF_SetByCode);
+	}
+
+	RTCConfig.set_min_port(MinPort);
+	RTCConfig.set_max_port(MaxPort);
+
+	RTCConfig.set_port_allocator_flags(UE::PixelStreaming::Settings::PortAllocatorParameters);
+#endif
+
 	Observer->OnSignallingConfig(RTCConfig);
 }
 
diff --git a/Engine/Plugins/Media/PixelStreaming/Source/PixelStreaming/Private/Settings.cpp b/Engine/Plugins/Media/PixelStreaming/Source/PixelStreaming/Private/Settings.cpp
--- a/Engine/Plugins/Media/PixelStreaming/Source/PixelStreaming/Private/Settings.cpp
+++ b/Engine/Plugins/Media/PixelStreaming/Source/PixelStreaming/Private/Settings.cpp
@@ -221,6 +221,24 @@ namespace UE::PixelStreaming::Settings
 		TEXT("Sets the WebRTC field trials string. Format:\"TRIAL1/VALUE1/TRIAL2/VALUE2/\""),
 		ECVF_Default);
 
+	TAutoConsoleVariable<int> CVarPixelStreamingWebRTCMinPort(
+		TEXT("PixelStreaming.WebRTC.MinPort"),
+		49152, // Default according to RFC5766
+		TEXT("Sets the minimum usable port for the WebRTC port allocator. Default: 49152"),
+		ECVF_Default);
+
+	TAutoConsoleVariable<int> CVarPixelStreamingWebRTCMaxPort(
+		TEXT("PixelStreaming.WebRTC.MaxPort"),
+		65535, // Default according to RFC5766
+		TEXT("Sets the maximum usable port for the WebRTC port allocator. Default: 65535"),
+		ECVF_Default);
+
+	TAutoConsoleVariable<FString> CVarPixelStreamingWebRTCPortAllocatorFlags(
+		TEXT("PixelStreaming.WebRTC.PortAllocator.Flags"),
+		TEXT(""),
+		TEXT("Sets the WebRTC port allocator flags. Format:\"DISABLE_UDP,DISABLE_STUN,...\""),
+		ECVF_Default);
+
 	// End WebRTC CVars
 
 	// Begin Pixel Streaming Plugin CVars
@@ -649,6 +667,86 @@ namespace UE::PixelStreaming::Settings
 		}
 	}
 
+	uint32 PortAllocatorParameters;
+
+	void OnPortAllocatorParametersChanged(IConsoleVariable* Var)
+	{
+		PortAllocatorParameters = 0;
+
+		FString StringOptions = Var->GetString();
+		if(StringOptions.IsEmpty())
+		{
+			return;
+		}
+
+		TArray<FString> FlagArray;
+		StringOptions.ParseIntoArray(FlagArray, TEXT(","), true);
+		int OptionCount = FlagArray.Num();
+		while (OptionCount > 0)
+		{
+			FString Flag = FlagArray[OptionCount - 1];	
+
+			// Flags must match what's in Engine\Source\ThirdParty\WebRTC\xxxx\Include\p2p\base\port_allocator.h
+			if(Flag == "DISABLE_UDP")
+			{
+				PortAllocatorParameters |= 0x01;
+			}
+			else if(Flag == "DISABLE_STUN")
+			{
+				PortAllocatorParameters |= 0x02;
+			}
+			else if(Flag == "DISABLE_RELAY")
+			{
+				PortAllocatorParameters |= 0x04;
+			}
+			else if(Flag == "DISABLE_TCP")
+			{
+				PortAllocatorParameters |= 0x08;
+			}
+			else if(Flag == "ENABLE_IPV6")
+			{
+				PortAllocatorParameters |= 0x40;
+			}
+			else if(Flag == "ENABLE_SHARED_SOCKET")
+			{
+				PortAllocatorParameters |= 0x100;
+			}
+			else if(Flag == "ENABLE_STUN_RETRANSMIT_ATTRIBUTE")
+			{
+				PortAllocatorParameters |= 0x200;
+			}
+			else if(Flag == "DISABLE_ADAPTER_ENUMERATION")
+			{
+				PortAllocatorParameters |= 0x400;
+			}
+			else if(Flag == "DISABLE_DEFAULT_LOCAL_CANDIDATE")
+			{
+				PortAllocatorParameters |= 0x800;
+			}
+			else if(Flag == "DISABLE_UDP_RELAY")
+			{
+				PortAllocatorParameters |= 0x1000;
+			}
+			else if(Flag == "ENABLE_IPV6_ON_WIFI")
+			{
+				PortAllocatorParameters |= 0x4000;
+			}
+			else if(Flag == "ENABLE_ANY_ADDRESS_PORTS")
+			{
+				PortAllocatorParameters |= 0x8000;
+			}
+			else if(Flag == "DISABLE_LINK_LOCAL_NETWORKS")
+			{
+				PortAllocatorParameters |= 0x10000;
+			}
+			else
+			{
+				UE_LOG(LogPixelStreaming, Warning, TEXT("Unknown port allocator flag: %s"), *Flag);
+			}
+			OptionCount--;
+		}
+	}
+
 	/*
 	 * Settings parsing and initialization.
 	 */
@@ -672,6 +770,7 @@ namespace UE::PixelStreaming::Settings
 		CVarPixelStreamingWebRTCMinBitrate.AsVariable()->SetOnChangedCallback(FConsoleVariableDelegate::CreateStatic(&OnWebRTCBitrateRangeChanged));
 		CVarPixelStreamingWebRTCMaxBitrate.AsVariable()->SetOnChangedCallback(FConsoleVariableDelegate::CreateStatic(&OnWebRTCBitrateRangeChanged));
 		CVarPixelStreamingEncoderKeyframeInterval.AsVariable()->SetOnChangedCallback(FConsoleVariableDelegate::CreateStatic(&OnKeyframeIntervalChanged));
+		CVarPixelStreamingWebRTCPortAllocatorFlags.AsVariable()->SetOnChangedCallback(FConsoleVariableDelegate::CreateStatic(&OnPortAllocatorParametersChanged));
 
 		// Values parse from commands line
 		CommandLineParseValue(TEXT("PixelStreamingEncoderKeyframeInterval="), CVarPixelStreamingEncoderKeyframeInterval);
@@ -701,6 +800,9 @@ namespace UE::PixelStreaming::Settings
 		CommandLineParseValue(TEXT("PixelStreamingInputController="), CVarPixelStreamingInputController);
 		CommandLineParseValue(TEXT("PixelStreamingSignalingReconnectInterval="), CVarPixelStreamingSignalingReconnectInterval);
 		CommandLineParseValue(TEXT("PixelStreamingDecoupleWaitFactor="), CVarPixelStreamingDecoupleWaitFactor);
+		CommandLineParseValue(TEXT("PixelStreamingWebRTCMinPort="), CVarPixelStreamingWebRTCMinPort);
+		CommandLineParseValue(TEXT("PixelStreamingWebRTCMaxPort="), CVarPixelStreamingWebRTCMaxPort);
+		CommandLineParseValue(TEXT("PixelStreamingWebRTCPortAllocatorFlags="), CVarPixelStreamingWebRTCPortAllocatorFlags);
 
 		// Options parse (if these exist they are set to true)
 		CommandLineParseOption(TEXT("PixelStreamingOnScreenStats"), CVarPixelStreamingOnScreenStats);
diff --git a/Engine/Plugins/Media/PixelStreaming/Source/PixelStreaming/Private/Settings.h b/Engine/Plugins/Media/PixelStreaming/Source/PixelStreaming/Private/Settings.h
--- a/Engine/Plugins/Media/PixelStreaming/Source/PixelStreaming/Private/Settings.h
+++ b/Engine/Plugins/Media/PixelStreaming/Source/PixelStreaming/Private/Settings.h
@@ -54,6 +54,9 @@ namespace UE::PixelStreaming::Settings
 	extern TAutoConsoleVariable<bool> CVarPixelStreamingWebRTCDisableFrameDropper;
 	extern TAutoConsoleVariable<float> CVarPixelStreamingWebRTCVideoPacingFactor;
 	extern TAutoConsoleVariable<float> CVarPixelStreamingWebRTCVideoPacingMaxDelay;
+	extern TAutoConsoleVariable<int> CVarPixelStreamingWebRTCMinPort;
+	extern TAutoConsoleVariable<int> CVarPixelStreamingWebRTCMaxPort;
+	extern TAutoConsoleVariable<FString> CVarPixelStreamingWebRTCPortAllocatorFlags;
 	// End WebRTC CVars
 
 	// Begin Pixel Streaming Plugin CVars
@@ -111,5 +114,6 @@ namespace UE::PixelStreaming::Settings
 	// End utility functions etc.
 
 	extern FPixelStreamingSimulcastParameters SimulcastParameters;
+	extern uint32 PortAllocatorParameters;
 
 } // namespace UE::PixelStreaming::Settings
\ No newline at end of file
