Fix: Prevent PixelStreamingMediaIOCapture from forcing MediaIO cvars.

diff --git a/Engine/Plugins/Media/PixelStreaming/Source/PixelStreaming/Internal/PixelStreamingMediaIOCapture.h b/Engine/Plugins/Media/PixelStreaming/Source/PixelStreaming/Internal/PixelStreamingMediaIOCapture.h
--- a/Engine/Plugins/Media/PixelStreaming/Source/PixelStreaming/Internal/PixelStreamingMediaIOCapture.h
+++ b/Engine/Plugins/Media/PixelStreaming/Source/PixelStreaming/Internal/PixelStreamingMediaIOCapture.h
@@ -55,9 +55,6 @@ public:
 	DECLARE_MULTICAST_DELEGATE(FOnCaptureViewportInitialized);
 	FOnCaptureViewportInitialized OnCaptureViewportInitialized;
 
-private:
-	void ConfigureThreadCaptureMode(bool bForceRenderThread);
-
 private:
 	TWeakPtr<FSceneViewport> SceneViewport;
 	TWeakPtr<FPixelStreamingVideoInput> VideoInput;
diff --git a/Engine/Plugins/Media/PixelStreaming/Source/PixelStreaming/Private/PixelStreamingMediaIOCapture.cpp b/Engine/Plugins/Media/PixelStreaming/Source/PixelStreaming/Private/PixelStreamingMediaIOCapture.cpp
--- a/Engine/Plugins/Media/PixelStreaming/Source/PixelStreaming/Private/PixelStreamingMediaIOCapture.cpp
+++ b/Engine/Plugins/Media/PixelStreaming/Source/PixelStreaming/Private/PixelStreamingMediaIOCapture.cpp
@@ -49,32 +49,11 @@ bool UPixelStreamingMediaIOCapture::InitializeCapture()
 	bViewportResized = false;
 	bDoGPUCopy = true;
 
-	ConfigureThreadCaptureMode(SupportsAnyThreadCapture());
-
 	SetState(EMediaCaptureState::Capturing);
 
 	return true;
 }
 
-void UPixelStreamingMediaIOCapture::ConfigureThreadCaptureMode(bool bForceRenderThread)
-{
-	char ForceRenderThreadBit = bForceRenderThread ? 0 : 1;
-
-	// Whether to wait for resource readback in a separate thread. (Experimental)
-	IConsoleVariable* CVarScheduleAnyThread = IConsoleManager::Get().FindConsoleVariable(TEXT("MediaIO.ScheduleOnAnyThread"));
-	if (CVarScheduleAnyThread)
-	{
-		CVarScheduleAnyThread->Set(ForceRenderThreadBit, EConsoleVariableFlags::ECVF_SetByCode);
-	}
-
-	// Whether to send out frame  in a separate thread. (Experimental)
-	IConsoleVariable* CVarExperimentalScheduling = IConsoleManager::Get().FindConsoleVariable(TEXT("MediaIO.EnableExperimentalScheduling"));
-	if (CVarExperimentalScheduling)
-	{
-		CVarExperimentalScheduling->Set(true, EConsoleVariableFlags::ECVF_SetByCode);
-	}
-}
-
 void UPixelStreamingMediaIOCapture::StopCaptureImpl(bool bAllowPendingFrameToBeProcess)
 {
 	// Todo: Any cleanup on capture stop should happen here.
