Fix: Remove H265 references in preperation for 5414

diff --git a/Engine/Plugins/Media/PixelStreaming/Source/PixelStreaming/Private/PixelStreamingPeerConnection.cpp b/Engine/Plugins/Media/PixelStreaming/Source/PixelStreaming/Private/PixelStreamingPeerConnection.cpp
--- a/Engine/Plugins/Media/PixelStreaming/Source/PixelStreaming/Private/PixelStreamingPeerConnection.cpp
+++ b/Engine/Plugins/Media/PixelStreaming/Source/PixelStreaming/Private/PixelStreamingPeerConnection.cpp
@@ -123,10 +123,12 @@ namespace
 					{
 						return EPixelStreamingCodec::VP9;
 					}
+					/*
 					else if (absl::EqualsIgnoreCase(Codec.name, cricket::kH265CodecName))
 					{
 						return EPixelStreamingCodec::H265;
 					}
+					*/
 					else if (absl::EqualsIgnoreCase(Codec.name, cricket::kH264CodecName))
 					{
 						return EPixelStreamingCodec::H264;
@@ -878,13 +880,12 @@ void FPixelStreamingPeerConnection::RemoveAudioInput(TSharedPtr<IPixelStreamingA
 	AudioMixer->DisconnectInput(StaticCastSharedPtr<FAudioInput>(AudioInput));
 }
 
-
 void InitializeFieldTrials()
 {
 	FString FieldTrials = Settings::CVarPixelStreamingWebRTCFieldTrials.GetValueOnAnyThread();
 
 	// Set the WebRTC-FrameDropper/Disabled/ if the CVar is set
-	if(Settings::CVarPixelStreamingWebRTCDisableFrameDropper.GetValueOnAnyThread())
+	if (Settings::CVarPixelStreamingWebRTCDisableFrameDropper.GetValueOnAnyThread())
 	{
 		FieldTrials += TEXT("WebRTC-FrameDropper/Disabled/");
 	}
@@ -894,16 +895,16 @@ void InitializeFieldTrials()
 		float OutPacingFactor = -1.0f;
 		float OutPacingMaxDelayMs = -1.0f;
 		bool bVideoPacingFieldTrial = Settings::GetVideoPacing(OutPacingFactor, OutPacingMaxDelayMs);
-		if(bVideoPacingFieldTrial)
+		if (bVideoPacingFieldTrial)
 		{
 			FString VideoPacingFieldTrialStr = TEXT("WebRTC-Video-Pacing/");
 			bool bHasPacingFactor = OutPacingFactor >= 0.0f;
-			if(bHasPacingFactor)
+			if (bHasPacingFactor)
 			{
 				VideoPacingFieldTrialStr += FString::Printf(TEXT("factor:%.1f"), OutPacingFactor);
 			}
 			bool bHasMaxDelay = OutPacingMaxDelayMs >= 0.0f;
-			if(bHasMaxDelay)
+			if (bHasMaxDelay)
 			{
 				VideoPacingFieldTrialStr += bHasPacingFactor ? TEXT(",") : TEXT("");
 				VideoPacingFieldTrialStr += FString::Printf(TEXT("max_delay:%.0f"), OutPacingMaxDelayMs);
@@ -913,8 +914,9 @@ void InitializeFieldTrials()
 		}
 	}
 
-	if (!FieldTrials.IsEmpty()) {
-		//Pass the field trials string to WebRTC. String must never be destroyed.
+	if (!FieldTrials.IsEmpty())
+	{
+		// Pass the field trials string to WebRTC. String must never be destroyed.
 		TStringConversion<TStringConvert<TCHAR, ANSICHAR>> Str = StringCast<ANSICHAR>(*FieldTrials);
 		int length = Str.Length() + 1;
 		char* WRTCFieldTrials = (char*)FMemory::SystemMalloc(length);
diff --git a/Engine/Plugins/Media/PixelStreaming/Source/PixelStreaming/Private/VideoDecoderFactory.cpp b/Engine/Plugins/Media/PixelStreaming/Source/PixelStreaming/Private/VideoDecoderFactory.cpp
--- a/Engine/Plugins/Media/PixelStreaming/Source/PixelStreaming/Private/VideoDecoderFactory.cpp
+++ b/Engine/Plugins/Media/PixelStreaming/Source/PixelStreaming/Private/VideoDecoderFactory.cpp
@@ -32,7 +32,7 @@ namespace UE::PixelStreaming
 		Codecs[EPixelStreamingCodec::VP9].push_back(webrtc::SdpVideoFormat(cricket::kVp9CodecName));
 		// Codecs[EPixelStreamingCodec::H264].push_back(CreateH264Format(webrtc::H264Profile::kProfileConstrainedBaseline, webrtc::H264Level::kLevel3_1));
 		// Codecs[EPixelStreamingCodec::H264].push_back(CreateH264Format(webrtc::H264Profile::kProfileBaseline, webrtc::H264Level::kLevel3_1));
-		Codecs[EPixelStreamingCodec::H265].push_back(webrtc::SdpVideoFormat(cricket::kH265CodecName));
+		// Codecs[EPixelStreamingCodec::H265].push_back(webrtc::SdpVideoFormat(cricket::kH265CodecName));
 
 		return Codecs;
 	}
@@ -41,8 +41,8 @@ namespace UE::PixelStreaming
 	 * Adds all the formats of a given codec to a destination list according to a list of supported formats
 	 */
 	void AddSupportedCodecFormats(EPixelStreamingCodec Codec,
-							 	  const TMap<EPixelStreamingCodec, std::vector<webrtc::SdpVideoFormat>>& SupportedFormatsForCodecs,
-							      std::vector<webrtc::SdpVideoFormat>& OutFormats)
+		const TMap<EPixelStreamingCodec, std::vector<webrtc::SdpVideoFormat>>& SupportedFormatsForCodecs,
+		std::vector<webrtc::SdpVideoFormat>& OutFormats)
 	{
 		if (SupportedFormatsForCodecs.Contains(Codec))
 		{
@@ -75,7 +75,7 @@ namespace UE::PixelStreaming
 		{
 			std::vector<webrtc::SdpVideoFormat> TempSupportedFormats;
 			AddSupportedCodecFormats(SelectedCodec, CodecMap, TempSupportedFormats);
-			
+
 			if (TempSupportedFormats.empty())
 			{
 				UE_LOG(LogPixelStreaming, Error, TEXT("Selected codec was not a supported codec, falling back to negotiating codecs..."));
@@ -127,10 +127,12 @@ namespace UE::PixelStreaming
 		{
 			return std::make_unique<VideoDecoderVPX>(9);
 		}
+		/*
 		else if (absl::EqualsIgnoreCase(format.name, cricket::kH265CodecName))
 		{
 			return std::make_unique<VideoDecoderH265>();
 		}
+		*/
 		return std::make_unique<FVideoDecoderStub>();
 	}
 } // namespace UE::PixelStreaming
diff --git a/Engine/Plugins/Media/PixelStreaming/Source/PixelStreaming/Private/VideoEncoderFactorySingleLayer.cpp b/Engine/Plugins/Media/PixelStreaming/Source/PixelStreaming/Private/VideoEncoderFactorySingleLayer.cpp
--- a/Engine/Plugins/Media/PixelStreaming/Source/PixelStreaming/Private/VideoEncoderFactorySingleLayer.cpp
+++ b/Engine/Plugins/Media/PixelStreaming/Source/PixelStreaming/Private/VideoEncoderFactorySingleLayer.cpp
@@ -38,7 +38,7 @@ namespace UE::PixelStreaming
 		Codecs[EPixelStreamingCodec::VP9].push_back(webrtc::SdpVideoFormat(cricket::kVp9CodecName));
 		Codecs[EPixelStreamingCodec::H264].push_back(UE::PixelStreaming::CreateH264Format(webrtc::H264Profile::kProfileConstrainedBaseline, webrtc::H264Level::kLevel3_1));
 		Codecs[EPixelStreamingCodec::H264].push_back(UE::PixelStreaming::CreateH264Format(webrtc::H264Profile::kProfileBaseline, webrtc::H264Level::kLevel3_1));
-		Codecs[EPixelStreamingCodec::H265].push_back(webrtc::SdpVideoFormat(cricket::kH265CodecName));
+		// Codecs[EPixelStreamingCodec::H265].push_back(webrtc::SdpVideoFormat(cricket::kH265CodecName));
 
 		return Codecs;
 	}
@@ -67,11 +67,12 @@ namespace UE::PixelStreaming
 			int32 MaxCVarAllowedSessions = Settings::CVarPixelStreamingEncoderMaxSessions.GetValueOnAnyThread();
 			bool bCanCreateHardwareEncoder = true;
 
-			if(MaxCVarAllowedSessions != -1 && NumEncoderSessions != -1) {
+			if (MaxCVarAllowedSessions != -1 && NumEncoderSessions != -1)
+			{
 				// If our CVar is set and we receive a valid session count
 				bCanCreateHardwareEncoder &= NumEncoderSessions < MaxCVarAllowedSessions;
-			} 
-			else if(MaxCVarAllowedSessions == -1)
+			}
+			else if (MaxCVarAllowedSessions == -1)
 			{
 				// If we receive a valid session count and our cvar isn't set
 				bCanCreateHardwareEncoder &= NvmlEncoder::IsEncoderSessionAvailable(0); // TODO we should probably actually figure out the GPU index rather than assume 0
@@ -177,6 +178,7 @@ namespace UE::PixelStreaming
 			FStats::Get()->StoreApplicationStat(FStatData(FName(TEXT("Video Codec - VP9")), 1, 0));
 			return std::make_unique<FVideoEncoderSingleLayerVPX>(9);
 		}
+		/*
 		else if (absl::EqualsIgnoreCase(format.name, cricket::kH265CodecName))
 		{
 			FStats::Get()->StoreApplicationStat(FStatData(FName(TEXT("Video Codec - H265")), 1, 0));
@@ -185,6 +187,7 @@ namespace UE::PixelStreaming
 			ActiveEncoders.Add(VideoEncoder.get());
 			return VideoEncoder;
 		}
+		*/
 		else
 		{
 			// Lock during encoder creation
