Fix: Update get_ps_servers scripts to pull 5.4 infra and work on Mac

diff --git a/Engine/Plugins/Media/PixelStreaming/Resources/WebServers/get_ps_servers.bat b/Engine/Plugins/Media/PixelStreaming/Resources/WebServers/get_ps_servers.bat
--- a/Engine/Plugins/Media/PixelStreaming/Resources/WebServers/get_ps_servers.bat
+++ b/Engine/Plugins/Media/PixelStreaming/Resources/WebServers/get_ps_servers.bat
@@ -84,11 +84,15 @@ if DEFINED UEVersion (
     SET PSInfraTagOrBranch=UE5.3
     SET IsTag=0
   )
+  if "%UEVersion%"=="5.4" (
+    SET PSInfraTagOrBranch=UE5.4
+    SET IsTag=0
+  )
 )
 
 @Rem If no arguments select a specific version, fetch the appropriate default
 if NOT DEFINED PSInfraTagOrBranch (
-    SET PSInfraTagOrBranch=UE5.3
+    SET PSInfraTagOrBranch=UE5.4
     SET IsTag=0
 )
 echo Tag or branch: !PSInfraTagOrBranch!
diff --git a/Engine/Plugins/Media/PixelStreaming/Resources/WebServers/get_ps_servers.sh b/Engine/Plugins/Media/PixelStreaming/Resources/WebServers/get_ps_servers.sh
--- a/Engine/Plugins/Media/PixelStreaming/Resources/WebServers/get_ps_servers.sh
+++ b/Engine/Plugins/Media/PixelStreaming/Resources/WebServers/get_ps_servers.sh
@@ -1,7 +1,7 @@
 #!/bin/bash
 # Copyright Epic Games, Inc. All Rights Reserved.
 
-BASH_LOCATION=$( cd -- "$( dirname -- "${BASH_SOURCE[0]}" )" &> /dev/null && pwd )
+BASH_LOCATION="$(cd -P -- "$(dirname -- "$0")" && pwd -P)"
 
 pushd "${BASH_LOCATION}" > /dev/null
 
@@ -26,7 +26,7 @@ print_help() {
 # Set all default variables (e.g. # Name and version of ps-infra that we are downloading)
 PSInfraOrg=EpicGames
 PSInfraRepo=PixelStreamingInfrastructure
-PSInfraTagOrBranch=UE5.3
+PSInfraTagOrBranch=UE5.4
 RefType=heads
 IsTag=0
 ReleaseUrlBase=https://github.com/EpicGames/PixelStreamingInfrastructure/releases/download
@@ -37,7 +37,7 @@ unset FlagPassed
 unset ReleaseVersion
 unset ReleaseUrl
 
-while(($#)) ; do
+while (($#)) ; do
   case "$1" in
    -h ) print_help;;
    -v ) UEVersion="$2"; FlagPassed=1; shift 2;;
@@ -81,12 +81,17 @@ then
     PSInfraTagOrBranch=UE5.3
     IsTag=0
   fi
+  if [ "$UEVersion" = "5.4" ]
+  then
+    PSInfraTagOrBranch=UE5.4
+    IsTag=0
+  fi
 fi
 
 # If no arguments select a specific version, fetch the appropriate default
 if [ -z "$PSInfraTagOrBranch" ]
 then
-  PSInfraTagOrBranch=UE5.3
+  PSInfraTagOrBranch=UE5.4
   IsTag=0
 fi
 echo "Tag or branch: $PSInfraTagOrBranch"
