Add: PixelStreaming on Mac support

diff --git a/Engine/Plugins/Media/PixelStreaming/PixelStreaming.uplugin b/Engine/Plugins/Media/PixelStreaming/PixelStreaming.uplugin
--- a/Engine/Plugins/Media/PixelStreaming/PixelStreaming.uplugin
+++ b/Engine/Plugins/Media/PixelStreaming/PixelStreaming.uplugin
@@ -19,49 +19,49 @@
 			"Name": "PixelStreaming",
 			"Type": "Runtime",
 			"LoadingPhase": "PreDefault",
-			"PlatformAllowList": [ "Win64", "Linux" ],
+			"PlatformAllowList": [ "Win64", "Linux", "Mac" ],
 			"TargetDenyList": [ ]
 		},
 		{
 			"Name": "PixelStreamingServers",
 			"Type": "Runtime",
 			"LoadingPhase": "PostConfigInit",
-			"PlatformAllowList": [ "Win64", "Linux" ],
+			"PlatformAllowList": [ "Win64", "Linux", "Mac" ],
 			"TargetDenyList": [ ]
 		},
 		{
             "Name": "PixelStreamingEditor",
             "Type": "Editor",
             "LoadingPhase": "PreDefault",
-            "PlatformAllowList": [ "Win64", "Linux" ],
+            "PlatformAllowList": [ "Win64", "Linux", "Mac" ],
             "TargetDenyList": [ "Server" ]
 		},
 		{
 			"Name": "PixelStreamingBlueprint",
 			"Type": "Runtime",
 			"LoadingPhase": "PreDefault",
-			"PlatformAllowList": [ "Win64", "Linux" ],
+			"PlatformAllowList": [ "Win64", "Linux", "Mac" ],
 			"TargetDenyList": [ "Server" ]
 		},
 		{
 			"Name": "PixelStreamingBlueprintEditor",
 			"Type": "Editor",
 			"LoadingPhase": "PostEngineInit",
-			"PlatformAllowList": [ "Win64", "Linux" ],
+			"PlatformAllowList": [ "Win64", "Linux", "Mac" ],
 			"TargetDenyList": [ "Server" ]
 		},
 		{
 			"Name": "PixelStreamingHMD",
 			"Type": "Runtime",
 			"LoadingPhase": "PreDefault",
-			"PlatformAllowList": [ "Win64", "Linux" ],
+			"PlatformAllowList": [ "Win64", "Linux", "Mac" ],
 			"TargetDenyList": [ "Server" ]
 		},
 		{
 			"Name": "PixelStreamingInput",
 			"Type": "Runtime",
 			"LoadingPhase": "PreDefault",
-			"PlatformAllowList": [ "Win64", "Linux" ],
+			"PlatformAllowList": [ "Win64", "Linux", "Mac" ],
 			"TargetDenyList": [ "Server" ]
 		}
 	],
@@ -73,7 +73,7 @@
 		{
 			"Name": "WmfMedia",
 			"Enabled": true,
-			"PlatformDenyList": [ "Linux" ]
+			"PlatformAllowList": [ "Win64" ]
 		},
 		{
 			"Name": "AVCodecsCore",
@@ -81,12 +81,18 @@
 		},
 		{
 			"Name": "NVCodecs",
-			"Enabled": true
+			"Enabled": true,
+            "PlatformAllowList": [ "Win64", "Linux" ]
 		},
 		{
 			"Name": "AMFCodecs",
 			"Enabled": true,
-			"PlatformDenyList": [ "Linux" ]
+			"PlatformAllowList": [ "Win64" ]
+		},
+        {
+			"Name": "VTCodecs",
+			"Enabled": true,
+			"PlatformAllowList": [ "Mac" ]
 		},
 		{
 			"Name": "WebSocketNetworking",
diff --git a/Engine/Plugins/Media/PixelStreaming/Source/PixelStreaming/PixelStreaming.Build.cs b/Engine/Plugins/Media/PixelStreaming/Source/PixelStreaming/PixelStreaming.Build.cs
--- a/Engine/Plugins/Media/PixelStreaming/Source/PixelStreaming/PixelStreaming.Build.cs
+++ b/Engine/Plugins/Media/PixelStreaming/Source/PixelStreaming/PixelStreaming.Build.cs
@@ -5,125 +5,131 @@ using System.IO;
 
 namespace UnrealBuildTool.Rules
 {
-	public class PixelStreaming : ModuleRules
-	{
-		string PixelStreamingProgramsOutputDirectory = "$(ProjectDir)/Samples/PixelStreaming/WebServers";
-		const string PixelStreamingProgramsDirectory = "../../Engine/Plugins/Media/PixelStreaming/Resources/WebServers";
-
-		private void AddFolder(string Folder)
-		{
-			string DirectoryToAdd = new DirectoryInfo(Path.Combine(PixelStreamingProgramsDirectory, Folder)).FullName;
-
-			if (!Directory.Exists(DirectoryToAdd))
-			{
-				return;
-			}
-
-			List<string> DependenciesToAdd = new List<string>();
-			DependenciesToAdd.AddRange(Directory.GetFiles(DirectoryToAdd, "*.*", SearchOption.AllDirectories));
-
-			List<string> PathsToIgnore = new List<string>();
-			PathsToIgnore.Add(new DirectoryInfo(DirectoryToAdd + "/node_modules").FullName);
-			PathsToIgnore.Add(new DirectoryInfo(DirectoryToAdd + "/implementations/typescript/node_modules").FullName);
-			PathsToIgnore.Add(new DirectoryInfo(DirectoryToAdd + "/implementations/EpicGames/node_modules").FullName);
-			PathsToIgnore.Add(new DirectoryInfo(DirectoryToAdd + "/implementations/react/node_modules").FullName);
-			PathsToIgnore.Add(new DirectoryInfo(DirectoryToAdd + "/library/node_modules").FullName);
-			PathsToIgnore.Add(new DirectoryInfo(DirectoryToAdd + "/ui-library/node_modules").FullName);
-			PathsToIgnore.Add(new DirectoryInfo(DirectoryToAdd + "/logs").FullName);
-			PathsToIgnore.Add(new DirectoryInfo(DirectoryToAdd + "/platform_scripts/cmd/node").FullName);
-			PathsToIgnore.Add(new DirectoryInfo(DirectoryToAdd + "/platform_scripts/cmd/coturn").FullName);
-			PathsToIgnore.Add(new DirectoryInfo(DirectoryToAdd + "/platform_scripts/bash/node").FullName);
-
-			// Go through each file/folder in the WebServers and selectively choose which ones we wish to package.
-			foreach (var DependencySource in DependenciesToAdd)
-			{
-				// We want to skip packaging the above directories as they are all generated by running scripts
-				bool bSkip = false;
-				foreach (var PathToIgnore in PathsToIgnore)
-				{
-					if (DependencySource.StartsWith(PathToIgnore))
-					{
-						bSkip = true;
-						continue;
-					}
-				}
-
-				// If we are not skipping, we want to add this to our packaged output
-				if (!bSkip)
-				{
-					string RelativeSourcePath = Path.GetRelativePath(PixelStreamingProgramsDirectory, DependencySource);
-					RuntimeDependencies.Add(Path.Combine(PixelStreamingProgramsOutputDirectory, RelativeSourcePath), DependencySource, StagedFileType.SystemNonUFS);
-				}
-			}
-		}
-
-		public PixelStreaming(ReadOnlyTargetRules Target) : base(Target)
-		{
-			// This is so for game projects using our public headers don't have to include extra modules they might not know about.
-			PublicDependencyModuleNames.AddRange(new string[]
-			{
-				"ApplicationCore",
-				"InputDevice",
-				"WebRTC",
-				"PixelCapture",
-				"PixelStreamingInput",
-				"MediaIOCore",
-			});
-
-			PrivateDependencyModuleNames.AddRange(new string[]
-			{
-				"Core",
-				"CoreUObject",
-				"Engine",
-				"EngineSettings",
-				"InputCore",
-				"Json",
-				"Renderer",
-				"RenderCore",
-				"RHI",
-				"SignalProcessing",
-				"Slate",
-				"SlateCore",
-				"AudioMixer",
-				"WebRTC",
-				"WebSockets",
-				"Sockets",
-				"MediaUtils",
-				"DeveloperSettings",
-				"AVCodecsCore",
-				"AVCodecsCoreRHI",
-				"PixelCaptureShaders",
-				"PixelStreamingServers",
-				"PixelStreamingHMD",
-				"TraceLog",
-				"HTTP",
-				"NVML"
-			});
-
-			PrivateDefinitions.Add("PIXELSTREAMING_DUMP_ENCODING=0");
-
-			PrivateDependencyModuleNames.Add("VulkanRHI");
-			AddEngineThirdPartyPrivateStaticDependencies(Target, "Vulkan", "CUDA");
-
-			if (Target.IsInPlatformGroup(UnrealPlatformGroup.Windows))
-			{
-				PrivateDependencyModuleNames.Add("D3D11RHI");
-				PrivateDependencyModuleNames.Add("D3D12RHI");
-
-				AddEngineThirdPartyPrivateStaticDependencies(Target, "DX11", "DX12");
-			}
-
-			// When we build a Game target we also package the servers with it as runtime dependencies
-			if (Target.Type == TargetType.Game && Target.ProjectFile != null)
-			{
-				AddFolder("Frontend");
-				AddFolder("SignallingWebServer");
-				AddFolder("Matchmaker");
-				AddFolder("SFU");
-
-				RuntimeDependencies.Add("$(ProjectDir)/Samples/PixelStreaming/WebServers/get_ps_servers.bat", "$(PluginDir)/Resources/WebServers/get_ps_servers.bat", StagedFileType.NonUFS);
-				RuntimeDependencies.Add("$(ProjectDir)/Samples/PixelStreaming/WebServers/get_ps_servers.sh", "$(PluginDir)/Resources/WebServers/get_ps_servers.sh", StagedFileType.NonUFS);
-			}
-		}
-	}
+    public class PixelStreaming : ModuleRules
+    {
+        string PixelStreamingProgramsOutputDirectory = "$(ProjectDir)/Samples/PixelStreaming/WebServers";
+        const string PixelStreamingProgramsDirectory = "../../Engine/Plugins/Media/PixelStreaming/Resources/WebServers";
+
+        private void AddFolder(string Folder)
+        {
+            string DirectoryToAdd = new DirectoryInfo(Path.Combine(PixelStreamingProgramsDirectory, Folder)).FullName;
+
+            if (!Directory.Exists(DirectoryToAdd))
+            {
+                return;
+            }
+
+            List<string> DependenciesToAdd = new List<string>();
+            DependenciesToAdd.AddRange(Directory.GetFiles(DirectoryToAdd, "*.*", SearchOption.AllDirectories));
+
+            List<string> PathsToIgnore = new List<string>();
+            PathsToIgnore.Add(new DirectoryInfo(DirectoryToAdd + "/node_modules").FullName);
+            PathsToIgnore.Add(new DirectoryInfo(DirectoryToAdd + "/implementations/typescript/node_modules").FullName);
+            PathsToIgnore.Add(new DirectoryInfo(DirectoryToAdd + "/implementations/EpicGames/node_modules").FullName);
+            PathsToIgnore.Add(new DirectoryInfo(DirectoryToAdd + "/implementations/react/node_modules").FullName);
+            PathsToIgnore.Add(new DirectoryInfo(DirectoryToAdd + "/library/node_modules").FullName);
+            PathsToIgnore.Add(new DirectoryInfo(DirectoryToAdd + "/ui-library/node_modules").FullName);
+            PathsToIgnore.Add(new DirectoryInfo(DirectoryToAdd + "/logs").FullName);
+            PathsToIgnore.Add(new DirectoryInfo(DirectoryToAdd + "/platform_scripts/cmd/node").FullName);
+            PathsToIgnore.Add(new DirectoryInfo(DirectoryToAdd + "/platform_scripts/cmd/coturn").FullName);
+            PathsToIgnore.Add(new DirectoryInfo(DirectoryToAdd + "/platform_scripts/bash/node").FullName);
+
+            // Go through each file/folder in the WebServers and selectively choose which ones we wish to package.
+            foreach (var DependencySource in DependenciesToAdd)
+            {
+                // We want to skip packaging the above directories as they are all generated by running scripts
+                bool bSkip = false;
+                foreach (var PathToIgnore in PathsToIgnore)
+                {
+                    if (DependencySource.StartsWith(PathToIgnore))
+                    {
+                        bSkip = true;
+                        continue;
+                    }
+                }
+
+                // If we are not skipping, we want to add this to our packaged output
+                if (!bSkip)
+                {
+                    string RelativeSourcePath = Path.GetRelativePath(PixelStreamingProgramsDirectory, DependencySource);
+                    RuntimeDependencies.Add(Path.Combine(PixelStreamingProgramsOutputDirectory, RelativeSourcePath), DependencySource, StagedFileType.SystemNonUFS);
+                }
+            }
+        }
+
+        public PixelStreaming(ReadOnlyTargetRules Target) : base(Target)
+        {
+            // This is so for game projects using our public headers don't have to include extra modules they might not know about.
+            PublicDependencyModuleNames.AddRange(new string[]
+            {
+                "ApplicationCore",
+                "InputDevice",
+                "WebRTC",
+                "PixelCapture",
+                "PixelStreamingInput",
+                "MediaIOCore",
+            });
+
+            PrivateDependencyModuleNames.AddRange(new string[]
+            {
+                "Core",
+                "CoreUObject",
+                "Engine",
+                "EngineSettings",
+                "InputCore",
+                "Json",
+                "Renderer",
+                "RenderCore",
+                "RHI",
+                "SignalProcessing",
+                "Slate",
+                "SlateCore",
+                "AudioMixer",
+                "WebRTC",
+                "WebSockets",
+                "Sockets",
+                "MediaUtils",
+                "DeveloperSettings",
+                "AVCodecsCore",
+                "AVCodecsCoreRHI",
+                "PixelCaptureShaders",
+                "PixelStreamingServers",
+                "PixelStreamingHMD",
+                "TraceLog",
+                "HTTP",
+                "NVML"
+            });
+
+            PrivateDefinitions.Add("PIXELSTREAMING_DUMP_ENCODING=0");
+
+            if (Target.IsInPlatformGroup(UnrealPlatformGroup.Windows) || Target.IsInPlatformGroup(UnrealPlatformGroup.Linux))
+            {
+                PrivateDependencyModuleNames.Add("VulkanRHI");
+                AddEngineThirdPartyPrivateStaticDependencies(Target, "Vulkan", "CUDA");
+            }
+
+            PrivateDefinitions.Add("PIXELSTREAMING_DUMP_ENCODING=0");
+
+
+            if (Target.IsInPlatformGroup(UnrealPlatformGroup.Windows))
+            {
+                PrivateDependencyModuleNames.Add("D3D11RHI");
+                PrivateDependencyModuleNames.Add("D3D12RHI");
+
+                AddEngineThirdPartyPrivateStaticDependencies(Target, "DX11", "DX12");
+            }
+
+            // When we build a Game target we also package the servers with it as runtime dependencies
+            if (Target.Type == TargetType.Game && Target.ProjectFile != null)
+            {
+                AddFolder("Frontend");
+                AddFolder("SignallingWebServer");
+                AddFolder("Matchmaker");
+                AddFolder("SFU");
+
+                RuntimeDependencies.Add("$(ProjectDir)/Samples/PixelStreaming/WebServers/get_ps_servers.bat", "$(PluginDir)/Resources/WebServers/get_ps_servers.bat", StagedFileType.NonUFS);
+                RuntimeDependencies.Add("$(ProjectDir)/Samples/PixelStreaming/WebServers/get_ps_servers.sh", "$(PluginDir)/Resources/WebServers/get_ps_servers.sh", StagedFileType.NonUFS);
+            }
+        }
+    }
 }
diff --git a/Engine/Plugins/Media/PixelStreaming/Source/PixelStreaming/Private/PixelStreamingModule.cpp b/Engine/Plugins/Media/PixelStreaming/Source/PixelStreaming/Private/PixelStreamingModule.cpp
--- a/Engine/Plugins/Media/PixelStreaming/Source/PixelStreaming/Private/PixelStreamingModule.cpp
+++ b/Engine/Plugins/Media/PixelStreaming/Source/PixelStreaming/Private/PixelStreamingModule.cpp
@@ -93,10 +93,10 @@ namespace UE::PixelStreaming
 
 		const ERHIInterfaceType RHIType = GDynamicRHI ? RHIGetInterfaceType() : ERHIInterfaceType::Hidden;
 		// only D3D11/D3D12/Vulkan is supported
-		if (!(RHIType == ERHIInterfaceType::D3D11 || RHIType == ERHIInterfaceType::D3D12 || RHIType == ERHIInterfaceType::Vulkan))
+		if (!(RHIType == ERHIInterfaceType::D3D11 || RHIType == ERHIInterfaceType::D3D12 || RHIType == ERHIInterfaceType::Vulkan || RHIType == ERHIInterfaceType::Metal))
 		{
 #if !WITH_DEV_AUTOMATION_TESTS
-			UE_LOG(LogPixelStreaming, Warning, TEXT("Only D3D11/D3D12/Vulkan Dynamic RHI is supported. Detected %s"), GDynamicRHI != nullptr ? GDynamicRHI->GetName() : TEXT("[null]"));
+			UE_LOG(LogPixelStreaming, Warning, TEXT("Only D3D11/D3D12/Vulkan/Metal Dynamic RHI is supported. Detected %s"), GDynamicRHI != nullptr ? GDynamicRHI->GetName() : TEXT("[null]"));
 #endif
 			return;
 		}
diff --git a/Engine/Plugins/Media/PixelStreaming/Source/PixelStreaming/Private/Streamer.cpp b/Engine/Plugins/Media/PixelStreaming/Source/PixelStreaming/Private/Streamer.cpp
--- a/Engine/Plugins/Media/PixelStreaming/Source/PixelStreaming/Private/Streamer.cpp
+++ b/Engine/Plugins/Media/PixelStreaming/Source/PixelStreaming/Private/Streamer.cpp
@@ -163,19 +163,9 @@ namespace UE::PixelStreaming
 	{
 		// This method is marked as deprecated but still calls the deprecated method on the input handler. As such, we disable
 		// the warnings that arise from using the input handlers method
-#if PLATFORM_WINDOWS
-	#pragma warning(push)
-	#pragma warning(disable : 4996)
-#elif PLATFORM_LINUX
-	#pragma clang diagnostic push
-	#pragma clang diagnostic ignored "-Wdeprecated-declarations"
-#endif
+        PRAGMA_DISABLE_DEPRECATION_WARNINGS
 		return InputHandler->GetTargetScreenSize();
-#if PLATFORM_WINDOWS
-	#pragma warning(pop)
-#elif PLATFORM_LINUX
-	#pragma clang diagnostic pop
-#endif
+        PRAGMA_ENABLE_DEPRECATION_WARNINGS
 	}
 
 	void FStreamer::SetTargetScreenRect(TWeakPtr<FIntRect> InTargetScreenRect)
diff --git a/Engine/Plugins/Media/PixelStreaming/Source/PixelStreaming/Private/VideoEncoderFactorySingleLayer.cpp b/Engine/Plugins/Media/PixelStreaming/Source/PixelStreaming/Private/VideoEncoderFactorySingleLayer.cpp
--- a/Engine/Plugins/Media/PixelStreaming/Source/PixelStreaming/Private/VideoEncoderFactorySingleLayer.cpp
+++ b/Engine/Plugins/Media/PixelStreaming/Source/PixelStreaming/Private/VideoEncoderFactorySingleLayer.cpp
@@ -59,6 +59,7 @@ namespace UE::PixelStreaming
 		static std::vector<webrtc::SdpVideoFormat> SupportedFormats;
 
 		EPixelStreamingCodec SelectedCodec = UE::PixelStreaming::Settings::GetSelectedCodec();
+#if PLATFORM_WINDOWS || PLATFORM_LINUX
 		if ((SelectedCodec == EPixelStreamingCodec::H264 || SelectedCodec == EPixelStreamingCodec::H265) && IsRHIDeviceNVIDIA())
 		{
 			// NOTE (william.belcher): This check will return false if all the encoding sessions are in use, even if the user intends
@@ -97,6 +98,7 @@ namespace UE::PixelStreaming
 				CodecMap.Remove(EPixelStreamingCodec::H265);
 			}
 		}
+#endif // PLATFORM_WINDOWS || PLATFORM_LINUX
 
 		// If we are not negotiating codecs simply return just the one codec that is selected in UE
 		if (!Settings::CVarPixelStreamingWebRTCNegotiateCodecs.GetValueOnAnyThread())
diff --git a/Engine/Plugins/Media/PixelStreaming/Source/PixelStreamingServers/Private/ServerUtils.cpp b/Engine/Plugins/Media/PixelStreaming/Source/PixelStreamingServers/Private/ServerUtils.cpp
--- a/Engine/Plugins/Media/PixelStreaming/Source/PixelStreamingServers/Private/ServerUtils.cpp
+++ b/Engine/Plugins/Media/PixelStreaming/Source/PixelStreamingServers/Private/ServerUtils.cpp
@@ -36,6 +36,9 @@ namespace UE::PixelStreamingServers::Utils
 #elif PLATFORM_LINUX
 			Args = FString::Printf(TEXT(" -- \"%s\" %s --nosudo"), *ExecutableAbsPath, *Args);
 			ExecutableAbsPath = TEXT("/usr/bin/bash");
+#elif PLATFORM_MAC
+            Args = FString::Printf(TEXT(" -- \"%s\" %s --nosudo"), *ExecutableAbsPath, *Args);
+            ExecutableAbsPath = TEXT("/bin/zsh");
 #else
 			UE_LOG(LogPixelStreamingServers, Error, TEXT("Unsupported platform for Pixel Streaming."));
 			return TSharedPtr<FMonitoredProcess>();
@@ -144,7 +147,7 @@ namespace UE::PixelStreamingServers::Utils
 		OutAbsPath = OutAbsPath / ServerDirectoryName / TEXT("platform_scripts");
 #if PLATFORM_WINDOWS
 		OutAbsPath = OutAbsPath / TEXT("cmd") / TEXT("run_local.bat");
-#elif PLATFORM_LINUX
+#elif PLATFORM_LINUX || PLATFORM_MAC
 		OutAbsPath = OutAbsPath / TEXT("bash") / TEXT("run_local.sh");
 #else
 		UE_LOG(LogPixelStreamingServers, Error, TEXT("Unsupported platform for Pixel Streaming scripts."));
@@ -180,7 +183,7 @@ namespace UE::PixelStreamingServers::Utils
 
 #if PLATFORM_WINDOWS
 		OutScriptPath = OutScriptPath / TEXT("get_ps_servers.bat");
-#elif PLATFORM_LINUX
+#elif PLATFORM_LINUX || PLATFORM_MAC
 		OutScriptPath = OutScriptPath / TEXT("get_ps_servers.sh");
 #else
 		UE_LOG(LogPixelStreamingServers, Error, TEXT("Unsupported platform for Pixel Streaming scripts."));
diff --git a/Engine/Plugins/Media/PixelStreaming/Source/PixelStreamingServers/Private/Tests/ServerTests.cpp b/Engine/Plugins/Media/PixelStreaming/Source/PixelStreamingServers/Private/Tests/ServerTests.cpp
--- a/Engine/Plugins/Media/PixelStreaming/Source/PixelStreamingServers/Private/Tests/ServerTests.cpp
+++ b/Engine/Plugins/Media/PixelStreaming/Source/PixelStreamingServers/Private/Tests/ServerTests.cpp
@@ -61,9 +61,12 @@ namespace UE::PixelStreamingServers
 		ServerPath = ServerPath / TEXT("cirrus.exe");
 #elif PLATFORM_LINUX
 		ServerPath = ServerPath / TEXT("cirrus");
+#elif PLATFORM_MAC
+        UE_LOG(LogPixelStreamingServers, Error, TEXT("No cirrus binaries exist for Mac!"));
+        return "";
 #else
-		UE_LOG(LogPixelStreaming, Error, TEXT("Unsupported platform for Pixel Streaming."));
-		return false
+		UE_LOG(LogPixelStreamingServers, Error, TEXT("Unsupported platform for Pixel Streaming."));
+		return "";
 #endif
 
 		ServerPath = FPaths::ConvertRelativePathToFull(ServerPath);
