Fix: Re-enable the possibility of data channel only streams. This can be enabled by launching with -PixelStreamingWebRTCDisableTransmitVideo and -PixelStreamingWebRTCDisableTransmitAudio.

diff --git a/Engine/Plugins/Media/PixelStreaming/Source/PixelStreaming/Private/PixelStreamingPeerConnection.cpp b/Engine/Plugins/Media/PixelStreaming/Source/PixelStreaming/Private/PixelStreamingPeerConnection.cpp
--- a/Engine/Plugins/Media/PixelStreaming/Source/PixelStreaming/Private/PixelStreamingPeerConnection.cpp
+++ b/Engine/Plugins/Media/PixelStreaming/Source/PixelStreaming/Private/PixelStreamingPeerConnection.cpp
@@ -460,6 +460,18 @@ void FPixelStreamingPeerConnection::SetVideoSource(rtc::scoped_refptr<webrtc::Vi
 
 	VideoSource = InVideoSource;
 
+	const bool bTransmitUEVideo = !Settings::CVarPixelStreamingWebRTCDisableTransmitVideo.GetValueOnAnyThread();
+
+	webrtc::RtpTransceiverDirection VideoTransceiverDirection;
+	if (bTransmitUEVideo)
+	{
+		VideoTransceiverDirection = webrtc::RtpTransceiverDirection::kSendOnly;
+	}
+	else
+	{
+		VideoTransceiverDirection = webrtc::RtpTransceiverDirection::kInactive;
+	}
+
 	// Create video track
 	rtc::scoped_refptr<webrtc::VideoTrackInterface> VideoTrack = PeerConnectionFactory->CreateVideoTrack(ToString(VideoTrackLabel), VideoSource.get());
 	VideoTrack->set_enabled(true);
@@ -491,7 +503,7 @@ void FPixelStreamingPeerConnection::SetVideoSource(rtc::scoped_refptr<webrtc::Vi
 			Sender->SetTrack(VideoTrack);
 #endif
 			Sender->SetStreams({ GetVideoStreamID() });
-			SetTransceiverDirection(*Transceiver, webrtc::RtpTransceiverDirection::kSendOnly);
+			SetTransceiverDirection(*Transceiver, VideoTransceiverDirection);
 			webrtc::RtpParameters ExistingParams = Sender->GetParameters();
 			ExistingParams.degradation_preference = Settings::GetDegradationPreference();
 		}
@@ -502,12 +514,12 @@ void FPixelStreamingPeerConnection::SetVideoSource(rtc::scoped_refptr<webrtc::Vi
 	{
 		webrtc::RtpTransceiverInit TransceiverOptions;
 		TransceiverOptions.stream_ids = { GetVideoStreamID() };
-		TransceiverOptions.direction = webrtc::RtpTransceiverDirection::kSendOnly;
+		TransceiverOptions.direction = VideoTransceiverDirection;
 		TransceiverOptions.send_encodings = CreateRTPEncodingParams(IsSFU);
 
 		webrtc::RTCErrorOr<rtc::scoped_refptr<webrtc::RtpTransceiverInterface>> Result = PeerConnection->AddTransceiver(VideoTrack, TransceiverOptions);
 		checkf(Result.ok(), TEXT("Failed to add Video transceiver to PeerConnection. Msg=%s"), *FString(Result.error().message()));
-		SetTransceiverDirection(*Result.value(), webrtc::RtpTransceiverDirection::kSendOnly);
+		SetTransceiverDirection(*Result.value(), VideoTransceiverDirection);
 		webrtc::RtpParameters ExistingParams = Result.value()->sender()->GetParameters();
 		ExistingParams.degradation_preference = Settings::GetDegradationPreference();
 	}
@@ -638,6 +650,19 @@ void FPixelStreamingPeerConnection::SetAudioSink(TSharedPtr<IPixelStreamingAudio
 	}
 }
 
+void FPixelStreamingPeerConnection::ForEachTransceiver(const TFunction<void(rtc::scoped_refptr<webrtc::RtpTransceiverInterface>)>& Func)
+{
+	if(!PeerConnection)
+	{
+		return;
+	}
+
+	for (auto& Transceiver : PeerConnection->GetTransceivers())
+	{
+		Func(Transceiver);
+	}
+}
+
 void FPixelStreamingPeerConnection::AddRemoteIceCandidate(const FString& SdpMid, int SdpMLineIndex, const FString& Sdp)
 {
 	webrtc::SdpParseError Error;
@@ -849,7 +874,7 @@ void FPixelStreamingPeerConnection::CreateSDP(ESDPType SDPType, EReceiveMediaOpt
 	bool voice_activity_detection = false;
 	bool ice_restart = true;
 	bool use_rtp_mux = true;
-
+	
 	webrtc::PeerConnectionInterface::RTCOfferAnswerOptions SDPOption{
 		offer_to_receive_video,
 		offer_to_receive_audio,
diff --git a/Engine/Plugins/Media/PixelStreaming/Source/PixelStreaming/Private/Settings.cpp b/Engine/Plugins/Media/PixelStreaming/Source/PixelStreaming/Private/Settings.cpp
--- a/Engine/Plugins/Media/PixelStreaming/Source/PixelStreaming/Private/Settings.cpp
+++ b/Engine/Plugins/Media/PixelStreaming/Source/PixelStreaming/Private/Settings.cpp
@@ -167,6 +167,12 @@ namespace UE::PixelStreaming::Settings
 		TEXT("Disables transmission of UE audio to the browser."),
 		ECVF_Default);
 
+	TAutoConsoleVariable<bool> CVarPixelStreamingWebRTCDisableTransmitVideo(
+		TEXT("PixelStreaming.WebRTC.DisableTransmitVideo"),
+		false,
+		TEXT("Disables transmission of UE video to the browser."),
+		ECVF_Default);
+
 	TAutoConsoleVariable<bool> CVarPixelStreamingWebRTCDisableAudioSync(
 		TEXT("PixelStreaming.WebRTC.DisableAudioSync"),
 		true,
@@ -812,6 +818,7 @@ namespace UE::PixelStreaming::Settings
 		CommandLineParseOption(TEXT("PixelStreamingWebRTCDisableStats"), CVarPixelStreamingWebRTCDisableStats);
 		CommandLineParseOption(TEXT("PixelStreamingWebRTCDisableReceiveAudio"), CVarPixelStreamingWebRTCDisableReceiveAudio);
 		CommandLineParseOption(TEXT("PixelStreamingWebRTCDisableTransmitAudio"), CVarPixelStreamingWebRTCDisableTransmitAudio);
+		CommandLineParseOption(TEXT("PixelStreamingWebRTCDisableTransmitVideo"), CVarPixelStreamingWebRTCDisableTransmitVideo);
 		CommandLineParseOption(TEXT("PixelStreamingWebRTCDisableAudioSync"), CVarPixelStreamingWebRTCDisableAudioSync);
 		CommandLineParseOption(TEXT("PixelStreamingWebRTCDisableFrameDropper"), CVarPixelStreamingWebRTCDisableFrameDropper);
 		CommandLineParseOption(TEXT("PixelStreamingSendPlayerIdAsInteger"), CVarSendPlayerIdAsInteger);
diff --git a/Engine/Plugins/Media/PixelStreaming/Source/PixelStreaming/Private/Settings.h b/Engine/Plugins/Media/PixelStreaming/Source/PixelStreaming/Private/Settings.h
--- a/Engine/Plugins/Media/PixelStreaming/Source/PixelStreaming/Private/Settings.h
+++ b/Engine/Plugins/Media/PixelStreaming/Source/PixelStreaming/Private/Settings.h
@@ -45,6 +45,7 @@ namespace UE::PixelStreaming::Settings
 	extern TAutoConsoleVariable<int> CVarPixelStreamingWebRTCHighQpThreshold;
 	extern TAutoConsoleVariable<bool> CVarPixelStreamingWebRTCDisableReceiveAudio;
 	extern TAutoConsoleVariable<bool> CVarPixelStreamingWebRTCDisableTransmitAudio;
+	extern TAutoConsoleVariable<bool> CVarPixelStreamingWebRTCDisableTransmitVideo;
 	extern TAutoConsoleVariable<bool> CVarPixelStreamingWebRTCDisableAudioSync;
 	extern TAutoConsoleVariable<bool> CVarPixelStreamingWebRTCUseLegacyAudioDevice;
 	extern TAutoConsoleVariable<bool> CVarPixelStreamingWebRTCDisableStats;
diff --git a/Engine/Plugins/Media/PixelStreaming/Source/PixelStreaming/Private/Streamer.cpp b/Engine/Plugins/Media/PixelStreaming/Source/PixelStreaming/Private/Streamer.cpp
--- a/Engine/Plugins/Media/PixelStreaming/Source/PixelStreaming/Private/Streamer.cpp
+++ b/Engine/Plugins/Media/PixelStreaming/Source/PixelStreaming/Private/Streamer.cpp
@@ -499,6 +499,11 @@ namespace UE::PixelStreaming
 		});
 	}
 
+	void FStreamer::ForEachPlayer(const TFunction<void(FPixelStreamingPlayerId, FPlayerContext)>& Func)
+	{
+		Players.Apply(Func);
+	}
+
 	bool FStreamer::CreateSession(FPixelStreamingPlayerId PlayerId)
 	{
 		if (FPlayerContext* PlayerContext = Players.Find(PlayerId))
@@ -556,12 +561,9 @@ namespace UE::PixelStreaming
 			if (FPlayerContext* PlayerContext = Players.Find(PlayerId))
 			{
 				const bool AllowSimulcast = PlayerContext->Config.IsSFU;
+				
 				PlayerContext->PeerConnection->SetVideoSource(VideoSourceGroup->CreateVideoSource([this, PlayerId]() { return ShouldPeerGenerateFrames(PlayerId); }));
-				if (!Settings::CVarPixelStreamingWebRTCDisableTransmitAudio.GetValueOnAnyThread())
-				{
-					PlayerContext->PeerConnection->SetAudioSource(FPixelStreamingPeerConnection::GetApplicationAudioSource());
-				}
-
+				PlayerContext->PeerConnection->SetAudioSource(FPixelStreamingPeerConnection::GetApplicationAudioSource());
 				PlayerContext->PeerConnection->SetAudioSink(MakeShared<FAudioSink>());
 			}
 		}
diff --git a/Engine/Plugins/Media/PixelStreaming/Source/PixelStreaming/Private/Streamer.h b/Engine/Plugins/Media/PixelStreaming/Source/PixelStreaming/Private/Streamer.h
--- a/Engine/Plugins/Media/PixelStreaming/Source/PixelStreaming/Private/Streamer.h
+++ b/Engine/Plugins/Media/PixelStreaming/Source/PixelStreaming/Private/Streamer.h
@@ -86,6 +86,8 @@ namespace UE::PixelStreaming
 
 		virtual void RefreshStreamBitrate() override;
 
+		void ForEachPlayer(const TFunction<void(FPixelStreamingPlayerId, FPlayerContext)>& Func);
+
 	private:
 		FStreamer(const FString& StreamerId);
 
diff --git a/Engine/Plugins/Media/PixelStreaming/Source/PixelStreaming/Private/Tests/AudioDirectionTest.cpp b/Engine/Plugins/Media/PixelStreaming/Source/PixelStreaming/Private/Tests/AudioDirectionTest.cpp
new file mode 100644
--- /dev/null
+++ b/Engine/Plugins/Media/PixelStreaming/Source/PixelStreaming/Private/Tests/AudioDirectionTest.cpp
@@ -0,0 +1,52 @@
+// Copyright Epic Games, Inc. All Rights Reserved.
+
+#if WITH_DEV_AUTOMATION_TESTS
+
+#include "Misc/AutomationTest.h"
+#include "TestUtils.h"
+#include "DirectionUtils.h"
+
+namespace UE::PixelStreaming
+{
+	IMPLEMENT_SIMPLE_AUTOMATION_TEST(FAudioSendRecvTest, "System.Plugins.PixelStreaming.FAudioSendRecvTest", EAutomationTestFlags::EditorContext | EAutomationTestFlags::ClientContext | EAutomationTestFlags::ProductFilter)
+	bool FAudioSendRecvTest::RunTest(const FString& Parameters)
+	{
+        SetDisableTransmitAudio(false);
+		SetDisableReceiveAudio(false);
+		DoDirectionTest(cricket::MediaType::MEDIA_TYPE_AUDIO, webrtc::RtpTransceiverDirection::kSendRecv);
+
+        return true;
+	}
+
+    IMPLEMENT_SIMPLE_AUTOMATION_TEST(FAudioSendOnlyTest, "System.Plugins.PixelStreaming.FAudioSendOnlyTest", EAutomationTestFlags::EditorContext | EAutomationTestFlags::ClientContext | EAutomationTestFlags::ProductFilter)
+	bool FAudioSendOnlyTest::RunTest(const FString& Parameters)
+	{
+        SetDisableTransmitAudio(false);
+		SetDisableReceiveAudio(true);
+		DoDirectionTest(cricket::MediaType::MEDIA_TYPE_AUDIO, webrtc::RtpTransceiverDirection::kSendOnly);
+
+        return true;
+	}
+
+	IMPLEMENT_SIMPLE_AUTOMATION_TEST(FAudioRecvOnlyTest, "System.Plugins.PixelStreaming.FAudioRecvOnlyTest", EAutomationTestFlags::EditorContext | EAutomationTestFlags::ClientContext | EAutomationTestFlags::ProductFilter)
+	bool FAudioRecvOnlyTest::RunTest(const FString& Parameters)
+	{
+        SetDisableTransmitAudio(true);
+		SetDisableReceiveAudio(false);
+		DoDirectionTest(cricket::MediaType::MEDIA_TYPE_AUDIO, webrtc::RtpTransceiverDirection::kRecvOnly);
+
+        return true;
+	}
+
+	IMPLEMENT_SIMPLE_AUTOMATION_TEST(FAudioInactiveTest, "System.Plugins.PixelStreaming.FAudioInactiveTest", EAutomationTestFlags::EditorContext | EAutomationTestFlags::ClientContext | EAutomationTestFlags::ProductFilter)
+	bool FAudioInactiveTest::RunTest(const FString& Parameters)
+	{
+        SetDisableTransmitAudio(true);
+		SetDisableReceiveAudio(true);
+		DoDirectionTest(cricket::MediaType::MEDIA_TYPE_AUDIO, webrtc::RtpTransceiverDirection::kInactive);
+
+        return true;
+	}
+} // namespace UE::PixelStreaming
+
+#endif // WITH_DEV_AUTOMATION_TESTS
\ No newline at end of file
diff --git a/Engine/Plugins/Media/PixelStreaming/Source/PixelStreaming/Private/Tests/CodecUtils.h b/Engine/Plugins/Media/PixelStreaming/Source/PixelStreaming/Private/Tests/CodecUtils.h
--- a/Engine/Plugins/Media/PixelStreaming/Source/PixelStreaming/Private/Tests/CodecUtils.h
+++ b/Engine/Plugins/Media/PixelStreaming/Source/PixelStreaming/Private/Tests/CodecUtils.h
@@ -13,8 +13,8 @@ namespace
 
 	void DoFrameReceiveTest()
 	{
-		int32 StreamerPort = 7564;
-		int32 PlayerPort = 4578;
+		int32 StreamerPort = TestUtils::NextStreamerPort();
+		int32 PlayerPort = TestUtils::NextPlayerPort();
 
 		FMockVideoFrameConfig FrameConfig = { 128 /*Width*/, 128 /*Height*/, 255 /*Y*/, 137 /*U*/, 216 /*V*/ };
 
@@ -41,8 +41,8 @@ namespace
 
 	void DoFrameResizeMultipleTimesTest()
 	{
-		int32 StreamerPort = 7564;
-		int32 PlayerPort = 4578;
+		int32 StreamerPort = TestUtils::NextStreamerPort();
+		int32 PlayerPort = TestUtils::NextPlayerPort();
 
 		TSharedPtr<UE::PixelStreamingServers::IServer> SignallingServer = CreateSignallingServer(StreamerPort, PlayerPort);
 
diff --git a/Engine/Plugins/Media/PixelStreaming/Source/PixelStreaming/Private/Tests/DataChannelTest.cpp b/Engine/Plugins/Media/PixelStreaming/Source/PixelStreaming/Private/Tests/DataChannelTest.cpp
--- a/Engine/Plugins/Media/PixelStreaming/Source/PixelStreaming/Private/Tests/DataChannelTest.cpp
+++ b/Engine/Plugins/Media/PixelStreaming/Source/PixelStreaming/Private/Tests/DataChannelTest.cpp
@@ -20,8 +20,8 @@ namespace UE::PixelStreaming
 		// need to be able to accept codec to handshake otherwise setting local description fails when generating an answer
 		SetCodec(EPixelStreamingCodec::VP8, false /*bUseComputeShader*/);
 
-		int32 StreamerPort = 7567;
-		int32 PlayerPort = 4581;
+		int32 StreamerPort = TestUtils::NextStreamerPort();
+		int32 PlayerPort = TestUtils::NextPlayerPort();
 
 		TSharedPtr<UE::PixelStreamingServers::IServer> SignallingServer = CreateSignallingServer(StreamerPort, PlayerPort);
 
diff --git a/Engine/Plugins/Media/PixelStreaming/Source/PixelStreaming/Private/Tests/DirectionUtils.h b/Engine/Plugins/Media/PixelStreaming/Source/PixelStreaming/Private/Tests/DirectionUtils.h
new file mode 100644
--- /dev/null
+++ b/Engine/Plugins/Media/PixelStreaming/Source/PixelStreaming/Private/Tests/DirectionUtils.h
@@ -0,0 +1,44 @@
+// Copyright Epic Games, Inc. All Rights Reserved.
+
+#pragma once
+
+#include "Misc/AutomationTest.h"
+#include "TestUtils.h"
+
+#if WITH_DEV_AUTOMATION_TESTS
+
+namespace
+{
+	using namespace UE::PixelStreaming;
+
+	void DoDirectionTest(cricket::MediaType ExpectedMediaType, webrtc::RtpTransceiverDirection ExpectedDirection)
+	{
+		int32 StreamerPort = TestUtils::NextStreamerPort();
+		int32 PlayerPort = TestUtils::NextPlayerPort();
+
+		TSharedPtr<UE::PixelStreamingServers::IServer> SignallingServer = CreateSignallingServer(StreamerPort, PlayerPort);
+
+		TSharedPtr<IPixelStreamingStreamer> Streamer = CreateStreamer(StreamerPort);
+		TSharedPtr<FPixelStreamingVideoInputI420> VideoInput = MakeShared<FPixelStreamingVideoInputI420>();
+		VideoInput->AddOutputFormat(PixelCaptureBufferFormat::FORMAT_I420);
+		Streamer->SetVideoInput(VideoInput);
+
+		TSharedPtr<FMockPlayer> Player = CreatePlayer(FMockPlayer::EMode::AcceptOffers);
+		TSharedPtr<FMockVideoSink> VideoSink = MakeShared<FMockVideoSink>();
+		Player->VideoSink = VideoSink;
+
+		Streamer->StartStreaming();
+
+		const TFunction<bool(rtc::scoped_refptr<webrtc::RtpTransceiverInterface>)> CheckFunc = [ExpectedMediaType, ExpectedDirection](rtc::scoped_refptr<webrtc::RtpTransceiverInterface> Transceiver) {
+			return ((Transceiver->media_type() == ExpectedMediaType) && (Transceiver->direction() == ExpectedDirection));
+		};
+
+		ADD_LATENT_AUTOMATION_COMMAND(FConnectPlayerAfterStreamerConnectedOrTimeout(5.0, Streamer, Player, PlayerPort))
+		ADD_LATENT_AUTOMATION_COMMAND(FWaitForPlayerConnectedOrTimeout(5.0, Player, PlayerPort))
+		ADD_LATENT_AUTOMATION_COMMAND(FWaitForICEConnectedOrTimeout(5.0, Player))
+		ADD_LATENT_AUTOMATION_COMMAND(FCheckTransceivers(Streamer, CheckFunc))
+		ADD_LATENT_AUTOMATION_COMMAND(FCleanupAll(SignallingServer, Streamer, Player))
+	}
+} // namespace UE::PixelStreaming
+
+#endif // WITH_DEV_AUTOMATION_TESTS
diff --git a/Engine/Plugins/Media/PixelStreaming/Source/PixelStreaming/Private/Tests/HandshakeTest.cpp b/Engine/Plugins/Media/PixelStreaming/Source/PixelStreaming/Private/Tests/HandshakeTest.cpp
--- a/Engine/Plugins/Media/PixelStreaming/Source/PixelStreaming/Private/Tests/HandshakeTest.cpp
+++ b/Engine/Plugins/Media/PixelStreaming/Source/PixelStreaming/Private/Tests/HandshakeTest.cpp
@@ -14,8 +14,8 @@ namespace UE::PixelStreaming
 		// need to be able to accept codec to handshake otherwise setting local description fails when generating an answer
 		SetCodec(EPixelStreamingCodec::VP8, false /*bUseComputeShader*/);
 
-		int32 StreamerPort = 8866;
-		int32 PlayerPort = 6688;
+		int32 StreamerPort = TestUtils::NextStreamerPort();
+		int32 PlayerPort = TestUtils::NextPlayerPort();
 
 		TSharedPtr<UE::PixelStreamingServers::IServer> OutSignallingServer = bUseLegacySignallingServer ? CreateLegacySignallingServer(StreamerPort, PlayerPort) : CreateSignallingServer(StreamerPort, PlayerPort);
 		TSharedPtr<IPixelStreamingStreamer> OutStreamer = CreateStreamer(StreamerPort);
diff --git a/Engine/Plugins/Media/PixelStreaming/Source/PixelStreaming/Private/Tests/MultipleSignallingConnectionsTest.cpp b/Engine/Plugins/Media/PixelStreaming/Source/PixelStreaming/Private/Tests/MultipleSignallingConnectionsTest.cpp
--- a/Engine/Plugins/Media/PixelStreaming/Source/PixelStreaming/Private/Tests/MultipleSignallingConnectionsTest.cpp
+++ b/Engine/Plugins/Media/PixelStreaming/Source/PixelStreaming/Private/Tests/MultipleSignallingConnectionsTest.cpp
@@ -156,7 +156,7 @@ namespace UE::PixelStreaming
 	{
 		UE_LOG(LogPixelStreaming, Log, TEXT("----------- ConnectAndDisconnectMultipleStreamersEmbeddedCirrus -----------"));
 
-		const int StreamerPort = 8989;
+		int32 StreamerPort = TestUtils::NextStreamerPort();
 		const int HttpPort = 85;
 
 		TSharedPtr<UE::PixelStreamingServers::IServer> SignallingServer = UE::PixelStreamingServers::MakeSignallingServer();
diff --git a/Engine/Plugins/Media/PixelStreaming/Source/PixelStreaming/Private/Tests/ProtocolTest.cpp b/Engine/Plugins/Media/PixelStreaming/Source/PixelStreaming/Private/Tests/ProtocolTest.cpp
--- a/Engine/Plugins/Media/PixelStreaming/Source/PixelStreaming/Private/Tests/ProtocolTest.cpp
+++ b/Engine/Plugins/Media/PixelStreaming/Source/PixelStreaming/Private/Tests/ProtocolTest.cpp
@@ -17,8 +17,8 @@ namespace UE::PixelStreaming
 	IMPLEMENT_SIMPLE_AUTOMATION_TEST(FProtocolTestAddMessage, "System.Plugins.PixelStreaming.FProtocolTestAddMessage", EAutomationTestFlags::EditorContext | EAutomationTestFlags::ClientContext | EAutomationTestFlags::ProductFilter)
 	bool FProtocolTestAddMessage::RunTest(const FString& Parameters)
 	{
-		int32 StreamerPort = 7569;
-		int32 PlayerPort = 4583;
+		int32 StreamerPort = TestUtils::NextStreamerPort();
+		int32 PlayerPort = TestUtils::NextPlayerPort();
 
 		TSharedPtr<UE::PixelStreamingServers::IServer> SignallingServer = CreateSignallingServer(StreamerPort, PlayerPort);
 
@@ -84,8 +84,8 @@ namespace UE::PixelStreaming
 	IMPLEMENT_SIMPLE_AUTOMATION_TEST(FProtocolTestUseCustomMessage, "System.Plugins.PixelStreaming.FProtocolTestUseCustomMessage", EAutomationTestFlags::EditorContext | EAutomationTestFlags::ClientContext | EAutomationTestFlags::ProductFilter)
 	bool FProtocolTestUseCustomMessage::RunTest(const FString& Parameters)
 	{
-		int32 StreamerPort = 7571;
-		int32 PlayerPort = 4585;
+		int32 StreamerPort = TestUtils::NextStreamerPort();
+		int32 PlayerPort = TestUtils::NextPlayerPort();
 
 		TSharedPtr<UE::PixelStreamingServers::IServer> SignallingServer = CreateSignallingServer(StreamerPort, PlayerPort);
 
diff --git a/Engine/Plugins/Media/PixelStreaming/Source/PixelStreaming/Private/Tests/TestUtils.cpp b/Engine/Plugins/Media/PixelStreaming/Source/PixelStreaming/Private/Tests/TestUtils.cpp
--- a/Engine/Plugins/Media/PixelStreaming/Source/PixelStreaming/Private/Tests/TestUtils.cpp
+++ b/Engine/Plugins/Media/PixelStreaming/Source/PixelStreaming/Private/Tests/TestUtils.cpp
@@ -6,11 +6,25 @@
 #include "PixelStreamingPrivate.h"
 #include "IPixelStreamingModule.h"
 #include "PixelCaptureInputFrameI420.h"
+#include "Streamer.h"
 
 #if WITH_DEV_AUTOMATION_TESTS
 
 namespace UE::PixelStreaming
 {
+	int32 TestUtils::NextStreamerPort()
+	{
+		// Start of IANA un-registerable ports (49152 - 65535)
+		static int NextStreamerPort = 49152;
+		return NextStreamerPort++;
+	}
+
+	int32 TestUtils::NextPlayerPort()
+	{
+		// Half of IANA un-registerable ports (49152 - 65535)
+		static int NextPlayerPort = 57344;
+		return NextPlayerPort++;
+	}
 
 	/* ---------- Latent Automation Commands ----------- */
 
@@ -248,6 +262,37 @@ namespace UE::PixelStreaming
 		return *bComplete.Get();
 	}
 
+	bool FCheckTransceivers::Update()
+	{
+		if(!Streamer)
+		{
+			UE_LOG(LogPixelStreaming, Error, TEXT("Streamer is invalid!"));
+			return true;
+		}
+
+        FStreamer* StreamerPtr = (FStreamer*)Streamer.Get();
+
+        bool bTransceiverExists = false;
+        const TFunction<void(rtc::scoped_refptr<webrtc::RtpTransceiverInterface>)> TranceiverFunc = [&bTransceiverExists = bTransceiverExists, this](rtc::scoped_refptr<webrtc::RtpTransceiverInterface> Transceiver) 
+        {
+            bTransceiverExists |= CheckFunc(Transceiver);
+		};
+
+        const TFunction<void(FPixelStreamingPlayerId, FPlayerContext)> PlayerFunc = [&TranceiverFunc = TranceiverFunc, this](FPixelStreamingPlayerId PlayerId, FPlayerContext PlayerContext) 
+        {
+            PlayerContext.PeerConnection->ForEachTransceiver(TranceiverFunc);
+        };
+
+        StreamerPtr->ForEachPlayer(PlayerFunc);
+
+        if(!bTransceiverExists)
+        {
+            UE_LOG(LogPixelStreaming, Error, TEXT("Transceivers don't exist!"));
+        }
+
+		return true;
+	}
+
 	bool FCleanupAll::Update()
 	{
 		if (OutPlayer)
@@ -281,6 +326,30 @@ namespace UE::PixelStreaming
 		});
 	}
 
+	void SetDisableTransmitAudio(bool bDisableTransmitAudio)
+	{
+		// Set whether to transmit audio
+		UE::PixelStreaming::DoOnGameThreadAndWait(MAX_uint32, [bDisableTransmitAudio]() {
+			Settings::CVarPixelStreamingWebRTCDisableTransmitAudio->Set(bDisableTransmitAudio, ECVF_SetByCode);
+		});
+	}
+
+	void SetDisableReceiveAudio(bool bDisableReceiveAudio)
+	{
+		// Set whether to receive audio
+		UE::PixelStreaming::DoOnGameThreadAndWait(MAX_uint32, [bDisableReceiveAudio]() {
+			Settings::CVarPixelStreamingWebRTCDisableReceiveAudio->Set(bDisableReceiveAudio, ECVF_SetByCode);
+		});
+	}
+
+	void SetDisableTransmitVideo(bool bDisableTransmitVideo)
+	{
+		// Set whether to transmit video
+		UE::PixelStreaming::DoOnGameThreadAndWait(MAX_uint32, [bDisableTransmitVideo]() {
+			Settings::CVarPixelStreamingWebRTCDisableTransmitVideo->Set(bDisableTransmitVideo, ECVF_SetByCode);
+		});
+	}
+
 	TSharedPtr<IPixelStreamingStreamer> CreateStreamer(int StreamerPort)
 	{
 		FString StreamerName("Mock Streamer");
@@ -333,7 +402,6 @@ namespace UE::PixelStreaming
 	{
 		return CreateSignallingServer(StreamerPort, PlayerPort, true /*bCreateLegacySignallingServer*/);
 	}
-
 } // namespace UE::PixelStreaming
 
 #endif // WITH_DEV_AUTOMATION_TESTS
diff --git a/Engine/Plugins/Media/PixelStreaming/Source/PixelStreaming/Private/Tests/TestUtils.h b/Engine/Plugins/Media/PixelStreaming/Source/PixelStreaming/Private/Tests/TestUtils.h
--- a/Engine/Plugins/Media/PixelStreaming/Source/PixelStreaming/Private/Tests/TestUtils.h
+++ b/Engine/Plugins/Media/PixelStreaming/Source/PixelStreaming/Private/Tests/TestUtils.h
@@ -12,6 +12,12 @@
 
 namespace UE::PixelStreaming
 {
+	namespace TestUtils
+	{
+		int32 NextStreamerPort();
+		int32 NextPlayerPort();
+	}
+
 	// Equivalent to DEFINE_LATENT_AUTOMATION_COMMAND_FOUR_PARAMETER, but instead we define a custom constructor
 	class FWaitForDataChannelMessageOrTimeout : public IAutomationLatentCommand
 	{
@@ -50,6 +56,7 @@ namespace UE::PixelStreaming
 	DEFINE_LATENT_AUTOMATION_COMMAND_TWO_PARAMETER(FWaitForStreamerConnectedOrTimeout, double, TimeoutSeconds, TSharedPtr<IPixelStreamingStreamer>, OutStreamer);
 	DEFINE_LATENT_AUTOMATION_COMMAND_FOUR_PARAMETER(FConnectPlayerAfterStreamerConnectedOrTimeout, double, TimeoutSeconds, TSharedPtr<IPixelStreamingStreamer>, OutStreamer, TSharedPtr<FMockPlayer>, OutPlayer, int, PlayerPort);
 	DEFINE_LATENT_AUTOMATION_COMMAND_THREE_PARAMETER(FWaitForPlayerConnectedOrTimeout, double, TimeoutSeconds, TSharedPtr<FMockPlayer>, OutPlayer, int, PlayerPort);
+	DEFINE_LATENT_AUTOMATION_COMMAND_TWO_PARAMETER(FCheckTransceivers, TSharedPtr<IPixelStreamingStreamer>, Streamer, const TFunction<bool(rtc::scoped_refptr<webrtc::RtpTransceiverInterface>)>, CheckFunc);
 	DEFINE_LATENT_AUTOMATION_COMMAND_THREE_PARAMETER(FCleanupAll, TSharedPtr<UE::PixelStreamingServers::IServer>, OutSignallingServer, TSharedPtr<IPixelStreamingStreamer>, OutStreamer, TSharedPtr<FMockPlayer>, OutPlayer);
 
 	TSharedPtr<IPixelStreamingStreamer> CreateStreamer(int StreamerPort);
@@ -57,7 +64,9 @@ namespace UE::PixelStreaming
 	TSharedPtr<UE::PixelStreamingServers::IServer> CreateSignallingServer(int StreamerPort, int PlayerPort);
 	TSharedPtr<UE::PixelStreamingServers::IServer> CreateLegacySignallingServer(int StreamerPort, int PlayerPort);
 	void SetCodec(EPixelStreamingCodec Codec, bool bUseComputeShaders);
-
+	void SetDisableTransmitAudio(bool bDisableTransmitAudio);
+	void SetDisableReceiveAudio(bool bDisableReceiveAudio);
+	void SetDisableTransmitVideo(bool bDisableTransmitVideo);
 } // namespace UE::PixelStreaming
 
 #endif // WITH_DEV_AUTOMATION_TESTS
\ No newline at end of file
diff --git a/Engine/Plugins/Media/PixelStreaming/Source/PixelStreaming/Private/Tests/VideoDirectionTest.cpp b/Engine/Plugins/Media/PixelStreaming/Source/PixelStreaming/Private/Tests/VideoDirectionTest.cpp
new file mode 100644
--- /dev/null
+++ b/Engine/Plugins/Media/PixelStreaming/Source/PixelStreaming/Private/Tests/VideoDirectionTest.cpp
@@ -0,0 +1,30 @@
+// Copyright Epic Games, Inc. All Rights Reserved.
+
+#if WITH_DEV_AUTOMATION_TESTS
+
+#include "Misc/AutomationTest.h"
+#include "TestUtils.h"
+#include "DirectionUtils.h"
+
+namespace UE::PixelStreaming
+{
+    IMPLEMENT_SIMPLE_AUTOMATION_TEST(FVideoSendOnlyTest, "System.Plugins.PixelStreaming.FVideoSendOnlyTest", EAutomationTestFlags::EditorContext | EAutomationTestFlags::ClientContext | EAutomationTestFlags::ProductFilter)
+	bool FVideoSendOnlyTest::RunTest(const FString& Parameters)
+	{
+        SetDisableTransmitVideo(false);
+		DoDirectionTest(cricket::MediaType::MEDIA_TYPE_VIDEO, webrtc::RtpTransceiverDirection::kSendOnly);
+
+        return true;
+	}
+
+    IMPLEMENT_SIMPLE_AUTOMATION_TEST(FVideoInactiveTest, "System.Plugins.PixelStreaming.FVideoInactiveTest", EAutomationTestFlags::EditorContext | EAutomationTestFlags::ClientContext | EAutomationTestFlags::ProductFilter)
+	bool FVideoInactiveTest::RunTest(const FString& Parameters)
+	{
+        SetDisableTransmitVideo(true);
+		DoDirectionTest(cricket::MediaType::MEDIA_TYPE_VIDEO, webrtc::RtpTransceiverDirection::kInactive);
+
+        return true;
+	}
+} // namespace UE::PixelStreaming
+
+#endif // WITH_DEV_AUTOMATION_TESTS
\ No newline at end of file
diff --git a/Engine/Plugins/Media/PixelStreaming/Source/PixelStreaming/Public/PixelStreamingPeerConnection.h b/Engine/Plugins/Media/PixelStreaming/Source/PixelStreaming/Public/PixelStreamingPeerConnection.h
--- a/Engine/Plugins/Media/PixelStreaming/Source/PixelStreaming/Public/PixelStreamingPeerConnection.h
+++ b/Engine/Plugins/Media/PixelStreaming/Source/PixelStreaming/Public/PixelStreamingPeerConnection.h
@@ -179,6 +179,13 @@ public:
 	 */
 	TSharedPtr<IPixelStreamingAudioSink> GetAudioSink() { return AudioSink; }
 
+	/**
+	 * @brief A method for iterating through all of the tranceivers on the peer connection. You could use this to check for the existence of an audio/video transceiver
+	 *
+	 * @param Func The lambda to execute with each transceiver
+	 */
+	virtual void ForEachTransceiver(const TFunction<void(rtc::scoped_refptr<webrtc::RtpTransceiverInterface>)>& Func);
+
 	/**
 	 * Adds ICE candidate data to the peer connection. Usually supplied from the signalling server.
 	 * @param SDPMid
