QOL: Remove duplicated code paths in PixelStreamingDataChannel.cpp

diff --git a/Engine/Plugins/Media/PixelStreaming/Source/PixelStreaming/Private/PixelStreamingDataChannel.cpp b/Engine/Plugins/Media/PixelStreaming/Source/PixelStreaming/Private/PixelStreamingDataChannel.cpp
--- a/Engine/Plugins/Media/PixelStreaming/Source/PixelStreaming/Private/PixelStreamingDataChannel.cpp
+++ b/Engine/Plugins/Media/PixelStreaming/Source/PixelStreaming/Private/PixelStreamingDataChannel.cpp
@@ -169,28 +169,17 @@ void FPixelStreamingDataChannel::OnMessage(const webrtc::DataBuffer& Buffer)
 			// the signalling thread or block it with mutexes etc.
 			TWeakPtr<FPixelStreamingDataChannel> WeakChannel = SharedThis;
 
-			if (GFirstFrameIntraFrameDebugging || GIntraFrameDebuggingGameThread)
+			// If we're streaming the editor and we hit a BP breakpoint, the gamethread is no longer able to respond to input
+			// in that case we post this task to the Main Queue as we know that that will still be running
+			ENamedThreads::Type TargetThread = (GFirstFrameIntraFrameDebugging || GIntraFrameDebuggingGameThread) ? ENamedThreads::MainQueue : ENamedThreads::GameThread;
+			AsyncTask(TargetThread, [WeakChannel, Buffer = Buffer]() 
 			{
-				// If we're streaming the editor and we hit a BP breakpoint, the gamethread is no longer able to respond to input
-				// in that case we post this task to the Main Queue as we know that that will still be running
-				AsyncTask(ENamedThreads::MainQueue, [WeakChannel, Buffer = Buffer]() {
-					if (TSharedPtr<FPixelStreamingDataChannel> DataChannel = WeakChannel.Pin())
-					{
-						const uint8 MsgType = static_cast<uint8>(Buffer.data.data()[0]);
-						DataChannel->OnMessageReceived.Broadcast(MsgType, Buffer);
-					}
-				});
-			}
-			else
-			{
-				AsyncTask(ENamedThreads::GameThread, [WeakChannel, Buffer = Buffer]() {
-					if (TSharedPtr<FPixelStreamingDataChannel> DataChannel = WeakChannel.Pin())
-					{
-						const uint8 MsgType = static_cast<uint8>(Buffer.data.data()[0]);
-						DataChannel->OnMessageReceived.Broadcast(MsgType, Buffer);
-					}
-				});
-			}
+				if (TSharedPtr<FPixelStreamingDataChannel> DataChannel = WeakChannel.Pin())
+				{
+					const uint8 MsgType = static_cast<uint8>(Buffer.data.data()[0]);
+					DataChannel->OnMessageReceived.Broadcast(MsgType, Buffer);
+				}
+			});
 		}
 	}
 }
