Fix: Editor can become locked if an editor stream is stopped from the browser

diff --git a/Engine/Plugins/Media/PixelStreaming/Source/PixelStreaming/Private/Streamer.cpp b/Engine/Plugins/Media/PixelStreaming/Source/PixelStreaming/Private/Streamer.cpp
--- a/Engine/Plugins/Media/PixelStreaming/Source/PixelStreaming/Private/Streamer.cpp
+++ b/Engine/Plugins/Media/PixelStreaming/Source/PixelStreaming/Private/Streamer.cpp
@@ -37,6 +37,7 @@
 #include "SignallingConnectionObserver.h"
 #include "ToStringExtensions.h"
 #include "Settings.h"
+#include "CoreGlobals.h"
 
 namespace UE::PixelStreaming
 {
@@ -1137,9 +1138,29 @@ namespace UE::PixelStreaming
 		{
 			// Force a MouseLeave event. This prevents the PixelStreamingApplicationWrapper from
 			// still wrapping the base FSlateApplication after we stop streaming
-			TArray<uint8> EmptyArray;
-			TFunction<void(FString, FMemoryReader)> MouseLeaveHandler = InputHandler->FindMessageHandler("MouseLeave");
-			// MouseLeaveHandler(FMemoryReader(EmptyArray));
+			const auto MouseLeaveFunction = [this]() 
+			{
+				TArray<uint8> EmptyArray;
+				TFunction<void(FString, FMemoryReader)> MouseLeaveHandler = InputHandler->FindMessageHandler("MouseLeave");	
+				MouseLeaveHandler("", FMemoryReader(EmptyArray));
+			};
+
+			if(IsInGameThread())
+			{
+				MouseLeaveFunction();
+			}
+			else
+			{
+				TWeakPtr<FStreamer> WeakStreamer = AsShared();
+				AsyncTask(ENamedThreads::GameThread, [this, MouseLeaveFunction, WeakStreamer]() {
+					if (!WeakStreamer.IsValid())
+					{
+						return;
+					}
+
+					MouseLeaveFunction();
+				});
+			}
 		}
 	}
 } // namespace UE::PixelStreaming
