Fix: Re-enable hardware accelerated PixelStreaming player

diff --git a/Engine/Plugins/Experimental/PixelStreamingPlayer/Source/PixelStreamingPlayer/Private/PixelStreamingMediaTexture.cpp b/Engine/Plugins/Experimental/PixelStreamingPlayer/Source/PixelStreamingPlayer/Private/PixelStreamingMediaTexture.cpp
--- a/Engine/Plugins/Experimental/PixelStreamingPlayer/Source/PixelStreamingPlayer/Private/PixelStreamingMediaTexture.cpp
+++ b/Engine/Plugins/Experimental/PixelStreamingPlayer/Source/PixelStreamingPlayer/Private/PixelStreamingMediaTexture.cpp
@@ -48,75 +48,13 @@ FTextureResource* UPixelStreamingMediaTexture::CreateResource()
 	return CurrentResource;
 }
 
-void UPixelStreamingMediaTexture::OnFrame(const webrtc::VideoFrame& frame)
+void UPixelStreamingMediaTexture::OnFrame(FTextureRHIRef Frame)
 {
-	// TODO
-	// Currently this will cause all input frames to be converted to I420 (using webrtc::ConvertFromI420)
-	// and then render the rgba buffer to a texture, but in some cases we might already have the native
-	// texture and can use it or copy it directly. At some point it would be useful to detect and fork this
-	// behaviour.
-	// Additionally, we're rendering to a texture then updating the texture resource with every frame. We
-	// might be able to just map a single dest texture and update it each frame.
-
-	static constexpr auto TEXTURE_PIXEL_FORMAT = PF_B8G8R8A8;
-	static constexpr auto WEBRTC_PIXEL_FORMAT = webrtc::VideoType::kARGB;
-
-	SizeX = frame.width();
-	SizeY = frame.height();
-	Format = TEXTURE_PIXEL_FORMAT;
-	NumMips = 0;
-	SamplerAddressMode = ESamplerAddressMode::AM_Clamp;
-
-	{
+	AsyncTask(ENamedThreads::ActualRenderingThread, [this, Frame]() {
 		FScopeLock Lock(&RenderSyncContext);
 
-		const uint32_t Size = webrtc::CalcBufferSize(WEBRTC_PIXEL_FORMAT, SizeX, SizeY);
-		if (Size > static_cast<uint32_t>(Buffer.Num()))
-		{
-			Buffer.SetNum(Size);
-		}
-		webrtc::ConvertFromI420(frame, WEBRTC_PIXEL_FORMAT, 0, Buffer.GetData());
-	}
-
-	AsyncTask(ENamedThreads::ActualRenderingThread, [this]() {
-		FScopeLock Lock(&RenderSyncContext);
-
-		const FIntPoint FrameSize = FIntPoint(SizeX, SizeY);
 		FRHICommandListImmediate& RHICmdList = FRHICommandListExecutor::GetImmediateCommandList();
-
-		if (!RenderTargetDescriptor.IsValid() || RenderTargetDescriptor.GetSize() != FIntVector(FrameSize.X, FrameSize.Y, 0))
-		{
-			// Create the RenderTarget descriptor
-			RenderTargetDescriptor = FPooledRenderTargetDesc::Create2DDesc(FrameSize,
-				Format,
-				FClearValueBinding::None,
-				TexCreate_None,
-				TexCreate_RenderTargetable,
-				false);
-
-			// Update the shader resource for the 'SourceTexture'
-			FRHITextureCreateDesc RenderTargetTextureDesc =
-				FRHITextureCreateDesc::Create2D(TEXT(""), FrameSize.X, FrameSize.Y, Format)
-					.SetClearValue(FClearValueBinding(FLinearColor(0.0f, 0.0f, 0.0f)))
-					.SetFlags(ETextureCreateFlags::Dynamic | ETextureCreateFlags::ShaderResource | TexCreate_RenderTargetable | ETextureCreateFlags::SRGB)
-					.SetInitialState(ERHIAccess::SRVMask);
-
-			SourceTexture = RHICreateTexture(RenderTargetTextureDesc);
-
-			// Find a free target-able texture from the render pool
-			GRenderTargetPool.FindFreeElement(RHICmdList,
-				RenderTargetDescriptor,
-				RenderTarget,
-				TEXT("PIXELSTEAMINGPLAYER"));
-		}
-
-		// Create the update region structure
-		const FUpdateTextureRegion2D Region(0, 0, 0, 0, FrameSize.X, FrameSize.Y);
-
-		// Set the Pixel data of the webrtc Frame to the SourceTexture
-		RHIUpdateTexture2D(SourceTexture, 0, Region, FrameSize.X * 4, Buffer.GetData());
-
-		UpdateTextureReference(RHICmdList, SourceTexture);
+		UpdateTextureReference(RHICmdList, Frame);
 	});
 }
 
diff --git a/Engine/Plugins/Experimental/PixelStreamingPlayer/Source/PixelStreamingPlayer/Public/PixelStreamingMediaTexture.h b/Engine/Plugins/Experimental/PixelStreamingPlayer/Source/PixelStreamingPlayer/Public/PixelStreamingMediaTexture.h
--- a/Engine/Plugins/Experimental/PixelStreamingPlayer/Source/PixelStreamingPlayer/Public/PixelStreamingMediaTexture.h
+++ b/Engine/Plugins/Experimental/PixelStreamingPlayer/Source/PixelStreamingPlayer/Public/PixelStreamingMediaTexture.h
@@ -4,7 +4,7 @@
 
 #include "Engine/Texture2DDynamic.h"
 #include "RenderTargetPool.h"
-#include "api/media_stream_interface.h"
+#include "PixelStreamingVideoSink.h"
 #include "PixelStreamingMediaTexture.generated.h"
 
 class FPixelStreamingMediaTextureResource;
@@ -13,7 +13,7 @@ class FPixelStreamingMediaTextureResource;
  * A Texture Object that can be used in materials etc. that takes updates from webrtc frames.
  */
 UCLASS(NotBlueprintType, NotBlueprintable, HideDropdown, HideCategories = (ImportSettings, Compression, Texture, Adjustments, Compositing, LevelOfDetail, Object), META = (DisplayName = "PixelStreaming Media Texture"))
-class PIXELSTREAMINGPLAYER_API UPixelStreamingMediaTexture : public UTexture2DDynamic, public rtc::VideoSinkInterface<webrtc::VideoFrame>
+class PIXELSTREAMINGPLAYER_API UPixelStreamingMediaTexture : public UTexture2DDynamic, public FPixelStreamingVideoSink
 {
 	GENERATED_UCLASS_BODY()
 
@@ -25,8 +25,8 @@ protected:
 	// UTexture implementation
 	virtual FTextureResource* CreateResource() override;
 
-	// from rtc::VideoSinkInterface<webrtc::VideoFrame>
-	virtual void OnFrame(const webrtc::VideoFrame& frame) override;
+	// FPixelStreamingVideoSink implementation
+	virtual void OnFrame(FTextureRHIRef Frame) override;
 
 private:
 	void InitializeResources();
diff --git a/Engine/Plugins/Media/PixelStreaming/Source/PixelStreaming/Private/PixelStreamingVideoSink.cpp b/Engine/Plugins/Media/PixelStreaming/Source/PixelStreaming/Private/PixelStreamingVideoSink.cpp
new file mode 100644
--- /dev/null
+++ b/Engine/Plugins/Media/PixelStreaming/Source/PixelStreaming/Private/PixelStreamingVideoSink.cpp
@@ -0,0 +1,95 @@
+// Copyright Epic Games, Inc. All Rights Reserved.
+
+#include "PixelStreamingVideoSink.h"
+#include "Async/Async.h"
+#include "FrameBufferRHI.h"
+#include "FrameBufferI420.h"
+
+void FPixelStreamingVideoSink::OnFrame(const webrtc::VideoFrame& Frame)
+{
+    static constexpr auto TEXTURE_PIXEL_FORMAT = PF_B8G8R8A8;
+	static constexpr auto WEBRTC_PIXEL_FORMAT = webrtc::VideoType::kARGB;
+
+	uint32 SizeX = Frame.width();
+	uint32 SizeY = Frame.height();
+	auto Format = TEXTURE_PIXEL_FORMAT;
+	auto NumMips = 0;
+	auto SamplerAddressMode = ESamplerAddressMode::AM_Clamp;
+
+    if (Frame.video_frame_buffer()->type() == webrtc::VideoFrameBuffer::Type::kNative)
+    {
+        UE::PixelStreaming::FFrameBufferRHI* const FrameBuffer = static_cast<UE::PixelStreaming::FFrameBufferRHI*>(Frame.video_frame_buffer().get());
+
+        if(FrameBuffer == nullptr)
+        {
+            return;
+        }
+
+        TSharedPtr<FVideoResourceRHI, ESPMode::ThreadSafe> VideoResource = FrameBuffer->GetVideoResource()->TransformResource(FVideoDescriptor(EVideoFormat::BGRA, SizeX, SizeY));
+		
+        auto& Raw = StaticCastSharedPtr<FVideoResourceRHI>(VideoResource)->GetRaw();
+        OnFrame(Raw.Texture);
+    }
+    else if (Frame.video_frame_buffer()->type() == webrtc::VideoFrameBuffer::Type::kI420)
+	{
+	    {
+		    FScopeLock Lock(&RenderSyncContext);
+
+		    const uint32_t Size = webrtc::CalcBufferSize(WEBRTC_PIXEL_FORMAT, SizeX, SizeY);
+		    if (Size > static_cast<uint32_t>(Buffer.Num()))
+		    {
+    			Buffer.SetNum(Size);
+		    }
+		    webrtc::ConvertFromI420(Frame, WEBRTC_PIXEL_FORMAT, 0, Buffer.GetData());
+	    }
+
+		AsyncTask(ENamedThreads::ActualRenderingThread, [SizeX, SizeY, Format, this]() 
+    	{
+    	    FScopeLock Lock(&RenderSyncContext);
+
+    	    const FIntPoint FrameSize = FIntPoint(SizeX, SizeY);
+    	    FRHICommandListImmediate& RHICmdList = FRHICommandListExecutor::GetImmediateCommandList();
+
+            FTextureRHIRef SourceTexture;
+
+    	    if (!RenderTargetDescriptor.IsValid() || RenderTargetDescriptor.GetSize() != FIntVector(FrameSize.X, FrameSize.Y, 0))
+    	    {
+    	        // Create the RenderTarget descriptor
+    	        RenderTargetDescriptor = FPooledRenderTargetDesc::Create2DDesc(FrameSize,
+    	            Format,
+    	            FClearValueBinding::None,
+    	            TexCreate_None,
+    	            TexCreate_RenderTargetable,
+    	            false);
+
+    	        // Update the shader resource for the 'SourceTexture'
+    	        FRHITextureCreateDesc RenderTargetTextureDesc =
+    	            FRHITextureCreateDesc::Create2D(TEXT(""), FrameSize.X, FrameSize.Y, Format)
+    	                .SetClearValue(FClearValueBinding::None)
+#if PLATFORM_MAC
+    	                .SetFlags(ETextureCreateFlags::CPUReadback | ETextureCreateFlags::SRGB)
+	                    .SetInitialState(ERHIAccess::CPURead);
+#else
+        	            .SetFlags(ETextureCreateFlags::Dynamic | ETextureCreateFlags::ShaderResource | TexCreate_RenderTargetable | ETextureCreateFlags::SRGB)
+    	                .SetInitialState(ERHIAccess::SRVMask);
+#endif
+
+            	SourceTexture = RHICreateTexture(RenderTargetTextureDesc);
+
+            	// Find a free target-able texture from the render pool
+            	GRenderTargetPool.FindFreeElement(RHICmdList,
+                	RenderTargetDescriptor,
+                	RenderTarget,
+                	TEXT("PIXELSTEAMINGPLAYER"));
+        	}
+
+        	// Create the update region structure
+        	const FUpdateTextureRegion2D Region(0, 0, 0, 0, FrameSize.X, FrameSize.Y);
+
+        	// Set the Pixel data of the webrtc Frame to the SourceTexture
+        	RHIUpdateTexture2D(SourceTexture, 0, Region, FrameSize.X * 4, Buffer.GetData());
+
+            OnFrame(SourceTexture);
+    	});
+	}
+}
\ No newline at end of file
diff --git a/Engine/Plugins/Media/PixelStreaming/Source/PixelStreaming/Private/VideoDecoderFactory.cpp b/Engine/Plugins/Media/PixelStreaming/Source/PixelStreaming/Private/VideoDecoderFactory.cpp
--- a/Engine/Plugins/Media/PixelStreaming/Source/PixelStreaming/Private/VideoDecoderFactory.cpp
+++ b/Engine/Plugins/Media/PixelStreaming/Source/PixelStreaming/Private/VideoDecoderFactory.cpp
@@ -3,7 +3,7 @@
 #include "Utils.h"
 #include "VideoDecoderStub.h"
 #include "VideoDecoderVPX.h"
-#include "VideoDecoderH265.h"
+#include "VideoDecoderHardware.h"
 #include "PixelStreamingPrivate.h"
 #include "Settings.h"
 
@@ -16,7 +16,7 @@
 namespace UE::PixelStreaming
 {
 	// the list of each individual codec we have decoder support for (the order of this array is preference order after the selected codec)
-	const TArray<EPixelStreamingCodec> SupportedDecoderCodecList{ EPixelStreamingCodec::VP8, EPixelStreamingCodec::VP9, /* EPixelStreamingCodec::H264, */ EPixelStreamingCodec::H265 };
+	const TArray<EPixelStreamingCodec> SupportedDecoderCodecList{ EPixelStreamingCodec::VP8, EPixelStreamingCodec::VP9, EPixelStreamingCodec::H264 };
 
 	// mapping of codec to a list of video formats
 	// done this way so we can order the list of formats based on selected codec in GetSupportedFormats
@@ -30,9 +30,8 @@ namespace UE::PixelStreaming
 
 		Codecs[EPixelStreamingCodec::VP8].push_back(webrtc::SdpVideoFormat(cricket::kVp8CodecName));
 		Codecs[EPixelStreamingCodec::VP9].push_back(webrtc::SdpVideoFormat(cricket::kVp9CodecName));
-		// Codecs[EPixelStreamingCodec::H264].push_back(CreateH264Format(webrtc::H264Profile::kProfileConstrainedBaseline, webrtc::H264Level::kLevel3_1));
-		// Codecs[EPixelStreamingCodec::H264].push_back(CreateH264Format(webrtc::H264Profile::kProfileBaseline, webrtc::H264Level::kLevel3_1));
-		// Codecs[EPixelStreamingCodec::H265].push_back(webrtc::SdpVideoFormat(cricket::kH265CodecName));
+		Codecs[EPixelStreamingCodec::H264].push_back(CreateH264Format(webrtc::H264Profile::kProfileConstrainedBaseline, webrtc::H264Level::kLevel3_1));
+		Codecs[EPixelStreamingCodec::H264].push_back(CreateH264Format(webrtc::H264Profile::kProfileBaseline, webrtc::H264Level::kLevel3_1));
 
 		return Codecs;
 	}
@@ -127,12 +126,10 @@ namespace UE::PixelStreaming
 		{
 			return std::make_unique<VideoDecoderVPX>(9);
 		}
-		/*
-		else if (absl::EqualsIgnoreCase(format.name, cricket::kH265CodecName))
+		else if (absl::EqualsIgnoreCase(format.name, cricket::kH264CodecName))
 		{
-			return std::make_unique<VideoDecoderH265>();
+			return std::make_unique<FVideoDecoderHardware>();
 		}
-		*/
 		return std::make_unique<FVideoDecoderStub>();
 	}
 } // namespace UE::PixelStreaming
diff --git a/Engine/Plugins/Media/PixelStreaming/Source/PixelStreaming/Private/VideoDecoderH265.cpp b/Engine/Plugins/Media/PixelStreaming/Source/PixelStreaming/Private/VideoDecoderH265.cpp
deleted file mode 100644
diff --git a/Engine/Plugins/Media/PixelStreaming/Source/PixelStreaming/Private/VideoDecoderH265.h b/Engine/Plugins/Media/PixelStreaming/Source/PixelStreaming/Private/VideoDecoderH265.h
deleted file mode 100644
diff --git a/Engine/Plugins/Media/PixelStreaming/Source/PixelStreaming/Private/VideoDecoderHardware.cpp b/Engine/Plugins/Media/PixelStreaming/Source/PixelStreaming/Private/VideoDecoderHardware.cpp
new file mode 100644
--- /dev/null
+++ b/Engine/Plugins/Media/PixelStreaming/Source/PixelStreaming/Private/VideoDecoderHardware.cpp
@@ -0,0 +1,93 @@
+// Copyright Epic Games, Inc. All Rights Reserved.
+
+#include "VideoDecoderHardware.h"
+
+#include "FrameBufferRHI.h"
+#include "PixelStreamingPrivate.h"
+#include "PixelStreamingTrace.h"
+
+namespace UE::PixelStreaming
+{
+	// HACK (aidan.possemiers) the API surface SendPacket surface wants a SharedPtr for the encoded data but WebRTC already owns that and the types are incompatible
+	struct FFakeDeleter
+	{
+		void operator()(uint8* Object) const
+		{
+		}
+	};
+
+	FVideoDecoderHardware::FVideoDecoderHardware()
+	{
+		FVideoDecoderConfigH264 DecoderConfig;
+		Decoder = FVideoDecoder::CreateChecked<FVideoResourceRHI>(FAVDevice::GetHardwareDevice(), DecoderConfig);
+	}
+
+	bool FVideoDecoderHardware::Configure(const Settings& settings)
+	{
+		return true;
+	}
+
+	int32 FVideoDecoderHardware::Decode(const webrtc::EncodedImage& input_image, bool missing_frames, int64_t render_time_ms)
+	{
+		const int64 TimestampDecodeStart = rtc::TimeMillis();
+		TRACE_CPUPROFILER_EVENT_SCOPE_ON_CHANNEL_STR("PixelStreaming Hardware Decoding Video", PixelStreamingChannel);
+
+		FAVResult Result = Decoder->SendPacket(FVideoPacket(
+			MakeShareable<uint8>(input_image.GetEncodedData()->data(), FFakeDeleter()),
+			input_image.GetEncodedData()->size(),
+			input_image.Timestamp(),
+			FrameCount++,
+			input_image.qp_,
+			input_image._frameType == webrtc::VideoFrameType::kVideoFrameKey));
+
+		if (Output != nullptr && Result.IsNotError())
+		{
+			FResolvableVideoResourceRHI DecoderResource;
+			const FAVResult DecodeResult = Decoder->ReceiveFrame(DecoderResource);
+
+			if (DecodeResult.IsSuccess())
+			{
+#if WEBRTC_5414
+				rtc::scoped_refptr<webrtc::VideoFrameBuffer> FrameBuffer = rtc::make_ref_counted<UE::PixelStreaming::FFrameBufferRHI>(DecoderResource);
+#else
+				rtc::scoped_refptr<webrtc::VideoFrameBuffer> FrameBuffer = new rtc::RefCountedObject<UE::PixelStreaming::FFrameBufferRHI>(DecoderResource);
+#endif
+				check(FrameBuffer->width() != 0 && FrameBuffer->height() != 0); // TODO we should probably check that we are getting the frame back that we are expecting
+
+				webrtc::VideoFrame Frame = webrtc::VideoFrame::Builder()
+											   .set_video_frame_buffer(FrameBuffer)
+											   .set_timestamp_rtp(input_image.Timestamp()) // NOTE This timestamp is load bearing as WebRTC stores the frames in a map with this as the index
+											   .set_color_space(input_image.ColorSpace())
+											   .build();
+
+				Output->Decoded(Frame, rtc::TimeMillis() - TimestampDecodeStart, input_image.qp_);
+
+				return WEBRTC_VIDEO_CODEC_OK;
+			}
+		}
+		else
+		{
+			UE_LOG(LogTemp, Warning, TEXT("FVideoDecoderHardware::Decode FAILED"));
+
+			return WEBRTC_VIDEO_CODEC_OK_REQUEST_KEYFRAME;
+		}
+
+		UE_LOG(LogTemp, Error, TEXT("FVideoDecoderHardware::Decode ERROR"));
+
+		return WEBRTC_VIDEO_CODEC_ERROR;
+	}
+
+	int32 FVideoDecoderHardware::RegisterDecodeCompleteCallback(webrtc::DecodedImageCallback* callback)
+	{
+		Output = callback;
+
+		return 0;
+	}
+
+	int32 FVideoDecoderHardware::Release()
+	{
+		Decoder.Reset();
+
+		return 0;
+	}
+} // namespace UE::PixelStreaming
diff --git a/Engine/Plugins/Media/PixelStreaming/Source/PixelStreaming/Private/VideoDecoderHardware.h b/Engine/Plugins/Media/PixelStreaming/Source/PixelStreaming/Private/VideoDecoderHardware.h
new file mode 100644
--- /dev/null
+++ b/Engine/Plugins/Media/PixelStreaming/Source/PixelStreaming/Private/VideoDecoderHardware.h
@@ -0,0 +1,31 @@
+// Copyright Epic Games, Inc. All Rights Reserved.
+
+#pragma once
+
+#include "WebRTCIncludes.h"
+
+#include "Video/Decoders/Configs/VideoDecoderConfigH264.h"
+#include "Video/Resources/VideoResourceRHI.h"
+
+namespace UE::PixelStreaming
+{
+	class FVideoDecoderHardware : public webrtc::VideoDecoder
+	{
+	public:
+		FVideoDecoderHardware();
+		virtual ~FVideoDecoderHardware() = default;
+
+		virtual bool Configure(const Settings& settings) override;
+		virtual int32 Decode(const webrtc::EncodedImage& input_image, bool missing_frames, int64_t render_time_ms) override;
+		virtual int32 RegisterDecodeCompleteCallback(webrtc::DecodedImageCallback* callback) override;
+		virtual int32 Release() override;
+		virtual const char* ImplementationName() const { return "FVideoDecoderHardware"; }
+
+	private:
+		TSharedPtr<TVideoDecoder<FVideoResourceRHI>> Decoder;
+
+		webrtc::DecodedImageCallback* Output = nullptr;
+
+		uint32 FrameCount = 0;
+	};
+} // namespace UE::PixelStreaming
diff --git a/Engine/Plugins/Media/PixelStreaming/Source/PixelStreaming/Public/PixelStreamingVideoSink.h b/Engine/Plugins/Media/PixelStreaming/Source/PixelStreaming/Public/PixelStreamingVideoSink.h
new file mode 100644
--- /dev/null
+++ b/Engine/Plugins/Media/PixelStreaming/Source/PixelStreaming/Public/PixelStreamingVideoSink.h
@@ -0,0 +1,29 @@
+// Copyright Epic Games, Inc. All Rights Reserved.
+
+#pragma once
+
+#include "PixelStreamingWebRTCIncludes.h"
+#include "Engine/Texture2DDynamic.h"
+#include "RenderTargetPool.h"
+
+class PIXELSTREAMING_API FPixelStreamingVideoSink : public rtc::VideoSinkInterface<webrtc::VideoFrame>
+{
+public:
+    FPixelStreamingVideoSink() = default;
+    ~FPixelStreamingVideoSink() = default;
+
+    // Begin VideoSinkInterface interface
+    // This method handles the conversion from a WebRTC frame to a UE texture
+    virtual void OnFrame(const webrtc::VideoFrame& Frame) override;
+    // End VideoSinkInterface interface
+
+    // Derivative classes implement this method to handle how the decoded textures should actually be displayed
+    // eg PixelStreamingBlueprint/PixelStreamingPeerVideoSink
+    virtual void OnFrame(FTextureRHIRef Frame) { unimplemented(); }
+
+private:
+    FCriticalSection RenderSyncContext;
+    FPooledRenderTargetDesc RenderTargetDescriptor;
+    TRefCountPtr<IPooledRenderTarget> RenderTarget;
+    TArray<uint8_t> Buffer;
+};
\ No newline at end of file
