Intercept resize earlier in PixelStreamingMediaIO so that we do not pass wrong size texture into the MediaIO capture pipeline and cause an error.

diff --git a/Engine/Plugins/Media/PixelStreaming/Source/PixelStreaming/Private/PixelStreamingMediaIOCapture.cpp b/Engine/Plugins/Media/PixelStreaming/Source/PixelStreaming/Private/PixelStreamingMediaIOCapture.cpp
--- a/Engine/Plugins/Media/PixelStreaming/Source/PixelStreaming/Private/PixelStreamingMediaIOCapture.cpp
+++ b/Engine/Plugins/Media/PixelStreaming/Source/PixelStreaming/Private/PixelStreamingMediaIOCapture.cpp
@@ -85,4 +85,9 @@ bool UPixelStreamingMediaIOCapture::PostInitializeCaptureViewport(TSharedPtr<FSc
 void UPixelStreamingMediaIOCapture::ViewportResized(FViewport* Viewport, uint32 ResizeCode)
 {
 	bViewportResized = true;
+	if(GetState() == EMediaCaptureState::Capturing)
+	{
+		UE_LOG(LogPixelStreaming, Warning, TEXT("Stopping PixelStreaming MediaIO capture because viewport was resized."));
+		StopCapture(false);
+	}
 }
