Fix: Allow resampling of audio data passed to PS before being sent to WebRTC

diff --git a/Engine/Plugins/Media/PixelStreaming/Source/PixelStreaming/PixelStreaming.Build.cs b/Engine/Plugins/Media/PixelStreaming/Source/PixelStreaming/PixelStreaming.Build.cs
--- a/Engine/Plugins/Media/PixelStreaming/Source/PixelStreaming/PixelStreaming.Build.cs
+++ b/Engine/Plugins/Media/PixelStreaming/Source/PixelStreaming/PixelStreaming.Build.cs
@@ -97,7 +97,8 @@ namespace UnrealBuildTool.Rules
                 "PixelStreamingHMD",
                 "TraceLog",
                 "HTTP",
-                "NVML"
+                "NVML",
+                "AudioPlatformConfiguration"
             });
 
             PrivateDefinitions.Add("PIXELSTREAMING_DUMP_ENCODING=0");
diff --git a/Engine/Plugins/Media/PixelStreaming/Source/PixelStreaming/Private/AudioInputMixer.cpp b/Engine/Plugins/Media/PixelStreaming/Source/PixelStreaming/Private/AudioInputMixer.cpp
--- a/Engine/Plugins/Media/PixelStreaming/Source/PixelStreaming/Private/AudioInputMixer.cpp
+++ b/Engine/Plugins/Media/PixelStreaming/Source/PixelStreaming/Private/AudioInputMixer.cpp
@@ -4,6 +4,7 @@
 #include "Settings.h"
 #include "SampleBuffer.h"
 #include "PixelStreamingTrace.h"
+#include "PixelStreamingPrivate.h"
 
 namespace UE::PixelStreaming
 {
@@ -158,23 +159,47 @@ namespace UE::PixelStreaming
 		NumChannels = InMixer.Mixers->NumChannels;
 		SampleRate = InMixer.Mixers->SampleRate;
 	}
-
+	
 	void FAudioInput::PushAudio(const float* InBuffer, int32 NumSamples, int32 InNumChannels, int32 InSampleRate)
 	{
 		TRACE_CPUPROFILER_EVENT_SCOPE_ON_CHANNEL_STR("PixelStreaming FAudioInputMixer::PushAudio", PixelStreamingChannel);
-		// Todo(Luke): we could resample instead of throwing this assertion.
-		verifyf(SampleRate == InSampleRate, TEXT("Sample rate was mismatched."));
+		TArray<float> AudioBuffer;
+
+		if(SampleRate != InSampleRate)
+		{
+			float SampleRateConversionRatio = static_cast<float>(SampleRate) / static_cast<float>(InSampleRate);
+			UE_LOG(LogPixelStreaming, VeryVerbose, TEXT("(FAudioInput) Sample rate conversion ratio is : %f."), SampleRateConversionRatio);
+			Resampler.Init(Audio::EResamplingMethod::Linear, SampleRateConversionRatio, InNumChannels);
+
+			int32 NumConvertedSamples = InSampleRate / SampleRateConversionRatio;
+			int32 OutputSamples = INDEX_NONE;
+			AudioBuffer.AddZeroed(NumConvertedSamples);
+			// Perform the sample rate conversion
+			int32 ErrorCode = Resampler.ProcessAudio(const_cast<float*>(InBuffer), NumSamples, false, AudioBuffer.GetData(), AudioBuffer.Num(), OutputSamples);
+			verifyf(OutputSamples <= NumConvertedSamples, TEXT("OutputSamples > NumConvertedSamples"));
+			if (ErrorCode != 0)
+			{
+				UE_LOG(LogPixelStreaming, Warning, TEXT("(FAudioInput) Problem occured resampling audio data. Code: %d"), ErrorCode);
+				return;
+			}
+		}
+		else
+		{
+			AudioBuffer.AddZeroed(NumSamples);
+			FMemory::Memcpy(AudioBuffer.GetData(), InBuffer, AudioBuffer.Num() * sizeof(float));
+		}
+
 
 		// Mix to our target number of channels if the source does not already match.
 		if (NumChannels != InNumChannels)
 		{
-			Audio::TSampleBuffer<float> Buffer(InBuffer, NumSamples, InNumChannels, InSampleRate);
+			Audio::TSampleBuffer<float> Buffer(AudioBuffer.GetData(), AudioBuffer.Num(), InNumChannels, InSampleRate);
 			Buffer.MixBufferToChannels(NumChannels);
 			PatchInput.PushAudio(Buffer.GetData(), Buffer.GetNumSamples());
 		}
 		else
 		{
-			PatchInput.PushAudio(InBuffer, NumSamples);
+			PatchInput.PushAudio(AudioBuffer.GetData(), AudioBuffer.Num());
 		}
 	}
 
diff --git a/Engine/Plugins/Media/PixelStreaming/Source/PixelStreaming/Private/AudioInputMixer.h b/Engine/Plugins/Media/PixelStreaming/Source/PixelStreaming/Private/AudioInputMixer.h
--- a/Engine/Plugins/Media/PixelStreaming/Source/PixelStreaming/Private/AudioInputMixer.h
+++ b/Engine/Plugins/Media/PixelStreaming/Source/PixelStreaming/Private/AudioInputMixer.h
@@ -69,6 +69,7 @@ namespace UE::PixelStreaming
 		const Audio::FPatchInput& GetPatchInput() const { return PatchInput; }
 
 	private:
+		Audio::FResampler Resampler;
 		Audio::FPatchInput PatchInput;
 		uint8 NumChannels;
 		uint32 SampleRate;
