Fix: Ensure OnTextCommitted is called when a text update is triggered by the LiveLink app

diff --git a/Engine/Plugins/Media/PixelStreaming/Source/PixelStreamingInput/Private/PixelStreamingInputHandler.cpp b/Engine/Plugins/Media/PixelStreaming/Source/PixelStreamingInput/Private/PixelStreamingInputHandler.cpp
--- a/Engine/Plugins/Media/PixelStreaming/Source/PixelStreamingInput/Private/PixelStreamingInputHandler.cpp
+++ b/Engine/Plugins/Media/PixelStreaming/Source/PixelStreamingInput/Private/PixelStreamingInputHandler.cpp
@@ -1200,6 +1200,18 @@ namespace UE::PixelStreamingInput
 					SMultiLineEditableTextBox* TextBox = static_cast<SMultiLineEditableTextBox*>(FocusedWidget.Get());
 					TextBox->SetText(FText::FromString(Text));
 				}
+
+				// We need to manually trigger an Enter key press so that the OnTextCommitted delegate gets fired
+				const uint32* KeyPtr = nullptr;
+				const uint32* CharacterPtr = nullptr;
+				FInputKeyManager::Get().GetCodesFromKey(EKeys::Enter, KeyPtr, CharacterPtr);
+				uint32 Key = KeyPtr ? *KeyPtr : 0;
+				uint32 Character = CharacterPtr ? *CharacterPtr : 0;
+				if(Key != 0 || Character != 0)
+				{
+					MessageHandler->OnKeyDown((int32)Key, (int32)Character, false);
+					MessageHandler->OnKeyUp((int32)Key, (int32)Character, false);
+				}
 			}
 		});
 	}
