Fix: EditableText doesn't respect ReadOnly property

diff --git a/Engine/Plugins/Media/PixelStreaming/Source/PixelStreamingInput/Private/PixelStreamingInputHandler.cpp b/Engine/Plugins/Media/PixelStreaming/Source/PixelStreamingInput/Private/PixelStreamingInputHandler.cpp
--- a/Engine/Plugins/Media/PixelStreaming/Source/PixelStreamingInput/Private/PixelStreamingInputHandler.cpp
+++ b/Engine/Plugins/Media/PixelStreaming/Source/PixelStreamingInput/Private/PixelStreamingInputHandler.cpp
@@ -1386,6 +1386,19 @@ namespace UE::PixelStreamingInput
 			static FName SEditableTextType(TEXT("SEditableText"));
 			static FName SMultiLineEditableTextType(TEXT("SMultiLineEditableText"));
 			bool bEditable = FocusedWidget && (FocusedWidget->GetType() == SEditableTextType || FocusedWidget->GetType() == SMultiLineEditableTextType);
+			if (bEditable)
+			{
+				if (FocusedWidget->GetType() == TEXT("SEditableText"))
+				{
+					SEditableText* TextBox = static_cast<SEditableText*>(FocusedWidget.Get());
+					bEditable = !TextBox->IsTextReadOnly();
+				}
+				else if (FocusedWidget->GetType() == TEXT("SMultiLineEditableText"))
+				{
+					SMultiLineEditableText* TextBox = static_cast<SMultiLineEditableText*>(FocusedWidget.Get());
+					bEditable = !TextBox->IsTextReadOnly();
+				}
+			}
 
 			// Tell the browser that the focus has changed.
 			TSharedPtr<FJsonObject> JsonObject = MakeShareable(new FJsonObject);
