Fix: EditText button doesn't show when tapping an editable text widget

diff --git a/Engine/Plugins/Media/PixelStreaming/Source/PixelStreamingInput/Private/PixelStreamingInputHandler.cpp b/Engine/Plugins/Media/PixelStreaming/Source/PixelStreamingInput/Private/PixelStreamingInputHandler.cpp
--- a/Engine/Plugins/Media/PixelStreaming/Source/PixelStreamingInput/Private/PixelStreamingInputHandler.cpp
+++ b/Engine/Plugins/Media/PixelStreaming/Source/PixelStreamingInput/Private/PixelStreamingInputHandler.cpp
@@ -47,8 +47,6 @@ namespace UE::PixelStreamingInput
 		, bIsMouseActive(false)
 		, MessageHandler(InTargetHandler)
 		, PixelStreamerApplicationWrapper(InApplicationWrapper)
-		, FocusedPos(FVector2D(-1.0f, -1.0f))
-		, UnfocusedPos(FVector2D(-1.0f, -1.0f))
 	{
 		// Register this input handler as an IMotionController. The module handles the registering as an IInputDevice
 		IModularFeatures::Get().RegisterModularFeature(GetModularFeatureName(), this);
@@ -1448,7 +1446,7 @@ namespace UE::PixelStreamingInput
 				{
 					FIntRect ScreenRect = *ScreenRectPtr;
 					FIntPoint SizeInScreen = ScreenRect.Max - ScreenRect.Min;
-					NormalizedLocation = FocusedPos / SizeInScreen;
+					NormalizedLocation = FocusedWidget->GetCachedGeometry().GetAbsolutePosition() / SizeInScreen;
 				}
 
 				NormalizedLocation *= uint16_MAX;
diff --git a/Engine/Plugins/Media/PixelStreaming/Source/PixelStreamingInput/Private/PixelStreamingInputHandler.h b/Engine/Plugins/Media/PixelStreaming/Source/PixelStreamingInput/Private/PixelStreamingInputHandler.h
--- a/Engine/Plugins/Media/PixelStreaming/Source/PixelStreamingInput/Private/PixelStreamingInputHandler.h
+++ b/Engine/Plugins/Media/PixelStreaming/Source/PixelStreamingInput/Private/PixelStreamingInputHandler.h
@@ -198,18 +198,6 @@ namespace UE::PixelStreamingInput
 		 */
 		bool bFakingTouchEvents;
 
-		/**
-		 * Touch only. Location of the focused UI widget. If no UI widget is focused
-		 * then this has the UnfocusedPos value.
-		 */
-		FVector2D FocusedPos;
-
-		/**
-		 * Touch only. A special position which indicates that no UI widget is
-		 * focused.
-		 */
-		const FVector2D UnfocusedPos;
-
 		/*
 		 * Padding for string parsing when handling messages.
 		 * 1 character for the actual message and then
