QOL: Allow the passing of a custom websocket to the signalling server connection.
This allows for the removal of the mock signalling connection from the tests and allows the testing of the `FPixelStreamingSignallingConnection`

diff --git a/Engine/Plugins/Media/PixelStreaming/Source/PixelStreaming/Private/PixelStreamingSignallingConnection.cpp b/Engine/Plugins/Media/PixelStreaming/Source/PixelStreaming/Private/PixelStreamingSignallingConnection.cpp
--- a/Engine/Plugins/Media/PixelStreaming/Source/PixelStreaming/Private/PixelStreamingSignallingConnection.cpp
+++ b/Engine/Plugins/Media/PixelStreaming/Source/PixelStreaming/Private/PixelStreamingSignallingConnection.cpp
@@ -17,9 +17,10 @@
 DECLARE_LOG_CATEGORY_EXTERN(LogPixelStreamingSS, Log, VeryVerbose);
 DEFINE_LOG_CATEGORY(LogPixelStreamingSS);
 
-FPixelStreamingSignallingConnection::FPixelStreamingSignallingConnection(TSharedPtr<IPixelStreamingSignallingConnectionObserver> InObserver, FString InStreamerId)
+FPixelStreamingSignallingConnection::FPixelStreamingSignallingConnection(TSharedPtr<IPixelStreamingSignallingConnectionObserver> InObserver, FString InStreamerId, TSharedPtr<IWebSocket> InWebSocket)
 	: Observer(InObserver)
 	, StreamerId(InStreamerId)
+	, WebSocket(InWebSocket)
 {
 	RegisterHandler("identify", [this](FJsonObjectPtr JsonMsg) { OnIdRequested(); });
 	RegisterHandler("config", [this](FJsonObjectPtr JsonMsg) { OnConfig(JsonMsg); });
@@ -39,6 +40,9 @@ FPixelStreamingSignallingConnection::FPixelStreamingSignallingConnection(TShared
 FPixelStreamingSignallingConnection::~FPixelStreamingSignallingConnection()
 {
 	Disconnect();
+
+	WebSocket.Reset();
+	Observer.Reset();
 }
 
 void FPixelStreamingSignallingConnection::Connect(FString InUrl, bool bIsReconnect)
@@ -72,8 +76,11 @@ void FPixelStreamingSignallingConnection::Connect(FString InUrl, bool bIsReconne
 		Url = Final;
 	}
 
-	WebSocket = FWebSocketsModule::Get().CreateWebSocket(Url, TEXT(""));
-	verifyf(WebSocket, TEXT("Web Socket Factory failed to return a valid Web Socket."));
+	if(!WebSocket.IsValid())
+	{	
+		WebSocket = FWebSocketsModule::Get().CreateWebSocket(Url, TEXT(""));
+		verifyf(WebSocket, TEXT("Web Socket Factory failed to return a valid Web Socket."));
+	}
 
 	OnConnectedHandle = WebSocket->OnConnected().AddLambda([this]() { OnConnected(); });
 	OnConnectionErrorHandle = WebSocket->OnConnectionError().AddLambda([this](const FString& Error) { OnConnectionError(Error); });
@@ -118,7 +125,6 @@ void FPixelStreamingSignallingConnection::Disconnect()
 	WebSocket->OnBinaryMessage().Remove(OnBinaryMessageHandle);
 
 	WebSocket->Close();
-	WebSocket = nullptr;
 	UE_LOG(LogPixelStreamingSS, Log, TEXT("Closing websocket to SS %s"), *Url);
 }
 
diff --git a/Engine/Plugins/Media/PixelStreaming/Source/PixelStreaming/Private/Tests/ProtocolTest.cpp b/Engine/Plugins/Media/PixelStreaming/Source/PixelStreaming/Private/Tests/ProtocolTest.cpp
--- a/Engine/Plugins/Media/PixelStreaming/Source/PixelStreaming/Private/Tests/ProtocolTest.cpp
+++ b/Engine/Plugins/Media/PixelStreaming/Source/PixelStreaming/Private/Tests/ProtocolTest.cpp
@@ -26,16 +26,14 @@ namespace UE::PixelStreaming
 		TSharedPtr<FPixelStreamingVideoInputI420> VideoInput = MakeShared<FPixelStreamingVideoInputI420>();
 		Streamer->SetVideoInput(VideoInput);
 
-		IPixelStreamingInputModule& PixelStreamingInputModule = IPixelStreamingInputModule::Get();
-		EPixelStreamingMessageDirection MessageDirection = EPixelStreamingMessageDirection::ToStreamer;
-
+		// Define our message
 		FPixelStreamingInputMessage Message = FPixelStreamingInputMessage({ EPixelStreamingMessageTypes::Uint16 } /* Structure */);
-
-		const TFunction<void(FString, FMemoryReader)> Handler = [this](FString, FMemoryReader Ar) { /* Do nothing */ };
-
-		TSharedPtr<IPixelStreamingInputHandler> InputHandler = Streamer->GetInputHandler().Pin();
+		// Add it to the protocol
 		FPixelStreamingInputProtocol::ToStreamerProtocol.Add("CustomMessage", Message);
-		InputHandler->RegisterMessageHandler("CustomMessage", Handler);
+		// Define a handler function
+		const TFunction<void(FString, FMemoryReader)> Handler = [this](FString, FMemoryReader Ar) { /* Do nothing */ };
+		// Add it to the streamer
+		Streamer->GetInputHandler().Pin()->RegisterMessageHandler("CustomMessage", Handler);
 
 		TSharedPtr<FMockPlayer> Player = CreatePlayer(FMockPlayer::EMode::AcceptOffers);
 		TSharedPtr<FMockVideoSink> VideoSink = MakeShared<FMockVideoSink>();
@@ -93,12 +91,11 @@ namespace UE::PixelStreaming
 		TSharedPtr<FPixelStreamingVideoInputI420> VideoInput = MakeShared<FPixelStreamingVideoInputI420>();
 		Streamer->SetVideoInput(VideoInput);
 
-		IPixelStreamingInputModule& PixelStreamingInputModule = IPixelStreamingInputModule::Get();
-		EPixelStreamingMessageDirection MessageDirection = EPixelStreamingMessageDirection::ToStreamer;
-
+		// Define our message
 		FPixelStreamingInputMessage Message = FPixelStreamingInputMessage({ EPixelStreamingMessageTypes::Uint16 } /* Structure */);
-
-		TFunction<void(uint8, const webrtc::DataBuffer&)> Callback = [](uint8 Type, const webrtc::DataBuffer& RawBuffer) { /* Do nothing */ };
+		// Add it to the protocol
+		FPixelStreamingInputProtocol::ToStreamerProtocol.Add("CustomMessage", Message);
+		// Define a handler function
 		TSharedPtr<bool> bComplete = MakeShared<bool>(false);
 		const TFunction<void(FString, FMemoryReader)> Handler = [this, bComplete](FString, FMemoryReader Ar) {
 			*bComplete.Get() = true;
@@ -106,10 +103,12 @@ namespace UE::PixelStreaming
 			Ar << Out;
 			TestTrue(TEXT("Expected message content to be 1337."), Out == 1337);
 		};
+		// Add it to the streamer
+		Streamer->GetInputHandler().Pin()->RegisterMessageHandler("CustomMessage", Handler);
+
+		
+		TFunction<void(uint8, const webrtc::DataBuffer&)> Callback = [](uint8 Type, const webrtc::DataBuffer& RawBuffer) { /* Do nothing */ };
 
-		TSharedPtr<IPixelStreamingInputHandler> InputHandler = Streamer->GetInputHandler().Pin();
-		FPixelStreamingInputProtocol::ToStreamerProtocol.Add("CustomMessage", Message);
-		InputHandler->RegisterMessageHandler("CustomMessage", Handler);
 
 		TSharedPtr<FMockPlayer> Player = CreatePlayer(FMockPlayer::EMode::AcceptOffers);
 		TSharedPtr<FMockVideoSink> VideoSink = MakeShared<FMockVideoSink>();
diff --git a/Engine/Plugins/Media/PixelStreaming/Source/PixelStreaming/Private/Tests/SignallingServerConnectionTests.cpp b/Engine/Plugins/Media/PixelStreaming/Source/PixelStreaming/Private/Tests/SignallingServerConnectionTests.cpp
--- a/Engine/Plugins/Media/PixelStreaming/Source/PixelStreaming/Private/Tests/SignallingServerConnectionTests.cpp
+++ b/Engine/Plugins/Media/PixelStreaming/Source/PixelStreaming/Private/Tests/SignallingServerConnectionTests.cpp
@@ -16,10 +16,10 @@ namespace
 	// using namespace UE::PixelStreaming;
 
 	// For validating callbacks from the signalling server connection
-	class FTestSSConnectionObserver : public IPixelStreamingSignallingConnectionObserver
+	class FMockConnectionObserver : public IPixelStreamingSignallingConnectionObserver
 	{
 	public:
-		virtual ~FTestSSConnectionObserver() = default;
+		virtual ~FMockConnectionObserver() = default;
 
 		virtual void OnSignallingConnected() override {}
 		virtual void OnSignallingDisconnected(int32 StatusCode, const FString& Reason, bool bWasClean) override {}
@@ -94,7 +94,7 @@ namespace
 		}
 		virtual void Close(int32 Code = 1000, const FString& Reason = FString()) override { bConnected = false; }
 		virtual bool IsConnected() override { return bConnected; }
-		virtual void Send(const FString& Data) override { OnMockResponseEvent.Broadcast(Data); }
+		virtual void Send(const FString& Data) override { OnMessageSentEvent.Broadcast(Data); }
 		virtual void Send(const void* Data, SIZE_T Size, bool bIsBinary = false) override {}
 		virtual void SetTextMessageMemoryLimit(uint64 TextMessageMemoryLimit) override {}
 		virtual FWebSocketConnectedEvent& OnConnected() override { return OnConnectedEvent; }
@@ -105,12 +105,6 @@ namespace
 		virtual FWebSocketRawMessageEvent& OnRawMessage() override { return OnRawMessageEvent; }
 		virtual FWebSocketMessageSentEvent& OnMessageSent() override { return OnMessageSentEvent; }
 
-		void MockSend(const FString& Msg) const { OnMessageEvent.Broadcast(Msg); }
-
-		// used to check outgoing messages
-		DECLARE_EVENT_OneParam(FMockWebSocket, FMockResponseEvent, const FString& /* Message */);
-		FMockResponseEvent OnMockResponseEvent;
-
 	private:
 		FWebSocketConnectedEvent OnConnectedEvent;
 		FWebSocketConnectionErrorEvent OnErrorEvent;
@@ -123,592 +117,6 @@ namespace
 		bool bConnected = false;
 	};
 
-	class FMockSignallingConnection : public IPixelStreamingSignallingConnection
-	{
-	public:
-		FMockSignallingConnection(IPixelStreamingSignallingConnectionObserver& InObserver, FString InStreamerId)
-			: Observer(InObserver)
-			, StreamerId(InStreamerId)
-		{
-			RegisterHandler("identify", [this](FJsonObjectPtr JsonMsg) {
-				FJsonObjectPtr Json = MakeShared<FJsonObject>();
-				Json->SetStringField(TEXT("type"), TEXT("endpointId"));
-				Json->SetStringField(TEXT("id"), StreamerId);
-				FString Msg = UE::PixelStreaming::ToString(Json, false);
-				UE_LOG(LogMockPixelStreamingSS, Verbose, TEXT("-> SS: endpointId\n%s"), *Msg);
-				SendMessage(Msg);
-			});
-			RegisterHandler("config", [this](FJsonObjectPtr JsonMsg) {
-				const TSharedPtr<FJsonObject>* PeerConnectionOptions = nullptr;
-				if (!JsonMsg->TryGetObjectField(TEXT("peerConnectionOptions"), PeerConnectionOptions))
-				{
-					UE_LOG(LogMockPixelStreamingSS, Error, TEXT("Cannot find `peerConnectionOptions` field in SS config\n%s"), *UE::PixelStreaming::ToString(JsonMsg));
-					return;
-				}
-
-				webrtc::PeerConnectionInterface::RTCConfiguration RTCConfig;
-
-				const TArray<TSharedPtr<FJsonValue>>* IceServers = nullptr;
-				if ((*PeerConnectionOptions)->TryGetArrayField(TEXT("iceServers"), IceServers))
-				{
-					for (const TSharedPtr<FJsonValue>& IceServerVal : *IceServers)
-					{
-						const TSharedPtr<FJsonObject>* IceServerJson;
-						if (!IceServerVal->TryGetObject(IceServerJson))
-						{
-							UE_LOG(LogMockPixelStreamingSS, Error, TEXT("Failed to parse SS config: `iceServer` - not an object\n%s"), *UE::PixelStreaming::ToString(*PeerConnectionOptions));
-							return;
-						}
-
-						RTCConfig.servers.push_back(webrtc::PeerConnectionInterface::IceServer{});
-						webrtc::PeerConnectionInterface::IceServer& IceServer = RTCConfig.servers.back();
-
-						TArray<FString> Urls;
-						if ((*IceServerJson)->TryGetStringArrayField(TEXT("urls"), Urls))
-						{
-							for (const FString& UrlElem : Urls)
-							{
-								IceServer.urls.push_back(UE::PixelStreaming::ToString(UrlElem));
-							}
-						}
-						else
-						{
-							// in the RTC Spec, "urls" can be an array or a single string
-							// https://www.w3.org/TR/webrtc/#dictionary-rtciceserver-members
-							FString UrlsSingle;
-							if ((*IceServerJson)->TryGetStringField(TEXT("urls"), UrlsSingle))
-							{
-								IceServer.urls.push_back(UE::PixelStreaming::ToString(UrlsSingle));
-							}
-						}
-
-						FString Username;
-						if ((*IceServerJson)->TryGetStringField(TEXT("username"), Username))
-						{
-							IceServer.username = UE::PixelStreaming::ToString(Username);
-						}
-
-						FString Credential;
-						if ((*IceServerJson)->TryGetStringField(TEXT("credential"), Credential))
-						{
-							IceServer.password = UE::PixelStreaming::ToString(Credential);
-						}
-					}
-				}
-
-				// force `UnifiedPlan` as we control both ends of WebRTC streaming
-				RTCConfig.sdp_semantics = webrtc::SdpSemantics::kUnifiedPlan;
-
-				Observer.OnSignallingConfig(RTCConfig);
-			});
-
-			RegisterHandler("offer", [this](FJsonObjectPtr JsonMsg) {
-				webrtc::SdpType Type = JsonMsg->GetStringField(TEXT("type")) == TEXT("offer") ? webrtc::SdpType::kOffer : webrtc::SdpType::kAnswer;
-
-				FString Sdp;
-				if (!JsonMsg->TryGetStringField(TEXT("sdp"), Sdp))
-				{
-					UE_LOG(LogMockPixelStreamingSS, Error, TEXT("Cannot find `sdp` in Streamer's answer\n%s"), *UE::PixelStreaming::ToString(JsonMsg));
-					return;
-				}
-
-				FPixelStreamingPlayerId PlayerId;
-				bool bGotPlayerId = GetPlayerIdJson(JsonMsg, PlayerId);
-				if (!bGotPlayerId)
-				{
-					Observer.OnSignallingSessionDescription(Type, Sdp);
-				}
-				else
-				{
-					Observer.OnSignallingSessionDescription(PlayerId, Type, Sdp);
-				}
-			});
-			RegisterHandler("answer", [this](FJsonObjectPtr JsonMsg) {
-				webrtc::SdpType Type = JsonMsg->GetStringField(TEXT("type")) == TEXT("offer") ? webrtc::SdpType::kOffer : webrtc::SdpType::kAnswer;
-
-				FString Sdp;
-				if (!JsonMsg->TryGetStringField(TEXT("sdp"), Sdp))
-				{
-					UE_LOG(LogMockPixelStreamingSS, Error, TEXT("Cannot find `sdp` in Streamer's answer\n%s"), *UE::PixelStreaming::ToString(JsonMsg));
-					return;
-				}
-
-				FPixelStreamingPlayerId PlayerId;
-				bool bGotPlayerId = GetPlayerIdJson(JsonMsg, PlayerId);
-				if (!bGotPlayerId)
-				{
-					Observer.OnSignallingSessionDescription(Type, Sdp);
-				}
-				else
-				{
-					Observer.OnSignallingSessionDescription(PlayerId, Type, Sdp);
-				}
-			});
-			RegisterHandler("iceCandidate", [this](FJsonObjectPtr JsonMsg) {
-				FPixelStreamingPlayerId PlayerId;
-				bool bGotPlayerId = GetPlayerIdJson(JsonMsg, PlayerId);
-
-				const FJsonObjectPtr* CandidateJson;
-				if (!JsonMsg->TryGetObjectField(TEXT("candidate"), CandidateJson))
-				{
-					PlayerError(PlayerId, TEXT("Failed to get `candiate` from remote `iceCandidate` message\n%s"), *UE::PixelStreaming::ToString(JsonMsg));
-				}
-
-				FString SdpMid;
-				if (!(*CandidateJson)->TryGetStringField(TEXT("sdpMid"), SdpMid))
-				{
-					PlayerError(PlayerId, TEXT("Failed to get `sdpMid` from remote `iceCandidate` message\n%s"), *UE::PixelStreaming::ToString(JsonMsg));
-				}
-
-				int32 SdpMLineIndex = 0;
-				if (!(*CandidateJson)->TryGetNumberField(TEXT("sdpMlineIndex"), SdpMLineIndex))
-				{
-					PlayerError(PlayerId, TEXT("Failed to get `sdpMlineIndex` from remote `iceCandidate` message\n%s"), *UE::PixelStreaming::ToString(JsonMsg));
-				}
-
-				FString CandidateStr;
-				if (!(*CandidateJson)->TryGetStringField(TEXT("candidate"), CandidateStr))
-				{
-					PlayerError(PlayerId, TEXT("Failed to get `candidate` from remote `iceCandidate` message\n%s"), *UE::PixelStreaming::ToString(JsonMsg));
-				}
-
-				if (bGotPlayerId)
-				{
-					Observer.OnSignallingRemoteIceCandidate(PlayerId, SdpMid, SdpMLineIndex, CandidateStr);
-				}
-				else
-				{
-					Observer.OnSignallingRemoteIceCandidate(SdpMid, SdpMLineIndex, CandidateStr);
-				}
-			});
-			RegisterHandler("ping", [this](FJsonObjectPtr JsonMsg) { /* nothing */ });
-			RegisterHandler("pong", [this](FJsonObjectPtr JsonMsg) { /* nothing */ });
-			RegisterHandler("playerCount", [this](FJsonObjectPtr JsonMsg) {
-				uint32 Count;
-				if (!JsonMsg->TryGetNumberField(TEXT("count"), Count))
-				{
-					UE_LOG(LogMockPixelStreamingSS, Error, TEXT("Failed to get `count` from `playerCount` message\n%s"), *UE::PixelStreaming::ToString(JsonMsg));
-					return;
-				}
-
-				Observer.OnSignallingPlayerCount(Count);
-			});
-			RegisterHandler("playerConnected", [this](FJsonObjectPtr JsonMsg) {
-				FPixelStreamingPlayerId PlayerId;
-				bool bGotPlayerId = GetPlayerIdJson(JsonMsg, PlayerId);
-				if (!bGotPlayerId)
-				{
-					UE_LOG(LogMockPixelStreamingSS, Error, TEXT("Failed to get `playerId` from `join` message\n%s"), *UE::PixelStreaming::ToString(JsonMsg));
-					return;
-				}
-
-				UE_LOG(LogMockPixelStreamingSS, Log, TEXT("Got player connected, player id=%s"), *PlayerId);
-
-				FPixelStreamingPlayerConfig PlayerConfig;
-
-				// Default to always making datachannel, unless explicitly set to false
-				JsonMsg->TryGetBoolField(TEXT("dataChannel"), PlayerConfig.SupportsDataChannel);
-
-				// Default peer is not an SFU, unless explictly set as SFU
-				JsonMsg->TryGetBoolField(TEXT("sfu"), PlayerConfig.IsSFU);
-
-				// Default to always sending an offer, unless explicitly set to false
-				bool bSendOffer = true;
-				JsonMsg->TryGetBoolField(TEXT("sendOffer"), bSendOffer);
-
-				Observer.OnSignallingPlayerConnected(PlayerId, PlayerConfig, bSendOffer);
-			});
-			RegisterHandler("playerDisconnected", [this](FJsonObjectPtr JsonMsg) {
-				FPixelStreamingPlayerId PlayerId;
-				bool bGotPlayerId = GetPlayerIdJson(JsonMsg, PlayerId);
-				if (!bGotPlayerId)
-				{
-					UE_LOG(LogMockPixelStreamingSS, Error, TEXT("Failed to get `playerId` from `playerDisconnected` message\n%s"), *UE::PixelStreaming::ToString(JsonMsg));
-					return;
-				}
-
-				Observer.OnSignallingPlayerDisconnected(PlayerId);
-			});
-			RegisterHandler("streamerDataChannels", [this](FJsonObjectPtr JsonMsg) {
-				FPixelStreamingPlayerId SFUId;
-				bool bSuccess = GetPlayerIdJson(JsonMsg, SFUId, TEXT("sfuId"));
-				if (!bSuccess)
-				{
-					UE_LOG(LogMockPixelStreamingSS, Error, TEXT("Failed to get `sfuId` from `streamerDataChannels` message\n%s"), *UE::PixelStreaming::ToString(JsonMsg));
-					return;
-				}
-
-				FPixelStreamingPlayerId PlayerId;
-				bSuccess = GetPlayerIdJson(JsonMsg, PlayerId, TEXT("playerId"));
-				if (!bSuccess)
-				{
-					UE_LOG(LogMockPixelStreamingSS, Error, TEXT("Failed to get `playerId` from `streamerDataChannels` message\n%s"), *UE::PixelStreaming::ToString(JsonMsg));
-					return;
-				}
-
-				int32 SendStreamId;
-				bSuccess = JsonMsg->TryGetNumberField(TEXT("sendStreamId"), SendStreamId);
-				if (!bSuccess)
-				{
-					UE_LOG(LogMockPixelStreamingSS, Error, TEXT("Failed to get `sendStreamId` from `streamerDataChannels` message\n%s"), *UE::PixelStreaming::ToString(JsonMsg));
-					return;
-				}
-
-				int32 RecvStreamId;
-				bSuccess = JsonMsg->TryGetNumberField(TEXT("recvStreamId"), RecvStreamId);
-				if (!bSuccess)
-				{
-					UE_LOG(LogMockPixelStreamingSS, Error, TEXT("Failed to get `recvStreamId` from `streamerDataChannels` message\n%s"), *UE::PixelStreaming::ToString(JsonMsg));
-					return;
-				}
-
-				Observer.OnSignallingSFUPeerDataChannels(SFUId, PlayerId, SendStreamId, RecvStreamId);
-			});
-			RegisterHandler("peerDataChannels", [this](FJsonObjectPtr JsonMsg) {
-				int32 SendStreamId = 0;
-				if (!JsonMsg->TryGetNumberField(TEXT("sendStreamId"), SendStreamId))
-				{
-					UE_LOG(LogMockPixelStreamingSS, Error, TEXT("Failed to get `sendStreamId` from remote `peerDataChannels` message\n%s"), *UE::PixelStreaming::ToString(JsonMsg));
-					return;
-				}
-				int32 RecvStreamId = 0;
-				if (!JsonMsg->TryGetNumberField(TEXT("recvStreamId"), RecvStreamId))
-				{
-					UE_LOG(LogMockPixelStreamingSS, Error, TEXT("Failed to get `recvStreamId` from remote `peerDataChannels` message\n%s"), *UE::PixelStreaming::ToString(JsonMsg));
-					return;
-				}
-				Observer.OnSignallingPeerDataChannels(SendStreamId, RecvStreamId);
-			});
-		}
-		virtual ~FMockSignallingConnection() = default;
-
-		virtual void TryConnect(FString URL) override
-		{
-			if (WebSocket && WebSocket->IsConnected())
-			{
-				UE_LOG(LogMockPixelStreamingSS, Log, TEXT("Skipping `Connect()` because we are already connected. If you want to reconnect then disconnect first."));
-				return;
-			}
-
-			// Reconnecting on an existing websocket can be problematic depending what state it was
-			// left in. Easier and safer to disconnect any existing socket/delegates and start fresh.
-			Disconnect();
-
-			Url = URL;
-
-			WebSocket = MakeShared<FMockWebSocket>();
-
-			OnConnectedHandle = WebSocket->OnConnected().AddLambda([this]() {
-				UE_LOG(LogMockPixelStreamingSS, Log, TEXT("Connected to SS %s"), *Url);
-				Observer.OnSignallingConnected();
-			});
-
-			OnConnectionErrorHandle = WebSocket->OnConnectionError().AddLambda([this](const FString& Error) {
-				UE_LOG(LogMockPixelStreamingSS, Log, TEXT("Failed to connect to SS %s - signalling server may not be up yet. Message: \"%s\""), *Url, *Error);
-				Observer.OnSignallingError(Error);
-			});
-
-			OnClosedHandle = WebSocket->OnClosed().AddLambda([this](int32 StatusCode, const FString& Reason, bool bWasClean) {
-				UE_LOG(LogMockPixelStreamingSS, Log, TEXT("Closed connection to SS %s - \n\tstatus %d\n\treason: %s\n\twas clean: %s"), *Url, StatusCode, *Reason, bWasClean ? TEXT("true") : TEXT("false"));
-				Observer.OnSignallingDisconnected(StatusCode, Reason, bWasClean);
-			});
-
-			OnMessageHandle = WebSocket->OnMessage().AddLambda([this](const FString& Msg) {
-				FJsonObjectPtr JsonMsg;
-				const auto JsonReader = TJsonReaderFactory<TCHAR>::Create(Msg);
-
-				if (!FJsonSerializer::Deserialize(JsonReader, JsonMsg))
-				{
-					UE_LOG(LogMockPixelStreamingSS, Error, TEXT("Failed to parse SS message:\n%s"), *Msg);
-					return;
-				}
-
-				FString MsgType;
-				if (!JsonMsg->TryGetStringField(TEXT("type"), MsgType))
-				{
-					UE_LOG(LogMockPixelStreamingSS, Error, TEXT("Cannot find `type` field in SS message:\n%s"), *Msg);
-					return;
-				}
-
-				TFunction<void(FJsonObjectPtr)>* Handler = MessageHandlers.Find(MsgType);
-				if (Handler != nullptr)
-				{
-					(*Handler)(JsonMsg);
-				}
-			});
-
-			UE_LOG(LogMockPixelStreamingSS, Log, TEXT("Connecting to SS %s"), *Url);
-			WebSocket->Connect();
-		}
-
-		virtual void Disconnect() override
-		{
-			if (!WebSocket)
-			{
-				return;
-			}
-
-			WebSocket->OnConnected().Remove(OnConnectedHandle);
-			WebSocket->OnConnectionError().Remove(OnConnectionErrorHandle);
-			WebSocket->OnClosed().Remove(OnClosedHandle);
-			WebSocket->OnMessage().Remove(OnMessageHandle);
-
-			WebSocket->Close();
-			WebSocket = nullptr;
-			UE_LOG(LogMockPixelStreamingSS, Log, TEXT("Closing websocket to SS %s"), *Url);
-		}
-		virtual bool IsConnected() const override
-		{
-			return WebSocket != nullptr && WebSocket.IsValid() && WebSocket->IsConnected();
-		}
-
-		virtual void SendOffer(FPixelStreamingPlayerId PlayerId, const webrtc::SessionDescriptionInterface& SDP) override
-		{
-			FJsonObjectPtr OfferJson = MakeShared<FJsonObject>();
-			OfferJson->SetStringField(TEXT("type"), TEXT("offer"));
-			SetPlayerIdJson(OfferJson, PlayerId);
-
-			std::string SdpAnsi;
-			SDP.ToString(&SdpAnsi);
-			FString SdpStr = UE::PixelStreaming::ToString(SdpAnsi);
-			OfferJson->SetStringField(TEXT("sdp"), SdpStr);
-
-#if WITH_EDITOR
-			// Default hover preference needs to be sent with the offer to prevent error messages on older infrastructure
-			OfferJson->SetStringField(TEXT("defaultToHover"), TEXT("true"));
-#endif // WITH_EDITOR
-
-			UE_LOG(LogMockPixelStreamingSS, Log, TEXT("Sending player=%s \"offer\" to SS %s"), *PlayerId, *Url);
-			UE_LOG(LogMockPixelStreamingSS, Verbose, TEXT("SDP offer\n%s"), *SdpStr);
-
-			SendMessage(UE::PixelStreaming::ToString(OfferJson, false));
-		}
-
-		virtual void SendAnswer(FPixelStreamingPlayerId PlayerId, const webrtc::SessionDescriptionInterface& SDP) override
-		{
-			FJsonObjectPtr AnswerJson = MakeShared<FJsonObject>();
-			AnswerJson->SetStringField(TEXT("type"), TEXT("answer"));
-			SetPlayerIdJson(AnswerJson, PlayerId);
-
-			std::string SdpAnsi;
-
-			if (SDP.ToString(&SdpAnsi))
-			{
-				FString SdpStr = UE::PixelStreaming::ToString(SdpAnsi);
-				AnswerJson->SetStringField(TEXT("sdp"), SdpStr);
-
-				UE_LOG(LogMockPixelStreamingSS, Log, TEXT("Sending player=%s \"answer\" to SS %s"), *PlayerId, *Url);
-				UE_LOG(LogMockPixelStreamingSS, Verbose, TEXT("SDP answer\n%s"), *SdpStr);
-
-				SendMessage(UE::PixelStreaming::ToString(AnswerJson, false));
-			}
-		}
-
-		virtual void SendIceCandidate(FPixelStreamingPlayerId PlayerId, const webrtc::IceCandidateInterface& IceCandidate) override
-		{
-			FJsonObjectPtr IceCandidateJson = MakeShared<FJsonObject>();
-
-			IceCandidateJson->SetStringField(TEXT("type"), TEXT("iceCandidate"));
-			SetPlayerIdJson(IceCandidateJson, PlayerId);
-
-			FJsonObjectPtr CandidateJson = MakeShared<FJsonObject>();
-			CandidateJson->SetStringField(TEXT("sdpMid"), IceCandidate.sdp_mid().c_str());
-			CandidateJson->SetNumberField(TEXT("sdpMLineIndex"), static_cast<double>(IceCandidate.sdp_mline_index()));
-
-			std::string CandidateAnsi;
-			if (IceCandidate.ToString(&CandidateAnsi))
-			{
-				FString CandidateStr = UE::PixelStreaming::ToString(CandidateAnsi);
-				CandidateJson->SetStringField(TEXT("candidate"), CandidateStr);
-				IceCandidateJson->SetObjectField(TEXT("candidate"), CandidateJson);
-				UE_LOG(LogMockPixelStreamingSS, Verbose, TEXT("-> SS: iceCandidate\n%s"), *UE::PixelStreaming::ToString(IceCandidateJson));
-				SendMessage(UE::PixelStreaming::ToString(IceCandidateJson, false));
-			}
-		}
-
-		virtual void SendDisconnectPlayer(FPixelStreamingPlayerId PlayerId, const FString& Reason) override
-		{
-			FJsonObjectPtr JsonMsg = MakeShared<FJsonObject>();
-
-			JsonMsg->SetStringField(TEXT("type"), TEXT("disconnectPlayer"));
-			SetPlayerIdJson(JsonMsg, PlayerId);
-			JsonMsg->SetStringField(TEXT("reason"), Reason);
-
-			FString Msg = UE::PixelStreaming::ToString(JsonMsg, false);
-			UE_LOG(LogMockPixelStreamingSS, Verbose, TEXT("Sending player=%s \"disconnectPlayer\" to SS %s"), *PlayerId, *Url);
-
-			SendMessage(Msg);
-		}
-
-		virtual void RequestStreamerList() override
-		{
-			FJsonObjectPtr Json = MakeShared<FJsonObject>();
-			Json->SetStringField(TEXT("type"), TEXT("streamerList"));
-
-			const FString JsonStr = UE::PixelStreaming::ToString(Json, false);
-			UE_LOG(LogMockPixelStreamingSS, Verbose, TEXT("-> SS: streamerList\n%s"), *JsonStr);
-
-			SendMessage(JsonStr);
-		}
-
-		virtual void SendSubscribe(const FString& ToStreamerId) override
-		{
-			FJsonObjectPtr SubscribeJson = MakeShared<FJsonObject>();
-			SubscribeJson->SetStringField(TEXT("type"), TEXT("subscribe"));
-			SubscribeJson->SetStringField(TEXT("streamerId"), ToStreamerId);
-
-			const FString SubscribeJsonStr = UE::PixelStreaming::ToString(SubscribeJson, false);
-			UE_LOG(LogMockPixelStreamingSS, Verbose, TEXT("-> SS: subscribe\n%s"), *SubscribeJsonStr);
-
-			SendMessage(SubscribeJsonStr);
-		}
-
-		virtual void SendUnsubscribe() override
-		{
-			FJsonObjectPtr SubscribeJson = MakeShared<FJsonObject>();
-			SubscribeJson->SetStringField(TEXT("type"), TEXT("unsubscribe"));
-
-			const FString SubscribeJsonStr = UE::PixelStreaming::ToString(SubscribeJson, false);
-			UE_LOG(LogMockPixelStreamingSS, Verbose, TEXT("-> SS: subscribe\n%s"), *SubscribeJsonStr);
-
-			SendMessage(SubscribeJsonStr);
-		}
-
-		virtual void SendOffer(const webrtc::SessionDescriptionInterface& SDP) override
-		{
-			FJsonObjectPtr OfferJson = MakeShared<FJsonObject>();
-			OfferJson->SetStringField(TEXT("type"), TEXT("offer"));
-
-			FString SdpStr = UE::PixelStreaming::ToString(SDP);
-			OfferJson->SetStringField(TEXT("sdp"), UE::PixelStreaming::ToString(SDP));
-
-			UE_LOG(LogMockPixelStreamingSS, Verbose, TEXT("-> SS: offer\n%s"), *SdpStr);
-
-			SendMessage(UE::PixelStreaming::ToString(OfferJson, false));
-		}
-
-		virtual void SendAnswer(const webrtc::SessionDescriptionInterface& SDP) override
-		{
-			FJsonObjectPtr AnswerJson = MakeShared<FJsonObject>();
-			AnswerJson->SetStringField(TEXT("type"), TEXT("answer"));
-
-			FString SdpStr = UE::PixelStreaming::ToString(SDP);
-			AnswerJson->SetStringField(TEXT("sdp"), UE::PixelStreaming::ToString(SDP));
-
-			UE_LOG(LogMockPixelStreamingSS, Verbose, TEXT("-> SS: answer\n%s"), *SdpStr);
-
-			SendMessage(UE::PixelStreaming::ToString(AnswerJson, false));
-		}
-
-		virtual void SendIceCandidate(const webrtc::IceCandidateInterface& IceCandidate) override
-		{
-			FJsonObjectPtr IceCandidateJson = MakeShared<FJsonObject>();
-
-			IceCandidateJson->SetStringField(TEXT("type"), TEXT("iceCandidate"));
-
-			FJsonObjectPtr CandidateJson = MakeShared<FJsonObject>();
-			CandidateJson->SetStringField(TEXT("sdpMid"), IceCandidate.sdp_mid().c_str());
-			CandidateJson->SetNumberField(TEXT("sdpMLineIndex"), static_cast<double>(IceCandidate.sdp_mline_index()));
-
-			std::string CandidateAnsi;
-			verifyf(IceCandidate.ToString(&CandidateAnsi), TEXT("Failed to serialize IceCandidate"));
-			FString CandidateStr = UE::PixelStreaming::ToString(CandidateAnsi);
-			CandidateJson->SetStringField(TEXT("candidate"), CandidateStr);
-
-			IceCandidateJson->SetObjectField(TEXT("candidate"), CandidateJson);
-
-			UE_LOG(LogMockPixelStreamingSS, Verbose, TEXT("-> SS: ice-candidate\n%s"), *UE::PixelStreaming::ToString(IceCandidateJson));
-
-			SendMessage(UE::PixelStreaming::ToString(IceCandidateJson, false));
-		}
-
-		virtual void SetKeepAlive(bool bKeepAlive) override
-		{
-			if (bKeepAliveEnabled == bKeepAlive)
-			{
-				return;
-			}
-
-			bKeepAliveEnabled = bKeepAlive;
-		}
-
-		virtual void SetAutoReconnect(bool bAutoReconnect) override
-		{
-			if (bAutoReconnectEnabled == bAutoReconnect)
-			{
-				return;
-			}
-
-			bAutoReconnectEnabled = bAutoReconnect;
-		}
-
-		virtual void SendMessage(const FString& Msg) override
-		{
-			if (WebSocket && WebSocket->IsConnected())
-			{
-				WebSocket->Send(Msg);
-			}
-		}
-
-		TSharedPtr<FMockWebSocket> WebSocket;
-
-	private:
-		using FJsonObjectPtr = TSharedPtr<FJsonObject>;
-		void RegisterHandler(const FString& messageType, const TFunction<void(FJsonObjectPtr)>& handler)
-		{
-			MessageHandlers.Add(messageType, handler);
-		}
-
-		bool GetPlayerIdJson(const FJsonObjectPtr& Json, FPixelStreamingPlayerId& OutPlayerId, const FString& FieldId = TEXT("playerId"))
-		{
-			uint32 PlayerIdInt;
-			if (Json->TryGetNumberField(FieldId, PlayerIdInt))
-			{
-				OutPlayerId = ToPlayerId(PlayerIdInt);
-				return true;
-			}
-
-			return false;
-		}
-
-		void SetPlayerIdJson(FJsonObjectPtr& JsonObject, FPixelStreamingPlayerId PlayerId)
-		{
-			int32 PlayerIdAsInt = PlayerIdToInt(PlayerId);
-			JsonObject->SetNumberField(TEXT("playerId"), PlayerIdAsInt);
-		}
-
-		void PlayerError(FPixelStreamingPlayerId PlayerId, const FString& Msg)
-		{
-			UE_LOG(LogMockPixelStreamingSS, Error, TEXT("player %s: %s"), *PlayerId, *Msg);
-			SendDisconnectPlayer(PlayerId, Msg);
-		}
-
-		template <typename FmtType, typename... T>
-		void PlayerError(FPixelStreamingPlayerId PlayerId, const FmtType& Msg, T... args)
-		{
-			const FString FormattedMsg = FString::Printf(Msg, args...);
-			PlayerError(PlayerId, FormattedMsg);
-		}
-
-		IPixelStreamingSignallingConnectionObserver& Observer;
-		FString StreamerId;
-		FString Url;
-
-		FDelegateHandle OnConnectedHandle;
-		FDelegateHandle OnConnectionErrorHandle;
-		FDelegateHandle OnClosedHandle;
-		FDelegateHandle OnMessageHandle;
-
-		bool bAutoReconnectEnabled = true;
-		bool bKeepAliveEnabled = true;
-		/** Handle for efficient management of KeepAlive timer */
-		FTimerHandle TimerHandle_KeepAlive;
-		FTimerHandle TimerHandle_Reconnect;
-		const float KEEP_ALIVE_INTERVAL = 60.0f;
-
-		TMap<FString, TFunction<void(FJsonObjectPtr)>> MessageHandlers;
-	};
-
 	class FMockSessionDescriptionInterface : public webrtc::SessionDescriptionInterface
 	{
 	public:
@@ -740,8 +148,9 @@ namespace
 	};
 
 	BEGIN_DEFINE_SPEC(FSignallingServerConnectionSpec, "System.Plugins.PixelStreaming.SignallingServerConnection", EAutomationTestFlags::EditorContext | EAutomationTestFlags::ClientContext | EAutomationTestFlags::ProductFilter)
-	TUniquePtr<FMockSignallingConnection> SignallingServerConnection; // the object being tested
-	FTestSSConnectionObserver Observer;
+	TUniquePtr<FPixelStreamingSignallingConnection> SignallingServerConnection; // the object being tested
+	TSharedPtr<FMockConnectionObserver> Observer;
+	TSharedPtr<FMockWebSocket> WebSocket;
 	FDelegateHandle OnMessageDelegateHandle; // so we can remove the delegate after tests
 	TSharedPtr<FJsonObject> LastJsonMsg;	 // the parsed JsonMsg of the last message received from the SignallingServerConnection
 	const FString ServerID = "MOCK SERVER";
@@ -754,20 +163,24 @@ namespace
 	{
 		Describe("Mocked connection tests", [this]() {
 			BeforeEach([this]() {
-				SignallingServerConnection = MakeUnique<FMockSignallingConnection>(Observer, ServerID);
+				Observer = MakeShared<FMockConnectionObserver>();
+				WebSocket = MakeShared<FMockWebSocket>();
+
+				SignallingServerConnection = MakeUnique<FPixelStreamingSignallingConnection>(Observer, ServerID, WebSocket);
 				SignallingServerConnection->TryConnect("Fake Url");
-				OnMessageDelegateHandle = SignallingServerConnection->WebSocket->OnMockResponseEvent.AddLambda([this](const FString& message) {
+				OnMessageDelegateHandle = WebSocket->OnMessageSent().AddLambda([this](const FString& message) {
 					OnMessageReceived(message);
 				});
 			});
 
 			AfterEach([this]() {
 				LastJsonMsg.Reset();
-				SignallingServerConnection->WebSocket->OnMockResponseEvent.Remove(OnMessageDelegateHandle);
+				WebSocket->OnMessageSent().Remove(OnMessageDelegateHandle);
 				SignallingServerConnection.Reset();
 			});
 
-			It("should be connected", [this]() {
+			It("should be connected", [this]() 
+			{
 				TestTrue("Connected", SignallingServerConnection->IsConnected());
 			});
 
@@ -803,73 +216,80 @@ namespace
 			});
 
 			It("should reply to an identity request", [this]() {
-				SignallingServerConnection->WebSocket->MockSend(R"({"type" : "identify"})");
+				// Simulate the signalling server sending this message to the application
+				WebSocket->OnMessage().Broadcast(R"({"type" : "identify"})");
 				TestEqual("Type", LastJsonMsg->GetStringField(TEXT("type")), "endpointId");
 				TestEqual("Id", LastJsonMsg->GetStringField(TEXT("id")), ServerID);
 			});
 
 			It("should receive configurations", [this]() {
-				SignallingServerConnection->WebSocket->MockSend(R"({ "type": "config", "peerConnectionOptions": {} })");
-				TestTrue("Config IsSet", Observer.SetConfig.IsSet());
+				// Simulate the signalling server sending this message to the application
+				WebSocket->OnMessage().Broadcast(R"({ "type": "config", "peerConnectionOptions": {} })");
+				TestTrue("Config IsSet", Observer->SetConfig.IsSet());
 			});
 
 			It("should be able to receive offers with no player id", [this]() {
-				SignallingServerConnection->WebSocket->MockSend(R"({ "type": "offer", "sdp": "sdp goes here" })");
-				if (TestTrue("SDP IsSet", Observer.SetSDP.IsSet()))
+				// Simulate the signalling server sending this message to the application
+				WebSocket->OnMessage().Broadcast(R"({ "type": "offer", "sdp": "sdp goes here" })");
+				if (TestTrue("SDP IsSet", Observer->SetSDP.IsSet()))
 				{
-					TestEqual("Offer type", Observer.SetSDP.GetValue().Type, webrtc::SdpType::kOffer);
-					TestFalse("Player id IsSet", Observer.SetSDP.GetValue().PlayerId.IsSet());
-					TestEqual("SDP value", Observer.SetSDP.GetValue().SDP, "sdp goes here");
+					TestEqual("Offer type", Observer->SetSDP.GetValue().Type, webrtc::SdpType::kOffer);
+					TestFalse("Player id IsSet", Observer->SetSDP.GetValue().PlayerId.IsSet());
+					TestEqual("SDP value", Observer->SetSDP.GetValue().SDP, "sdp goes here");
 				}
 			});
 
 			It("should be able to receive offers with a player id", [this]() {
-				SignallingServerConnection->WebSocket->MockSend(R"({ "type": "offer", "playerId": "101", "sdp": "sdp goes here" })");
-				if (TestTrue("SDP IsSet", Observer.SetSDP.IsSet()))
+				// Simulate the signalling server sending this message to the application
+				WebSocket->OnMessage().Broadcast(R"({ "type": "offer", "playerId": "101", "sdp": "sdp goes here" })");
+				if (TestTrue("SDP IsSet", Observer->SetSDP.IsSet()))
 				{
-					TestEqual("Offer type", Observer.SetSDP.GetValue().Type, webrtc::SdpType::kOffer);
-					if (TestTrue("Player id IsSet", Observer.SetSDP.GetValue().PlayerId.IsSet()))
+					TestEqual("Offer type", Observer->SetSDP.GetValue().Type, webrtc::SdpType::kOffer);
+					if (TestTrue("Player id IsSet", Observer->SetSDP.GetValue().PlayerId.IsSet()))
 					{
-						TestEqual("Player id", Observer.SetSDP.GetValue().PlayerId.GetValue(), "101");
+						TestEqual("Player id", Observer->SetSDP.GetValue().PlayerId.GetValue(), "101");
 					}
-					TestEqual("SDP value", Observer.SetSDP.GetValue().SDP, "sdp goes here");
+					TestEqual("SDP value", Observer->SetSDP.GetValue().SDP, "sdp goes here");
 				}
 			});
 
 			It("should be able to receive an answer", [this]() {
-				SignallingServerConnection->WebSocket->MockSend(R"({ "type": "answer", "playerId": "102", "sdp": "sdp goes here" })");
-				if (TestTrue("SDP IsSet", Observer.SetSDP.IsSet()))
+				// Simulate the signalling server sending this message to the application
+				WebSocket->OnMessage().Broadcast(R"({ "type": "answer", "playerId": "102", "sdp": "sdp goes here" })");
+				if (TestTrue("SDP IsSet", Observer->SetSDP.IsSet()))
 				{
-					TestEqual("Answer type", Observer.SetSDP.GetValue().Type, webrtc::SdpType::kAnswer);
-					TestEqual("SDP value", Observer.SetSDP.GetValue().SDP, "sdp goes here");
-					if (TestTrue("Player id IsSet", Observer.SetSDP.GetValue().PlayerId.IsSet()))
+					TestEqual("Answer type", Observer->SetSDP.GetValue().Type, webrtc::SdpType::kAnswer);
+					TestEqual("SDP value", Observer->SetSDP.GetValue().SDP, "sdp goes here");
+					if (TestTrue("Player id IsSet", Observer->SetSDP.GetValue().PlayerId.IsSet()))
 					{
-						TestEqual("Player id", Observer.SetSDP.GetValue().PlayerId.GetValue(), "102");
+						TestEqual("Player id", Observer->SetSDP.GetValue().PlayerId.GetValue(), "102");
 					}
 				}
 			});
 
 			It("should be able to receive an ice candidate with no player id", [this]() {
-				SignallingServerConnection->WebSocket->MockSend(R"({ "type": "iceCandidate", "candidate": { "sdpMid": "mid", "sdpMlineIndex": 10, "candidate": "sdp string" } })");
-				if (TestTrue("Ice Candidate IsSet", Observer.SetIceCandidate.IsSet()))
+				// Simulate the signalling server sending this message to the application
+				WebSocket->OnMessage().Broadcast(R"({ "type": "iceCandidate", "candidate": { "sdpMid": "mid", "sdpMlineIndex": 10, "candidate": "sdp string" } })");
+				if (TestTrue("Ice Candidate IsSet", Observer->SetIceCandidate.IsSet()))
 				{
-					TestEqual("sdp mid", Observer.SetIceCandidate.GetValue().SdpMid, "mid");
-					TestEqual("sdp mline index", Observer.SetIceCandidate.GetValue().SdpMLineIndex, 10);
-					TestEqual("sdp value", Observer.SetIceCandidate.GetValue().SDP, "sdp string");
-					TestFalse("PlayerId IsSet", Observer.SetIceCandidate.GetValue().PlayerId.IsSet());
+					TestEqual("sdp mid", Observer->SetIceCandidate.GetValue().SdpMid, "mid");
+					TestEqual("sdp mline index", Observer->SetIceCandidate.GetValue().SdpMLineIndex, 10);
+					TestEqual("sdp value", Observer->SetIceCandidate.GetValue().SDP, "sdp string");
+					TestFalse("PlayerId IsSet", Observer->SetIceCandidate.GetValue().PlayerId.IsSet());
 				}
 			});
 
 			It("should be able to receive an ice candidate with a player id", [this]() {
-				SignallingServerConnection->WebSocket->MockSend(R"({ "type": "iceCandidate", "playerId": "103", "candidate": { "sdpMid": "mid", "sdpMlineIndex": 10, "candidate": "sdp string" } })");
-				if (TestTrue("Ice Candidate IsSet", Observer.SetIceCandidate.IsSet()))
+				// Simulate the signalling server sending this message to the application
+				WebSocket->OnMessage().Broadcast(R"({ "type": "iceCandidate", "playerId": "103", "candidate": { "sdpMid": "mid", "sdpMlineIndex": 10, "candidate": "sdp string" } })");
+				if (TestTrue("Ice Candidate IsSet", Observer->SetIceCandidate.IsSet()))
 				{
-					TestEqual("sdp mid", Observer.SetIceCandidate.GetValue().SdpMid, "mid");
-					TestEqual("sdp mline index", Observer.SetIceCandidate.GetValue().SdpMLineIndex, 10);
-					TestEqual("sdp value", Observer.SetIceCandidate.GetValue().SDP, "sdp string");
-					if (TestTrue("Player id IsSet", Observer.SetIceCandidate.GetValue().PlayerId.IsSet()))
+					TestEqual("sdp mid", Observer->SetIceCandidate.GetValue().SdpMid, "mid");
+					TestEqual("sdp mline index", Observer->SetIceCandidate.GetValue().SdpMLineIndex, 10);
+					TestEqual("sdp value", Observer->SetIceCandidate.GetValue().SDP, "sdp string");
+					if (TestTrue("Player id IsSet", Observer->SetIceCandidate.GetValue().PlayerId.IsSet()))
 					{
-						TestEqual("Player id", Observer.SetIceCandidate.GetValue().PlayerId.GetValue(), "103");
+						TestEqual("Player id", Observer->SetIceCandidate.GetValue().PlayerId.GetValue(), "103");
 					}
 				}
 			});
diff --git a/Engine/Plugins/Media/PixelStreaming/Source/PixelStreaming/Public/PixelStreamingSignallingConnection.h b/Engine/Plugins/Media/PixelStreaming/Source/PixelStreaming/Public/PixelStreamingSignallingConnection.h
--- a/Engine/Plugins/Media/PixelStreaming/Source/PixelStreaming/Public/PixelStreamingSignallingConnection.h
+++ b/Engine/Plugins/Media/PixelStreaming/Source/PixelStreaming/Public/PixelStreamingSignallingConnection.h
@@ -17,7 +17,7 @@ class IWebSocket;
 class PIXELSTREAMING_API FPixelStreamingSignallingConnection : public IPixelStreamingSignallingConnection
 {
 public:
-	FPixelStreamingSignallingConnection(TSharedPtr<IPixelStreamingSignallingConnectionObserver> InObserver, FString InStreamerId = "");
+	FPixelStreamingSignallingConnection(TSharedPtr<IPixelStreamingSignallingConnectionObserver> InObserver, FString InStreamerId = "", TSharedPtr<IWebSocket> InWebSocket = nullptr);
 	virtual ~FPixelStreamingSignallingConnection();
 
 	/* IPixelStreamingSignallingConnection Interface */
