Deprecating kick player functionality

diff --git a/Engine/Plugins/Media/PixelStreaming/Source/PixelStreaming/Private/PixelStreamingBlueprints.cpp b/Engine/Plugins/Media/PixelStreaming/Source/PixelStreaming/Private/PixelStreamingBlueprints.cpp
--- a/Engine/Plugins/Media/PixelStreaming/Source/PixelStreaming/Private/PixelStreamingBlueprints.cpp
+++ b/Engine/Plugins/Media/PixelStreaming/Source/PixelStreaming/Private/PixelStreamingBlueprints.cpp
@@ -177,7 +177,9 @@ void UPixelStreamingBlueprints::KickPlayer(FString PlayerId)
 	{
 		return;
 	}
+	PRAGMA_DISABLE_DEPRECATION_WARNINGS
 	Streamer->KickPlayer(ToPlayerId(PlayerId));
+	PRAGMA_ENABLE_DEPRECATION_WARNINGS
 }
 
 void UPixelStreamingBlueprints::StreamerKickPlayer(FString StreamerId, FString PlayerId)
@@ -192,7 +194,9 @@ void UPixelStreamingBlueprints::StreamerKickPlayer(FString StreamerId, FString P
 	{
 		return;
 	}
+	PRAGMA_DISABLE_DEPRECATION_WARNINGS
 	Streamer->KickPlayer(ToPlayerId(PlayerId));
+	PRAGMA_ENABLE_DEPRECATION_WARNINGS
 }
 
 void UPixelStreamingBlueprints::SetPlayerLayerPreference(FString PlayerId, int SpatialLayerId, int TemporalLayerId)
diff --git a/Engine/Plugins/Media/PixelStreaming/Source/PixelStreaming/Private/PixelStreamingBlueprints.h b/Engine/Plugins/Media/PixelStreaming/Source/PixelStreaming/Private/PixelStreamingBlueprints.h
--- a/Engine/Plugins/Media/PixelStreaming/Source/PixelStreaming/Private/PixelStreamingBlueprints.h
+++ b/Engine/Plugins/Media/PixelStreaming/Source/PixelStreaming/Private/PixelStreamingBlueprints.h
@@ -98,18 +98,18 @@ public:
 	static void StreamerUnfreezeStream(FString StreamerId);
 
 	/**
-	 * Kick a player.
+	 * Kick a player. DEPRECATED
 	 * @param   PlayerId         The ID of the player to kick.
 	 */
-	UFUNCTION(BlueprintCallable, Category = "Pixel Streaming")
+	UFUNCTION(BlueprintCallable, Category = "Pixel Streaming", meta=(DeprecatedFunction, DeprecationMessage = "Kick player has been deprecated, will no longer function and will be removed in future versions."))
 	static void KickPlayer(FString PlayerId);
 
 	/**
-	 * Kick a player.
+	 * Kick a player. DEPRECATED
 	 * @param	StreamerId		The streamer which the player belongs
 	 * @param   PlayerId        The ID of the player to kick.
 	 */
-	UFUNCTION(BlueprintCallable, Category = "Pixel Streaming")
+	UFUNCTION(BlueprintCallable, Category = "Pixel Streaming", meta=(DeprecatedFunction, DeprecationMessage = "Kick player has been deprecated, will no longer function and will be removed in future versions."))
 	static void StreamerKickPlayer(FString StreamerId, FString PlayerId);
 
 	/**
diff --git a/Engine/Plugins/Media/PixelStreaming/Source/PixelStreaming/Public/IPixelStreamingStreamer.h b/Engine/Plugins/Media/PixelStreaming/Source/PixelStreaming/Public/IPixelStreamingStreamer.h
--- a/Engine/Plugins/Media/PixelStreaming/Source/PixelStreaming/Public/IPixelStreamingStreamer.h
+++ b/Engine/Plugins/Media/PixelStreaming/Source/PixelStreaming/Public/IPixelStreamingStreamer.h
@@ -218,9 +218,10 @@ public:
 	virtual void SendFileData(const TArray64<uint8>& ByteData, FString& MimeType, FString& FileExtension) = 0;
 
 	/**
-	 * @brief Kick a player by player id.
+	 * @brief Kick a player by player id. DEPRECATED
 	 * @param PlayerId		- The ID of the player to kick
 	 */
+	UE_DEPRECATED(5.4, "Kick player has been deprecated, will no longer function and will be removed in future versions.")
 	virtual void KickPlayer(FPixelStreamingPlayerId PlayerId) = 0;
 
 	/**
