Fix: Refactor signalling servers to better reflect cirrus's functionality. Also update tests to respect the proper connection flow

diff --git a/Engine/Plugins/Media/PixelStreaming/Source/PixelStreaming/Private/PixelStreamingSignallingConnection.cpp b/Engine/Plugins/Media/PixelStreaming/Source/PixelStreaming/Private/PixelStreamingSignallingConnection.cpp
--- a/Engine/Plugins/Media/PixelStreaming/Source/PixelStreaming/Private/PixelStreamingSignallingConnection.cpp
+++ b/Engine/Plugins/Media/PixelStreaming/Source/PixelStreaming/Private/PixelStreamingSignallingConnection.cpp
@@ -98,6 +98,7 @@ void FPixelStreamingSignallingConnection::Connect(FString InUrl, bool bIsReconne
 	}
 
 	WebSocket->Connect();
+	bIsConnected = true;
 }
 
 void FPixelStreamingSignallingConnection::TryConnect(FString InUrl)
@@ -113,7 +114,7 @@ void FPixelStreamingSignallingConnection::Disconnect()
 		StopReconnectTimer();
 	}
 
-	if (!WebSocket)
+	if (!WebSocket || !bIsConnected)
 	{
 		return;
 	}
@@ -126,6 +127,8 @@ void FPixelStreamingSignallingConnection::Disconnect()
 
 	WebSocket->Close();
 	UE_LOG(LogPixelStreamingSS, Log, TEXT("Closing websocket to SS %s"), *Url);
+
+	bIsConnected = false;
 }
 
 bool FPixelStreamingSignallingConnection::IsConnected() const
diff --git a/Engine/Plugins/Media/PixelStreaming/Source/PixelStreaming/Private/Tests/MockPlayer.h b/Engine/Plugins/Media/PixelStreaming/Source/PixelStreaming/Private/Tests/MockPlayer.h
--- a/Engine/Plugins/Media/PixelStreaming/Source/PixelStreaming/Private/Tests/MockPlayer.h
+++ b/Engine/Plugins/Media/PixelStreaming/Source/PixelStreaming/Private/Tests/MockPlayer.h
@@ -91,6 +91,10 @@ namespace UE::PixelStreaming
 		void Connect(int Port)
 		{
 			FString Url = FString::Printf(TEXT("ws://127.0.0.1:%d"), Port);
+			if(Mode == EMode::CreateOffers)
+			{
+				Url += "/?OfferToReceive=true";
+			}
 			SignallingServerConnection->TryConnect(Url);
 		}
 
@@ -159,6 +163,10 @@ namespace UE::PixelStreaming
 					[](const FString& Error) {
 					});
 			}
+			else if (Mode == EMode::AcceptOffers)
+			{
+				SignallingServerConnection->RequestStreamerList();
+			}
 		}
 
 		virtual void OnSignallingSessionDescription(webrtc::SdpType Type, const FString& Sdp) override
@@ -200,7 +208,15 @@ namespace UE::PixelStreaming
 		virtual void OnSignallingPlayerConnected(FPixelStreamingPlayerId PlayerId, const FPixelStreamingPlayerConfig& PlayerConfig, bool bSendOffer) override {}
 		virtual void OnSignallingPlayerDisconnected(FPixelStreamingPlayerId PlayerId) override {}
 		virtual void OnSignallingSFUPeerDataChannels(FPixelStreamingPlayerId SFUId, FPixelStreamingPlayerId PlayerId, int32 SendStreamId, int32 RecvStreamId) override {}
-		virtual void OnSignallingStreamerList(const TArray<FString>& StreamerList) override {}
+		virtual void OnSignallingStreamerList(const TArray<FString>& StreamerList) override 
+		{
+			if(StreamerList.Num() != 1)
+			{
+				UE_LOG(LogPixelStreaming, Error, TEXT("Invalid number streamers connected (%d). Expected just one!"), StreamerList.Num());
+			}
+
+			SignallingServerConnection->SendSubscribe(StreamerList[0]);
+		}
 		virtual void OnPlayerRequestsBitrate(FPixelStreamingPlayerId PlayerId, int MinBitrate, int MaxBitrate) override {}
 
 		DECLARE_MULTICAST_DELEGATE(FOnConnectionEstablished);
diff --git a/Engine/Plugins/Media/PixelStreaming/Source/PixelStreaming/Public/PixelStreamingSignallingConnection.h b/Engine/Plugins/Media/PixelStreaming/Source/PixelStreaming/Public/PixelStreamingSignallingConnection.h
--- a/Engine/Plugins/Media/PixelStreaming/Source/PixelStreaming/Public/PixelStreamingSignallingConnection.h
+++ b/Engine/Plugins/Media/PixelStreaming/Source/PixelStreaming/Public/PixelStreamingSignallingConnection.h
@@ -94,6 +94,7 @@ private:
 
 	bool bAutoReconnectEnabled = true;
 	bool bKeepAliveEnabled = true;
+	bool bIsConnected = false;
 	/** Handle for efficient management of KeepAlive timer */
 	FTimerHandle TimerHandle_KeepAlive;
 	FTimerHandle TimerHandle_Reconnect;
diff --git a/Engine/Plugins/Media/PixelStreaming/Source/PixelStreamingServers/Private/SignallingServer.cpp b/Engine/Plugins/Media/PixelStreaming/Source/PixelStreamingServers/Private/SignallingServer.cpp
--- a/Engine/Plugins/Media/PixelStreaming/Source/PixelStreamingServers/Private/SignallingServer.cpp
+++ b/Engine/Plugins/Media/PixelStreaming/Source/PixelStreamingServers/Private/SignallingServer.cpp
@@ -212,12 +212,17 @@ namespace UE::PixelStreamingServers
 			UnsubscribePlayer(PlayerConnectionId);
 		}
 
+		// We don't want to make the connections shared to prevent someone accidentally holding on to it. So we use it raw here
+		FWebSocketConnection* PlayerWS = (*PlayersWS->GetConnections().Find(PlayerConnectionId)).Get();
+		bool bUESendsOffer = !PlayerWS->GetUrlArgs().Contains(TEXT("OfferToReceive=true"));
+
 		// Send "playerConnected" message to streamer which kicks off making a new RTC connection
 		TSharedRef<FJsonObject> OnPlayerConnectedJSON = MakeShared<FJsonObject>();
 		OnPlayerConnectedJSON->SetStringField("type", "playerConnected");
 		OnPlayerConnectedJSON->SetStringField("playerId", FString::FromInt(PlayerConnectionId));
 		OnPlayerConnectedJSON->SetBoolField("dataChannel", true);
 		OnPlayerConnectedJSON->SetBoolField("sfu", false);
+		OnPlayerConnectedJSON->SetBoolField("sendOffer", bUESendsOffer);
 		SendStreamerMessage(StreamerConnectionId, OnPlayerConnectedJSON);
 
 		PlayerSubscriptions.Add(PlayerConnectionId, StreamerConnectionId);
@@ -350,12 +355,22 @@ namespace UE::PixelStreamingServers
 		}
 		else
 		{
-			if (PlayerSubscriptions.Contains(ConnectionId))
+			if (!PlayerSubscriptions.Contains(ConnectionId))
 			{
-				// Add player id to any messages going to streamer so streamer knows who sent it
-				JSONObj->SetStringField(TEXT("playerId"), FString::FromInt(ConnectionId));
-				SendStreamerMessage(PlayerSubscriptions[ConnectionId], JSONObj);
+				TArray<FString> StreamerConnections = StreamersWS->GetConnectionNames();
+				if(StreamerConnections.Num() == 0)
+				{
+					UE_LOG(LogPixelStreamingServers, Error, TEXT("Player %d sent a message, but no streamers were connected"), ConnectionId);
+					return;
+				}
+
+				UE_LOG(LogPixelStreamingServers, Log, TEXT("Player %d attempted to send an outgoing message without having subscribed first. Defaulting to %s"), ConnectionId, *StreamerConnections[0]);
+				SubscribePlayer(ConnectionId, StreamerConnections[0]);
 			}
+
+			// Add player id to any messages going to streamer so streamer knows who sent it
+			JSONObj->SetStringField(TEXT("playerId"), FString::FromInt(ConnectionId));
+			SendStreamerMessage(PlayerSubscriptions[ConnectionId], JSONObj);
 		}
 	}
 
@@ -366,16 +381,6 @@ namespace UE::PixelStreamingServers
 		{
 			StreamersWS->NameConnection(ConnectionId, StreamerName);
 			StreamersWS->RemoveName(LEGACY_NAME);
-
-			// subscribe any unsubscribed players to this new streamer
-			for (auto& ConnectionPair : PlayersWS->GetConnections())
-			{
-				const uint16 PlayerConnectionId = ConnectionPair.Key;
-				if (!PlayerSubscriptions.Contains(PlayerConnectionId))
-				{
-					SubscribePlayer(PlayerConnectionId, StreamerName);
-				}
-			}
 		}
 	}
 
diff --git a/Engine/Plugins/Media/PixelStreaming/Source/PixelStreamingServers/Private/SignallingServerLegacy.cpp b/Engine/Plugins/Media/PixelStreaming/Source/PixelStreamingServers/Private/SignallingServerLegacy.cpp
--- a/Engine/Plugins/Media/PixelStreaming/Source/PixelStreamingServers/Private/SignallingServerLegacy.cpp
+++ b/Engine/Plugins/Media/PixelStreaming/Source/PixelStreamingServers/Private/SignallingServerLegacy.cpp
@@ -4,7 +4,6 @@
 
 namespace UE::PixelStreamingServers
 {
-
 	void FSignallingServerLegacy::SendPlayerMessage(uint16 PlayerId, TSharedPtr<FJsonObject> JSONObj)
 	{
 		// legacy supports: answer, disconnectPlayer, iceCandidate, config
@@ -90,4 +89,62 @@ namespace UE::PixelStreamingServers
 			return;
 		}
 	}
+
+	void FSignallingServerLegacy::OnStreamerConnected(uint16 ConnectionId)
+	{
+		if(StreamersWS->GetConnections().Num() > 1)
+		{
+			UE_LOG(LogPixelStreamingServers, Warning, TEXT("Streamer (id=%d) attempted to connect to the signalling server, but we already have a streamer connected"), ConnectionId);	
+			return;
+		}
+
+		UE_LOG(LogPixelStreamingServers, Log, TEXT("Streamer websocket connected, id=%d"), ConnectionId);
+
+		// Send a config message to the streamer passing ICE servers to be used.
+		TSharedPtr<FJsonObject> JSONObj = CreateConfigJSON();
+		SendStreamerMessage(ConnectionId, JSONObj);
+
+		// request the streamer id
+		TSharedRef<FJsonObject> idJSON = MakeShared<FJsonObject>();
+		idJSON->SetStringField("type", "identify");
+		SendStreamerMessage(ConnectionId, idJSON);
+
+		StreamersWS->NameConnection(ConnectionId, TEXT("_LEGACY_"));
+	}
+
+	void FSignallingServerLegacy::OnPlayerMessage(uint16 ConnectionId, TArrayView<uint8> Message) 
+	{
+		const FString Msg = Utils::ToString(Message);
+		UE_LOG(LogPixelStreamingServers, Log, TEXT("From Player id=%d: %s"), ConnectionId, *Msg);
+
+		FString MsgType;
+		TSharedPtr<FJsonObject> JSONObj = ParseMessage(Msg, MsgType);
+		if (!JSONObj)
+		{
+			UE_LOG(LogPixelStreamingServers, Error, TEXT("Failed to parse incoming player message."));
+			return;
+		}
+
+		if (auto* Handler = PlayerMessageHandlers.Find(MsgType))
+		{
+			Handler->Execute(ConnectionId, JSONObj);
+		}
+		else
+		{
+			TArray<FString> StreamerConnections = StreamersWS->GetConnectionNames();
+			if(StreamerConnections.Num() == 0)
+			{
+				UE_LOG(LogPixelStreamingServers, Error, TEXT("Player %d sent a message, but no streamers were connected"), ConnectionId);
+				return;
+			}
+
+			uint16 StreamerConnectionId = INDEX_NONE;
+			if (StreamersWS->GetNamedConnection(StreamerConnections[0], StreamerConnectionId))
+			{
+				// Add player id to any messages going to streamer so streamer knows who sent it
+				JSONObj->SetStringField(TEXT("playerId"), FString::FromInt(ConnectionId));
+				SendStreamerMessage(StreamerConnectionId, JSONObj);
+			}
+		}
+	}
 } // namespace UE::PixelStreamingServers
\ No newline at end of file
diff --git a/Engine/Plugins/Media/PixelStreaming/Source/PixelStreamingServers/Private/SignallingServerLegacy.h b/Engine/Plugins/Media/PixelStreaming/Source/PixelStreamingServers/Private/SignallingServerLegacy.h
--- a/Engine/Plugins/Media/PixelStreaming/Source/PixelStreamingServers/Private/SignallingServerLegacy.h
+++ b/Engine/Plugins/Media/PixelStreaming/Source/PixelStreamingServers/Private/SignallingServerLegacy.h
@@ -17,6 +17,9 @@ namespace UE::PixelStreamingServers
 	protected:
 		virtual void SendPlayerMessage(uint16 PlayerId, TSharedPtr<FJsonObject> JSONObj) override;
 		virtual void SendStreamerMessage(uint16 StreamerId, TSharedPtr<FJsonObject> JSONObj) override;
+		
+		virtual void OnStreamerConnected(uint16 ConnectionId) override;
+		virtual void OnPlayerMessage(uint16 ConnectionId, TArrayView<uint8> Message) override;
 	};
 
 } // namespace UE::PixelStreamingServers
\ No newline at end of file
diff --git a/Engine/Plugins/Media/PixelStreaming/Source/PixelStreamingServers/Private/WebSocketServerWrapper.cpp b/Engine/Plugins/Media/PixelStreaming/Source/PixelStreamingServers/Private/WebSocketServerWrapper.cpp
--- a/Engine/Plugins/Media/PixelStreaming/Source/PixelStreamingServers/Private/WebSocketServerWrapper.cpp
+++ b/Engine/Plugins/Media/PixelStreaming/Source/PixelStreamingServers/Private/WebSocketServerWrapper.cpp
@@ -18,6 +18,7 @@ namespace UE::PixelStreamingServers
 		: Id(IdGenerator.Increment())
 		, SocketConnection(InSocketConnection)
 	{
+		UrlArgs = SocketConnection->GetUrlArgs();
 	}
 
 	FWebSocketConnection::~FWebSocketConnection()
@@ -34,6 +35,11 @@ namespace UE::PixelStreamingServers
 		return Id;
 	}
 
+	TArray<FString> FWebSocketConnection::GetUrlArgs() const
+	{
+		return UrlArgs;
+	}
+
 	bool FWebSocketConnection::Send(FString Message) const
 	{
 		// Convert FString into uint8 array.
diff --git a/Engine/Plugins/Media/PixelStreaming/Source/PixelStreamingServers/Private/WebSocketServerWrapper.h b/Engine/Plugins/Media/PixelStreaming/Source/PixelStreamingServers/Private/WebSocketServerWrapper.h
--- a/Engine/Plugins/Media/PixelStreaming/Source/PixelStreamingServers/Private/WebSocketServerWrapper.h
+++ b/Engine/Plugins/Media/PixelStreaming/Source/PixelStreamingServers/Private/WebSocketServerWrapper.h
@@ -21,6 +21,7 @@ namespace UE::PixelStreamingServers
 		FWebSocketConnection(INetworkingWebSocket* InSocketConnection);
 		~FWebSocketConnection();
 		uint16 GetId() const;
+		TArray<FString> GetUrlArgs() const;
 		bool Send(FString Message) const;
 		void SetCallbacks();
 
@@ -31,6 +32,7 @@ namespace UE::PixelStreamingServers
 	private:
 		static FThreadSafeCounter IdGenerator;
 		uint16 Id;
+		TArray<FString> UrlArgs;
 		INetworkingWebSocket* SocketConnection;
 	};
 
