QOL: Display notification to user when using Editor Pixel Streaming that the plugin is going to disable the editor setting "Use Less CPU when in Background".

diff --git a/Engine/Plugins/Media/PixelStreaming/Source/PixelStreamingEditor/Private/PixelStreamingEditorModule.cpp b/Engine/Plugins/Media/PixelStreaming/Source/PixelStreamingEditor/Private/PixelStreamingEditorModule.cpp
--- a/Engine/Plugins/Media/PixelStreaming/Source/PixelStreamingEditor/Private/PixelStreamingEditorModule.cpp
+++ b/Engine/Plugins/Media/PixelStreaming/Source/PixelStreamingEditor/Private/PixelStreamingEditorModule.cpp
@@ -43,11 +43,6 @@ void FPixelStreamingEditorModule::StartupModule()
 	FPixelStreamingStyle::ReloadTextures();
 	Toolbar = MakeShared<FPixelStreamingToolbar>();
 
-	// Update editor settings so that editor won't slow down if not in focus
-	UEditorPerformanceSettings* Settings = GetMutableDefault<UEditorPerformanceSettings>();
-	Settings->bThrottleCPUWhenNotForeground = false;
-	Settings->PostEditChange();
-
 	Settings::InitialiseSettings();
 	bUseExternalSignallingServer = Settings::CVarEditorPixelStreamingUseRemoteSignallingServer.GetValueOnAnyThread();
 
@@ -90,6 +85,10 @@ void FPixelStreamingEditorModule::InitEditorStreaming(IPixelStreamingModule& Mod
 
 	EditorStreamer = Module.CreateStreamer(EditorStreamerID);
 
+	// Bind to start/stop streaming so we disable/restore relevant editor settings
+	EditorStreamer->OnStreamingStarted().AddLambda([this](IPixelStreamingStreamer* Streamer){ DisableCPUThrottlingSetting(); });
+	EditorStreamer->OnStreamingStopped().AddLambda([this](IPixelStreamingStreamer* Streamer){ RestoreCPUThrottlingSetting(); });
+
 	// Give the editor streamer the default url if the user hasn't specified one when launching the editor
 	if (EditorStreamer->GetSignallingServerURL().IsEmpty())
 	{
@@ -408,5 +407,32 @@ void FPixelStreamingEditorModule::MaybeResizeEditor(TSharedPtr<SWindow> RootWind
 	}
 }
 
+void FPixelStreamingEditorModule::RestoreCPUThrottlingSetting()
+{
+	UEditorPerformanceSettings* Settings = GetMutableDefault<UEditorPerformanceSettings>();
+	Settings->bThrottleCPUWhenNotForeground = bOldCPUThrottlingSetting;
+	Settings->PostEditChange();
+}
+
+void FPixelStreamingEditorModule::DisableCPUThrottlingSetting()
+{
+	// Update editor settings so that editor won't slow down if not in focus
+	UEditorPerformanceSettings* Settings = GetMutableDefault<UEditorPerformanceSettings>();
+
+	// Store whatever value the user had in here so we can restore it when we are done streaming.
+	bOldCPUThrottlingSetting = Settings->bThrottleCPUWhenNotForeground;
+
+	if(Settings->bThrottleCPUWhenNotForeground)
+	{
+		Settings->bThrottleCPUWhenNotForeground = false;
+		Settings->PostEditChange();
+
+		// Let the user know we are forcing this editor setting (so they know why their setting is not working potentially)
+		FNotificationInfo Info(LOCTEXT("PixelStreamingEditorModule_Notification", "Pixel Streaming: Disabling setting \"Use less CPU in background\" for streaming performance."));
+		Info.ExpireDuration = 5.0f;
+		FSlateNotificationManager::Get().AddNotification(Info);
+	}
+}
+
 #undef LOCTEXT_NAMESPACE
 IMPLEMENT_MODULE(FPixelStreamingEditorModule, PixelStreamingEditor)
diff --git a/Engine/Plugins/Media/PixelStreaming/Source/PixelStreamingEditor/Private/PixelStreamingEditorModule.h b/Engine/Plugins/Media/PixelStreaming/Source/PixelStreamingEditor/Private/PixelStreamingEditorModule.h
--- a/Engine/Plugins/Media/PixelStreaming/Source/PixelStreamingEditor/Private/PixelStreamingEditorModule.h
+++ b/Engine/Plugins/Media/PixelStreaming/Source/PixelStreamingEditor/Private/PixelStreamingEditorModule.h
@@ -47,6 +47,8 @@ private:
 	bool ParseResolution(const TCHAR* InResolution, uint32& OutX, uint32& OutY);
 	void MaybeResizeEditor(TSharedPtr<SWindow> RootWindow);
 	void OnFrameSizeChanged(TWeakPtr<FIntRect> NewTargetRect);
+	void DisableCPUThrottlingSetting();
+	void RestoreCPUThrottlingSetting();
 
 	TSharedPtr<UE::EditorPixelStreaming::FPixelStreamingToolbar> Toolbar;
 	// Signalling/webserver
@@ -67,6 +69,7 @@ private:
 	TSharedPtr<IPixelStreamingStreamer> EditorStreamer;
 
 	bool bUseExternalSignallingServer = false;
+	bool bOldCPUThrottlingSetting = false;
 
 	TMap<Audio::FDeviceId, TSharedPtr<UE::EditorPixelStreaming::FEditorSubmixListener, ESPMode::ThreadSafe>> AudioInputs;
 };
