Fix warning with duplicate LOCTEXT id in PixelStreamingEditorModule.cpp

diff --git a/Engine/Plugins/Media/PixelStreaming/Source/PixelStreamingEditor/Private/PixelStreamingEditorModule.cpp b/Engine/Plugins/Media/PixelStreaming/Source/PixelStreamingEditor/Private/PixelStreamingEditorModule.cpp
--- a/Engine/Plugins/Media/PixelStreaming/Source/PixelStreamingEditor/Private/PixelStreamingEditorModule.cpp
+++ b/Engine/Plugins/Media/PixelStreaming/Source/PixelStreamingEditor/Private/PixelStreamingEditorModule.cpp
@@ -428,7 +428,7 @@ void FPixelStreamingEditorModule::DisableCPUThrottlingSetting()
 		Settings->PostEditChange();
 
 		// Let the user know we are forcing this editor setting (so they know why their setting is not working potentially)
-		FNotificationInfo Info(LOCTEXT("PixelStreamingEditorModule_Notification", "Pixel Streaming: Disabling setting \"Use less CPU in background\" for streaming performance."));
+		FNotificationInfo Info(LOCTEXT("PixelStreamingEditorModule_CPUThrottlingNotification", "Pixel Streaming: Disabling setting \"Use less CPU in background\" for streaming performance."));
 		Info.ExpireDuration = 5.0f;
 		FSlateNotificationManager::Get().AddNotification(Info);
 	}
