Fix: Don't pipe script output on Mac to prevent getting hung on test shutdown

diff --git a/Engine/Plugins/Media/PixelStreaming/Source/PixelStreamingServers/Private/ServerUtils.cpp b/Engine/Plugins/Media/PixelStreaming/Source/PixelStreamingServers/Private/ServerUtils.cpp
--- a/Engine/Plugins/Media/PixelStreaming/Source/PixelStreamingServers/Private/ServerUtils.cpp
+++ b/Engine/Plugins/Media/PixelStreaming/Source/PixelStreamingServers/Private/ServerUtils.cpp
@@ -38,14 +38,24 @@ namespace UE::PixelStreamingServers::Utils
 			ExecutableAbsPath = TEXT("/usr/bin/bash");
 #elif PLATFORM_MAC
             Args = FString::Printf(TEXT(" -- \"%s\" %s --nosudo"), *ExecutableAbsPath, *Args);
-            ExecutableAbsPath = TEXT("/bin/zsh");
+            ExecutableAbsPath = TEXT("/bin/bash");
 #else
 			UE_LOG(LogPixelStreamingServers, Error, TEXT("Unsupported platform for Pixel Streaming."));
 			return TSharedPtr<FMonitoredProcess>();
 #endif
 		}
 
-		TSharedPtr<FMonitoredProcess> ChildProcess = MakeShared<FMonitoredProcess>(ExecutableAbsPath, Args, true, true);
+		TSharedPtr<FMonitoredProcess> ChildProcess = MakeShared<FMonitoredProcess>(
+			ExecutableAbsPath, 
+			Args, 
+			true, 
+#if PLATFORM_MAC
+			// Pipes cause UE to lockup when destroying on Mac
+			false
+#else
+			true
+#endif
+		);
 		// Bind to output so we can capture the output in the log
 		ChildProcess->OnOutput().BindLambda([LogPrefix](FString Output) {
 			UE_LOG(LogPixelStreamingServers, Log, TEXT("%s - %s"), *LogPrefix, *Output);
