Pixel Streaming MediaIO capture improvements and bug fixes

- Fix error where Pixel Streaming media capture seems to be running too early before a frame is rendered.
- Add option to specify custom media capture in PixelStreamingVideoInputMediaCapture
- Fix default case of capturing active scene viewport to be more explicitly named.
- Specify the overrun action to be "skip" for lowest latency, unspecified default was flush, which will block the rendering thread on overrun to ensure the capture is done.

diff --git a/Engine/Plugins/Media/PixelStreaming/Source/PixelStreaming/Internal/PixelStreamingVideoInputMediaCapture.h b/Engine/Plugins/Media/PixelStreaming/Source/PixelStreaming/Internal/PixelStreamingVideoInputMediaCapture.h
--- a/Engine/Plugins/Media/PixelStreaming/Source/PixelStreaming/Internal/PixelStreamingVideoInputMediaCapture.h
+++ b/Engine/Plugins/Media/PixelStreaming/Source/PixelStreaming/Internal/PixelStreamingVideoInputMediaCapture.h
@@ -2,6 +2,8 @@
 
 #pragma once
 
+#include "Delegates/IDelegateInstance.h"
+#include "Misc/Optional.h"
 #include "MediaOutput.h"
 #include "PixelStreamingVideoInputRHI.h"
 #include "PixelStreamingMediaIOCapture.h"
@@ -27,21 +29,41 @@ private:
 class PIXELSTREAMING_API FPixelStreamingVideoInputMediaCapture : public FPixelStreamingVideoInputRHI, public TSharedFromThis<FPixelStreamingVideoInputMediaCapture>
 {
 public:
-	static TSharedPtr<FPixelStreamingVideoInputMediaCapture> Create();
+
+	/**
+	 * @brief Creates a MediaIO capture of the active viewport and starts capturing as soon as possible.
+	 * @return A video input backed by the created MediaIO capture that sends frames from the active viewport.
+	*/
+	static TSharedPtr<FPixelStreamingVideoInputMediaCapture> CreateActiveViewportCapture();
+
+	/**
+	 * @brief Creates a video input where the user can specify their own MediaIO output and capture.
+	 * This method does not does not configure or start capturing, this is left to the user.
+	 * Use this constructor if you know how to configure the MediaIOCapture yourself or don't want to capture the active viewport.
+	 * @param MediaCapture The custom MediaIOCapture that will pass its captured frames as video input.
+	 * @return A video input backed by the passed in MediaIO capture.
+	*/
+	static TSharedPtr<FPixelStreamingVideoInputMediaCapture> Create(TObjectPtr<UPixelStreamingMediaIOCapture> MediaCapture);
+
+	FPixelStreamingVideoInputMediaCapture();
+	FPixelStreamingVideoInputMediaCapture(TObjectPtr<UPixelStreamingMediaIOCapture> MediaCapture);
 	virtual ~FPixelStreamingVideoInputMediaCapture();
 
 	virtual FString ToString() override;
 
 protected:
+
 	virtual TSharedPtr<FPixelCaptureCapturer> CreateCapturer(int32 FinalFormat, float FinalScale) override;
-	
+	void StartActiveViewportCapture();
+	void LateStartActiveViewportCapture();
+
 private:
-	FPixelStreamingVideoInputMediaCapture();
-	void StartCapture();
-	void OnCaptureStateChanged();
+	void CreateDefaultActiveViewportMediaCapture();
+	void OnCaptureActiveViewportStateChanged();
 
-	UPROPERTY(Transient)
-	TObjectPtr<UPixelStreamingMediaIOOutput> MediaOutput = nullptr;
+private:
 	UPROPERTY(Transient)
 	TObjectPtr<UPixelStreamingMediaIOCapture> MediaCapture = nullptr;
+
+	TOptional<FDelegateHandle> OnFrameEndDelegateHandle;
 };
diff --git a/Engine/Plugins/Media/PixelStreaming/Source/PixelStreaming/Private/PixelStreamingModule.cpp b/Engine/Plugins/Media/PixelStreaming/Source/PixelStreaming/Private/PixelStreamingModule.cpp
--- a/Engine/Plugins/Media/PixelStreaming/Source/PixelStreaming/Private/PixelStreamingModule.cpp
+++ b/Engine/Plugins/Media/PixelStreaming/Source/PixelStreaming/Private/PixelStreamingModule.cpp
@@ -427,7 +427,7 @@ namespace UE::PixelStreaming
 			// in that case, set up the video input for them (as long as we're not in editor)
 			if (Settings::CVarPixelStreamingUseMediaCapture.GetValueOnAnyThread())
 			{
-				DefaultStreamer->SetVideoInput(FPixelStreamingVideoInputMediaCapture::Create());
+				DefaultStreamer->SetVideoInput(FPixelStreamingVideoInputMediaCapture::CreateActiveViewportCapture());
 			}
 			else
 			{
diff --git a/Engine/Plugins/Media/PixelStreaming/Source/PixelStreaming/Private/PixelStreamingVideoInputMediaCapture.cpp b/Engine/Plugins/Media/PixelStreaming/Source/PixelStreaming/Private/PixelStreamingVideoInputMediaCapture.cpp
--- a/Engine/Plugins/Media/PixelStreaming/Source/PixelStreaming/Private/PixelStreamingVideoInputMediaCapture.cpp
+++ b/Engine/Plugins/Media/PixelStreaming/Source/PixelStreaming/Private/PixelStreamingVideoInputMediaCapture.cpp
@@ -10,34 +10,26 @@
 #include "PixelStreamingPrivate.h"
 #include "Engine/GameEngine.h"
 
-TSharedPtr<FPixelStreamingVideoInputMediaCapture> FPixelStreamingVideoInputMediaCapture::Create()
+TSharedPtr<FPixelStreamingVideoInputMediaCapture> FPixelStreamingVideoInputMediaCapture::CreateActiveViewportCapture()
 {
-	TSharedPtr<FPixelStreamingVideoInputMediaCapture> NewInput = TSharedPtr<FPixelStreamingVideoInputMediaCapture>(new FPixelStreamingVideoInputMediaCapture());
+	TSharedPtr<FPixelStreamingVideoInputMediaCapture> NewInput = MakeShared<FPixelStreamingVideoInputMediaCapture>();
+	NewInput->LateStartActiveViewportCapture();
+	return NewInput;
+}
 
-	if (GEngine == nullptr)
-	{
-		// this is usually hit via the blueprint implementations of this class. the editor creates these objects
-		// on startup which causes CaptureActiveSceneViewport() to crash trying to access GEngine
-
-		// delay this until GEngine is ready
-		TWeakPtr<FPixelStreamingVideoInputMediaCapture> WeakInput = NewInput;
-		FCoreDelegates::OnPostEngineInit.AddLambda([WeakInput]() {
-			if (TSharedPtr<FPixelStreamingVideoInputMediaCapture> Input = WeakInput.Pin()) {
-				Input->StartCapture();
-			}
-			});
-	}
-	else
-	{
-		// otherwise start capture
-		NewInput->StartCapture();
-	}
+TSharedPtr<FPixelStreamingVideoInputMediaCapture> FPixelStreamingVideoInputMediaCapture::Create(TObjectPtr<UPixelStreamingMediaIOCapture> MediaCapture)
+{
+	TSharedPtr<FPixelStreamingVideoInputMediaCapture> VideoInput = MakeShared<FPixelStreamingVideoInputMediaCapture>(MediaCapture);
+	MediaCapture->SetVideoInput(VideoInput);
+	return VideoInput;
+}
 
-	return NewInput;
+FPixelStreamingVideoInputMediaCapture::FPixelStreamingVideoInputMediaCapture(TObjectPtr<UPixelStreamingMediaIOCapture> InMediaCapture)
+	: MediaCapture(InMediaCapture)
+{
 }
 
 FPixelStreamingVideoInputMediaCapture::FPixelStreamingVideoInputMediaCapture()
-	:MediaOutput(NewObject<UPixelStreamingMediaIOOutput>())
 {
 }
 
@@ -50,8 +42,15 @@ FPixelStreamingVideoInputMediaCapture::~FPixelStreamingVideoInputMediaCapture()
 	}
 }
 
-void FPixelStreamingVideoInputMediaCapture::StartCapture()
+void FPixelStreamingVideoInputMediaCapture::StartActiveViewportCapture()
 {
+	// If we were bound to the OnFrameEnd delegate to ensure a frame was rendered before starting, then we can unset it here.
+	if(OnFrameEndDelegateHandle.IsSet())
+	{
+		FCoreDelegates::OnEndFrame.Remove(OnFrameEndDelegateHandle.GetValue());
+		OnFrameEndDelegateHandle.Reset();
+	}
+
 	if (MediaCapture)
 	{
 		MediaCapture->OnStateChangedNative.RemoveAll(this);
@@ -60,10 +59,23 @@ void FPixelStreamingVideoInputMediaCapture::StartCapture()
 
 	MediaCapture = NewObject<UPixelStreamingMediaIOCapture>();
 	MediaCapture->AddToRoot(); // prevent GC on this
-	MediaCapture->SetMediaOutput(MediaOutput);
+	MediaCapture->SetMediaOutput(NewObject<UPixelStreamingMediaIOOutput>());
 	MediaCapture->SetVideoInput(AsShared());
-	MediaCapture->OnStateChangedNative.AddSP(this, &FPixelStreamingVideoInputMediaCapture::OnCaptureStateChanged);
-	MediaCapture->CaptureActiveSceneViewport({});
+	MediaCapture->OnStateChangedNative.AddSP(this, &FPixelStreamingVideoInputMediaCapture::OnCaptureActiveViewportStateChanged);
+
+	FMediaCaptureOptions Options;
+	Options.bSkipFrameWhenRunningExpensiveTasks = false;
+	Options.OverrunAction = EMediaCaptureOverrunAction::Skip;
+	Options.ResizeMethod = EMediaCaptureResizeMethod::ResizeSource;
+
+	// Start capturing the active viewport
+	MediaCapture->CaptureActiveSceneViewport(Options);
+}
+
+void FPixelStreamingVideoInputMediaCapture::LateStartActiveViewportCapture()
+{
+	// Bind the OnEndFrame delegate to ensure we only start capture once a frame has been rendered
+	OnFrameEndDelegateHandle = FCoreDelegates::OnEndFrame.AddSP(this, &FPixelStreamingVideoInputMediaCapture::StartActiveViewportCapture);
 }
 
 FString FPixelStreamingVideoInputMediaCapture::ToString()
@@ -116,7 +128,7 @@ TSharedPtr<FPixelCaptureCapturer> FPixelStreamingVideoInputMediaCapture::CreateC
 	}
 }
 
-void FPixelStreamingVideoInputMediaCapture::OnCaptureStateChanged()
+void FPixelStreamingVideoInputMediaCapture::OnCaptureActiveViewportStateChanged()
 {
 	if (!MediaCapture)
 	{
@@ -133,7 +145,7 @@ void FPixelStreamingVideoInputMediaCapture::OnCaptureStateChanged()
 		{
 			UE_LOG(LogPixelStreaming, Log, TEXT("Pixel Streaming capture was stopped due to resize, going to restart capture."));
 			// If it was stopped and viewport resized we assume resize caused the stop, so try a restart of capture here.
-			StartCapture();
+			StartActiveViewportCapture();
 		}
 		else
 		{
diff --git a/Engine/Plugins/Media/PixelStreaming/Source/PixelStreamingBlueprint/Private/PixelStreamingStreamerVideoInputMediaCapture.cpp b/Engine/Plugins/Media/PixelStreaming/Source/PixelStreamingBlueprint/Private/PixelStreamingStreamerVideoInputMediaCapture.cpp
--- a/Engine/Plugins/Media/PixelStreaming/Source/PixelStreamingBlueprint/Private/PixelStreamingStreamerVideoInputMediaCapture.cpp
+++ b/Engine/Plugins/Media/PixelStreaming/Source/PixelStreamingBlueprint/Private/PixelStreamingStreamerVideoInputMediaCapture.cpp
@@ -16,6 +16,6 @@ UPixelStreamingStreamerVideoInputMediaCapture::UPixelStreamingStreamerVideoInput
 	}
 	else
 	{
-		VideoInput = FPixelStreamingVideoInputMediaCapture::Create();
+		VideoInput = FPixelStreamingVideoInputMediaCapture::CreateActiveViewportCapture();
 	}
 }
