Improve Websocket disconnect messages in the Pixel Streaming plugin. Previously all disconnects on the UE side were reported with no disconnect message.

diff --git a/Engine/Plugins/Experimental/PixelStreamingPlayer/Source/PixelStreamingPlayer/Private/PixelStreamingSignallingComponent.cpp b/Engine/Plugins/Experimental/PixelStreamingPlayer/Source/PixelStreamingPlayer/Private/PixelStreamingSignallingComponent.cpp
--- a/Engine/Plugins/Experimental/PixelStreamingPlayer/Source/PixelStreamingPlayer/Private/PixelStreamingSignallingComponent.cpp
+++ b/Engine/Plugins/Experimental/PixelStreamingPlayer/Source/PixelStreamingPlayer/Private/PixelStreamingSignallingComponent.cpp
@@ -34,7 +34,7 @@ void UPixelStreamingSignallingComponent::Connect(const FString& Url)
 
 void UPixelStreamingSignallingComponent::Disconnect()
 {
-	SignallingConnection->Disconnect();
+	SignallingConnection->Disconnect(TEXT("Pixel Streaming player plugin disconnected"));
 }
 
 void UPixelStreamingSignallingComponent::Subscribe(const FString& StreamerId)
diff --git a/Engine/Plugins/Media/PixelStreaming/Source/PixelStreaming/Private/PixelStreamingSignallingConnection.cpp b/Engine/Plugins/Media/PixelStreaming/Source/PixelStreaming/Private/PixelStreamingSignallingConnection.cpp
--- a/Engine/Plugins/Media/PixelStreaming/Source/PixelStreaming/Private/PixelStreamingSignallingConnection.cpp
+++ b/Engine/Plugins/Media/PixelStreaming/Source/PixelStreaming/Private/PixelStreamingSignallingConnection.cpp
@@ -39,7 +39,7 @@ FPixelStreamingSignallingConnection::FPixelStreamingSignallingConnection(TShared
 
 FPixelStreamingSignallingConnection::~FPixelStreamingSignallingConnection()
 {
-	Disconnect();
+	Disconnect(TEXT("Streamer destroying websocket connection"));
 
 	WebSocket.Reset();
 	Observer.Reset();
@@ -60,7 +60,7 @@ void FPixelStreamingSignallingConnection::Connect(FString InUrl, bool bIsReconne
 
 	// Reconnecting on an existing websocket can be problematic depending what state it was
 	// left in. Easier and safer to disconnect any existing socket/delegates and start fresh.
-	Disconnect();
+	Disconnect(TEXT("Streamer cleaning up old connection"));
 
 	Url = InUrl;
 
@@ -107,12 +107,22 @@ void FPixelStreamingSignallingConnection::TryConnect(FString InUrl)
 }
 
 void FPixelStreamingSignallingConnection::Disconnect()
+{
+	// Do not call this deprecated method, please call FPixelStreamingSignallingConnection::Disconnect(FString Reason)
+	Disconnect(TEXT("Unknown reason"));
+}
+
+void FPixelStreamingSignallingConnection::Disconnect(FString Reason)
 {
 	if (!IsEngineExitRequested())
 	{
 		StopKeepAliveTimer();
 		StopReconnectTimer();
 	}
+	else
+	{
+		Reason = TEXT("Streamed application is shutting down");
+	}
 
 	if (!WebSocket || !bIsConnected)
 	{
@@ -125,7 +135,7 @@ void FPixelStreamingSignallingConnection::Disconnect()
 	WebSocket->OnMessage().Remove(OnMessageHandle);
 	WebSocket->OnBinaryMessage().Remove(OnBinaryMessageHandle);
 
-	WebSocket->Close();
+	WebSocket->Close(1000, Reason);
 	UE_LOG(LogPixelStreamingSS, Log, TEXT("Closing websocket to SS %s"), *Url);
 
 	bIsConnected = false;
diff --git a/Engine/Plugins/Media/PixelStreaming/Source/PixelStreaming/Private/Streamer.cpp b/Engine/Plugins/Media/PixelStreaming/Source/PixelStreaming/Private/Streamer.cpp
--- a/Engine/Plugins/Media/PixelStreaming/Source/PixelStreaming/Private/Streamer.cpp
+++ b/Engine/Plugins/Media/PixelStreaming/Source/PixelStreaming/Private/Streamer.cpp
@@ -243,7 +243,7 @@ namespace UE::PixelStreaming
 
 		if (SignallingServerConnection)
 		{
-			SignallingServerConnection->Disconnect();
+			SignallingServerConnection->Disconnect(TEXT("Streamer stopped streaming"));
 		}
 		VideoSourceGroup->Stop();
 		TriggerMouseLeave(StreamerId);
diff --git a/Engine/Plugins/Media/PixelStreaming/Source/PixelStreaming/Private/Tests/MockPlayer.h b/Engine/Plugins/Media/PixelStreaming/Source/PixelStreaming/Private/Tests/MockPlayer.h
--- a/Engine/Plugins/Media/PixelStreaming/Source/PixelStreaming/Private/Tests/MockPlayer.h
+++ b/Engine/Plugins/Media/PixelStreaming/Source/PixelStreaming/Private/Tests/MockPlayer.h
@@ -72,7 +72,7 @@ namespace UE::PixelStreaming
 	public:
 		virtual ~FMockPlayer()
 		{
-			Disconnect();
+			Disconnect(TEXT("Mock player being destroyed"));
 
 			UE::PixelStreaming::DoOnGameThread([]() {
 				Settings::CVarPixelStreamingSuppressICECandidateErrors->Set(false, ECVF_SetByCode);
@@ -98,9 +98,9 @@ namespace UE::PixelStreaming
 			SignallingServerConnection->TryConnect(Url);
 		}
 
-		void Disconnect()
+		void Disconnect(FString Reason)
 		{
-			SignallingServerConnection->Disconnect();
+			SignallingServerConnection->Disconnect(Reason);
 		}
 
 		bool IsSignallingConnected()
diff --git a/Engine/Plugins/Media/PixelStreaming/Source/PixelStreaming/Private/Tests/TestUtils.cpp b/Engine/Plugins/Media/PixelStreaming/Source/PixelStreaming/Private/Tests/TestUtils.cpp
--- a/Engine/Plugins/Media/PixelStreaming/Source/PixelStreaming/Private/Tests/TestUtils.cpp
+++ b/Engine/Plugins/Media/PixelStreaming/Source/PixelStreaming/Private/Tests/TestUtils.cpp
@@ -297,7 +297,7 @@ namespace UE::PixelStreaming
 	{
 		if (OutPlayer)
 		{
-			OutPlayer->Disconnect();
+			OutPlayer->Disconnect(TEXT("Cleaning up tests"));
 			OutPlayer.Reset();
 		}
 
diff --git a/Engine/Plugins/Media/PixelStreaming/Source/PixelStreaming/Public/IPixelStreamingSignallingConnection.h b/Engine/Plugins/Media/PixelStreaming/Source/PixelStreaming/Public/IPixelStreamingSignallingConnection.h
--- a/Engine/Plugins/Media/PixelStreaming/Source/PixelStreaming/Public/IPixelStreamingSignallingConnection.h
+++ b/Engine/Plugins/Media/PixelStreaming/Source/PixelStreaming/Public/IPixelStreamingSignallingConnection.h
@@ -26,8 +26,14 @@ public:
 	/**
 	 * Disconnects from the signalling server. Safe to call even when not connected.
 	 */
+	UE_DEPRECATED(5.4, "Disconnect has been deprecated and will be removed in future versions. Please use: Disconnect(FString Reason)")
 	virtual void Disconnect() = 0;
 
+	/**
+	 * Disconnects from the signalling server with a reason message. Safe to call even when not connected.
+	 */
+	virtual void Disconnect(FString Reason) = 0;
+
 	/**
 	 * Checks if the connection is established.
 	 * @returns True when the signalling server is connected.
diff --git a/Engine/Plugins/Media/PixelStreaming/Source/PixelStreaming/Public/PixelStreamingSignallingConnection.h b/Engine/Plugins/Media/PixelStreaming/Source/PixelStreaming/Public/PixelStreamingSignallingConnection.h
--- a/Engine/Plugins/Media/PixelStreaming/Source/PixelStreaming/Public/PixelStreamingSignallingConnection.h
+++ b/Engine/Plugins/Media/PixelStreaming/Source/PixelStreaming/Public/PixelStreamingSignallingConnection.h
@@ -22,7 +22,9 @@ public:
 
 	/* IPixelStreamingSignallingConnection Interface */
 	virtual void TryConnect(FString URL) override;
+	UE_DEPRECATED(5.4, "Disconnect has been deprecated and will be removed in future versions. Please use: Disconnect(FString Reason)")
 	virtual void Disconnect() override;
+	virtual void Disconnect(FString Reason) override;
 	virtual bool IsConnected() const override;
 	virtual void SendOffer(FPixelStreamingPlayerId PlayerId, const webrtc::SessionDescriptionInterface& SDP) override;
 	virtual void SendAnswer(FPixelStreamingPlayerId PlayerId, const webrtc::SessionDescriptionInterface& SDP) override;
