Fix: Stream stutter on new peer connections

diff --git a/Engine/Plugins/Media/PixelStreaming/Source/PixelStreaming/Private/Streamer.cpp b/Engine/Plugins/Media/PixelStreaming/Source/PixelStreaming/Private/Streamer.cpp
--- a/Engine/Plugins/Media/PixelStreaming/Source/PixelStreaming/Private/Streamer.cpp
+++ b/Engine/Plugins/Media/PixelStreaming/Source/PixelStreaming/Private/Streamer.cpp
@@ -537,7 +537,7 @@ namespace UE::PixelStreaming
 				{
 					SFUPlayerId = PlayerId;
 				}
-				else
+				else if(QualityControllingId == INVALID_PLAYER_ID)
 				{
 					SetQualityController(PlayerId);
 				}
diff --git a/Engine/Plugins/Media/PixelStreaming/Source/PixelStreaming/Private/VideoEncoderFactorySingleLayer.cpp b/Engine/Plugins/Media/PixelStreaming/Source/PixelStreaming/Private/VideoEncoderFactorySingleLayer.cpp
--- a/Engine/Plugins/Media/PixelStreaming/Source/PixelStreaming/Private/VideoEncoderFactorySingleLayer.cpp
+++ b/Engine/Plugins/Media/PixelStreaming/Source/PixelStreaming/Private/VideoEncoderFactorySingleLayer.cpp
@@ -234,13 +234,12 @@ namespace UE::PixelStreaming
 	{
 		// first clear unused encoders.
 		TArray<uint32> DeleteStreamIds;
-		for (auto&& [EncoderStreamId, EncoderPtr] : HardwareEncoders)
-		{
+		HardwareEncoders.Apply([&DeleteStreamIds](uint32 EncoderStreamId, TSharedPtr<FVideoEncoderHardware> EncoderPtr) {
 			if (EncoderPtr.IsUnique())
 			{
 				DeleteStreamIds.Add(EncoderStreamId);
 			}
-		}
+		});
 		for (uint32 DeleteStreamId : DeleteStreamIds)
 		{
 			HardwareEncoders.Remove(DeleteStreamId);
diff --git a/Engine/Plugins/Media/PixelStreaming/Source/PixelStreaming/Private/VideoEncoderFactorySingleLayer.h b/Engine/Plugins/Media/PixelStreaming/Source/PixelStreaming/Private/VideoEncoderFactorySingleLayer.h
--- a/Engine/Plugins/Media/PixelStreaming/Source/PixelStreaming/Private/VideoEncoderFactorySingleLayer.h
+++ b/Engine/Plugins/Media/PixelStreaming/Source/PixelStreaming/Private/VideoEncoderFactorySingleLayer.h
@@ -7,7 +7,7 @@
 #include "Video/Encoders/Configs/VideoEncoderConfigH264.h"
 #include "Video/Encoders/Configs/VideoEncoderConfigAV1.h"
 #include "Video/Resources/VideoResourceRHI.h"
-
+#include "ThreadSafeMap.h"
 #include "PixelStreamingPrivate.h"
 
 namespace UE::PixelStreaming
@@ -77,7 +77,7 @@ namespace UE::PixelStreaming
 		void FreeUnusedEncoders();
 		bool CheckEncoderSessionAvailable() const;
 
-		TMap<uint32, TSharedPtr<FVideoEncoderHardware>> HardwareEncoders;
+		TThreadSafeMap<uint32, TSharedPtr<FVideoEncoderHardware>> HardwareEncoders;
 
 		uint8 bForceNextKeyframe : 1;
 
@@ -88,9 +88,5 @@ namespace UE::PixelStreaming
 
 		// Init encoder guard
 		FCriticalSection InitEncoderGuard;
-
-		// A map of GPU device names and their respective Max # of Concurrent Sessions as defined here: https://developer.nvidia.com/video-encode-and-decode-gpu-support-matrix-new
-		// If a device is not in this map, it has an unrestriced number of sessions
-		static TMap<FString, uint8> MaxConcurrentNvEncSessionsMap;
 	};
 } // namespace UE::PixelStreaming
