Fix: Change frame receive tests to send multiple frames to increase test robustness.

- Refactored FreezeFrame into its own class
- Added a codepath in FreezeFrame for VPX encoding that does not require execution on the rendering thread (this was causing tests to fail)

diff --git a/Engine/Plugins/Media/PixelStreaming/Source/PixelStreaming/Private/FreezeFrame.cpp b/Engine/Plugins/Media/PixelStreaming/Source/PixelStreaming/Private/FreezeFrame.cpp
new file mode 100644
--- /dev/null
+++ b/Engine/Plugins/Media/PixelStreaming/Source/PixelStreaming/Private/FreezeFrame.cpp
@@ -0,0 +1,335 @@
+// Copyright Epic Games, Inc. All Rights Reserved.
+
+#include "FreezeFrame.h"
+#include "EngineModule.h"
+#include "IImageWrapper.h"
+#include "IImageWrapperModule.h"
+#include "Modules/ModuleManager.h"
+#include "PixelStreamingPrivate.h"
+#include "PixelStreamingInputProtocol.h"
+#include "PixelCaptureOutputFrameI420.h"
+#include "PixelCaptureBufferFormat.h"
+#include "PixelCaptureOutputFrameRHI.h"
+#include "PixelCaptureBufferI420.h"
+#include "TextureResource.h"
+#include "RHICommandList.h"
+#include "RHIFwd.h"
+#include "ScreenPass.h"
+#include "ScreenRendering.h"
+#include "Settings.h"
+#include "Utils.h"
+#include "MediaShaders.h"
+#include "RHI.h"
+#include "RenderingThread.h"
+#include "libyuv/convert.h"
+#include "libyuv/video_common.h"
+#include "common_video/libyuv/include/webrtc_libyuv.h"
+
+namespace UE::PixelStreaming
+{
+	TSharedPtr<FFreezeFrame> FFreezeFrame::Create(TWeakPtr<TThreadSafeMap<FPixelStreamingPlayerId,FPlayerContext>> InPlayers)
+	{
+		TSharedPtr<FFreezeFrame> FF = TSharedPtr<FFreezeFrame>(new FFreezeFrame());
+		FF->SetPlayers(InPlayers);
+		return FF;
+	}
+
+	FFreezeFrame::~FFreezeFrame()
+	{
+		RemoveFreezeFrameBinding();
+	}
+
+	void FFreezeFrame::SetVideoInput(TWeakPtr<FPixelStreamingVideoInput> InVideoInput)
+	{
+		VideoInput = InVideoInput;
+		RemoveFreezeFrameBinding();
+	}
+
+	void FFreezeFrame::SetPlayers(TWeakPtr<TThreadSafeMap<FPixelStreamingPlayerId,FPlayerContext>> InPlayers)
+	{
+		WeakPlayers = InPlayers;
+	}
+
+	/*
+	 * Add the commands to the RHI command list to copy a texture from source to dest - even if the format is different.
+	 * Assumes SourceTexture is in ERHIAccess::CopySrc and DestTexture is in ERHIAccess::CopyDest
+	 */
+	void CopyTexture(FRHICommandList& RHICmdList, FTextureRHIRef SourceTexture, FTextureRHIRef DestTexture)
+	{
+		if (SourceTexture->GetDesc().Format == DestTexture->GetDesc().Format
+			&& SourceTexture->GetDesc().Extent.X == DestTexture->GetDesc().Extent.X
+			&& SourceTexture->GetDesc().Extent.Y == DestTexture->GetDesc().Extent.Y)
+		{
+
+			RHICmdList.Transition(FRHITransitionInfo(SourceTexture, ERHIAccess::Unknown, ERHIAccess::CopySrc));
+			RHICmdList.Transition(FRHITransitionInfo(DestTexture, ERHIAccess::Unknown, ERHIAccess::CopyDest));
+
+			// source and dest are the same. simple copy
+			RHICmdList.CopyTexture(SourceTexture, DestTexture, {});
+		}
+		else
+		{
+			IRendererModule* RendererModule = &FModuleManager::GetModuleChecked<IRendererModule>("Renderer");
+
+			RHICmdList.Transition(FRHITransitionInfo(SourceTexture, ERHIAccess::Unknown, ERHIAccess::SRVMask));
+			RHICmdList.Transition(FRHITransitionInfo(DestTexture, ERHIAccess::Unknown, ERHIAccess::RTV));
+
+			// source and destination are different. rendered copy
+			FRHIRenderPassInfo RPInfo(DestTexture, ERenderTargetActions::Load_Store);
+			RHICmdList.BeginRenderPass(RPInfo, TEXT("PixelStreaming::CopyTexture"));
+			{
+				FGlobalShaderMap* ShaderMap = GetGlobalShaderMap(GMaxRHIFeatureLevel);
+				TShaderMapRef<FScreenVS> VertexShader(ShaderMap);
+				TShaderMapRef<FScreenPS> PixelShader(ShaderMap);
+
+				RHICmdList.SetViewport(0, 0, 0.0f, DestTexture->GetDesc().Extent.X, DestTexture->GetDesc().Extent.Y, 1.0f);
+
+				FGraphicsPipelineStateInitializer GraphicsPSOInit;
+				RHICmdList.ApplyCachedRenderTargets(GraphicsPSOInit);
+				GraphicsPSOInit.BlendState = TStaticBlendState<>::GetRHI();
+				GraphicsPSOInit.RasterizerState = TStaticRasterizerState<>::GetRHI();
+				GraphicsPSOInit.DepthStencilState = TStaticDepthStencilState<false, CF_Always>::GetRHI();
+				GraphicsPSOInit.BoundShaderState.VertexDeclarationRHI = GFilterVertexDeclaration.VertexDeclarationRHI;
+				GraphicsPSOInit.BoundShaderState.VertexShaderRHI = VertexShader.GetVertexShader();
+				GraphicsPSOInit.BoundShaderState.PixelShaderRHI = PixelShader.GetPixelShader();
+				GraphicsPSOInit.PrimitiveType = PT_TriangleList;
+				SetGraphicsPipelineState(RHICmdList, GraphicsPSOInit, 0);
+
+				SetShaderParametersLegacyPS(RHICmdList, PixelShader, TStaticSamplerState<SF_Point>::GetRHI(), SourceTexture);
+
+				FIntPoint TargetBufferSize(DestTexture->GetDesc().Extent.X, DestTexture->GetDesc().Extent.Y);
+				RendererModule->DrawRectangle(RHICmdList, 0, 0, // Dest X, Y
+					DestTexture->GetDesc().Extent.X,			// Dest Width
+					DestTexture->GetDesc().Extent.Y,			// Dest Height
+					0, 0,										// Source U, V
+					1, 1,										// Source USize, VSize
+					TargetBufferSize,							// Target buffer size
+					FIntPoint(1, 1),							// Source texture size
+					VertexShader, EDRF_Default);
+			}
+
+			RHICmdList.EndRenderPass();
+
+			RHICmdList.Transition(FRHITransitionInfo(SourceTexture, ERHIAccess::SRVMask, ERHIAccess::CopySrc));
+			RHICmdList.Transition(FRHITransitionInfo(DestTexture, ERHIAccess::RTV, ERHIAccess::CopyDest));
+		}
+	}
+
+	void FFreezeFrame::StartFreeze(UTexture2D* Texture)
+	{
+		if (Texture)
+		{
+			ENQUEUE_RENDER_COMMAND(ReadSurfaceCommand)
+			([this, Texture](FRHICommandListImmediate& RHICmdList) {
+				// A frame is supplied so immediately read its data and send as a JPEG.
+				FTextureRHIRef TextureRHI = Texture->GetResource() ? Texture->GetResource()->TextureRHI : nullptr;
+				if (!TextureRHI)
+				{
+					UE_LOG(LogPixelStreaming, Error, TEXT("Attempting freeze frame with texture %s with no texture RHI"), *Texture->GetName());
+					return;
+				}
+				uint32 Width = TextureRHI->GetDesc().Extent.X;
+				uint32 Height = TextureRHI->GetDesc().Extent.Y;
+
+				FRHITextureCreateDesc TextureDesc =
+						FRHITextureCreateDesc::Create2D(TEXT("PixelStreamingBlankTexture"), Width, Height, EPixelFormat::PF_B8G8R8A8)
+						.SetClearValue(FClearValueBinding::None)
+						.SetFlags(ETextureCreateFlags::RenderTargetable)
+						.SetInitialState(ERHIAccess::Present)
+						.DetermineInititialState();
+
+				FTextureRHIRef DestTexture = RHICreateTexture(TextureDesc);
+
+				// Copy freeze frame texture to empty texture
+				CopyTexture(RHICmdList, TextureRHI, DestTexture);
+
+				TArray<FColor> Data;
+				FIntRect Rect(0, 0, Width, Height);
+				// This `ReadSurfaceData` makes a blocking call from CPU -> GPU -> CPU
+				// WHich is how on the very next line we are able to copy the data out and send it.
+				RHICmdList.ReadSurfaceData(DestTexture, Rect, Data, FReadSurfaceDataFlags());
+				SendFreezeFrame(MoveTemp(Data), Rect);
+			});
+		}
+		else
+		{
+			// A frame is not supplied, so we need to get it from the video input
+			// at the next opportunity and send as a JPEG.
+			SetupFreezeFrameCapture();
+		}
+	}
+
+	void FFreezeFrame::StopFreeze()
+	{
+		TSharedPtr<TThreadSafeMap<FPixelStreamingPlayerId, FPlayerContext>> Players = WeakPlayers.Pin();
+
+		if(!Players)
+		{
+			return;
+		}
+
+		Players->Apply([this](FPixelStreamingPlayerId PlayerId, FPlayerContext& PlayerContext) {
+				if (PlayerContext.DataChannel)
+				{
+					PlayerContext.DataChannel->SendMessage(FPixelStreamingInputProtocol::FromStreamerProtocol.Find("UnfreezeFrame")->GetID());
+				}
+			});
+
+		CachedJpegBytes.Empty();
+	}
+
+	void FFreezeFrame::SendCachedFreezeFrameTo(FPixelStreamingPlayerId PlayerId) const
+	{
+		TSharedPtr<TThreadSafeMap<FPixelStreamingPlayerId, FPlayerContext>> Players = WeakPlayers.Pin();
+
+		if(!Players)
+		{
+			return;
+		}
+
+		if (CachedJpegBytes.Num() > 0)
+		{
+			if (const FPlayerContext* PlayerContext = Players->Find(PlayerId))
+			{
+				if (PlayerContext->DataChannel)
+				{
+					PlayerContext->DataChannel->SendArbitraryData(FPixelStreamingInputProtocol::FromStreamerProtocol.Find("FreezeFrame")->GetID(), CachedJpegBytes);
+				}
+			}
+		}
+	}
+
+	void FFreezeFrame::SendFreezeFrame(TArray<FColor> RawData, const FIntRect & Rect)
+	{
+		TSharedPtr<TThreadSafeMap<FPixelStreamingPlayerId, FPlayerContext>> Players = WeakPlayers.Pin();
+
+		if(!Players)
+		{
+			return;
+		}
+
+		IImageWrapperModule& ImageWrapperModule = FModuleManager::GetModuleChecked<IImageWrapperModule>(TEXT("ImageWrapper"));
+		TSharedPtr<IImageWrapper> ImageWrapper = ImageWrapperModule.CreateImageWrapper(EImageFormat::JPEG);
+		bool bSuccess = ImageWrapper->SetRaw(RawData.GetData(), RawData.Num() * sizeof(FColor), Rect.Width(), Rect.Height(), ERGBFormat::BGRA, 8);
+		if (bSuccess)
+		{
+			// Compress to a JPEG of the maximum possible quality.
+			int32 Quality = Settings::CVarPixelStreamingFreezeFrameQuality.GetValueOnAnyThread();
+			const TArray64<uint8>& JpegBytes = ImageWrapper->GetCompressed(Quality);
+			Players->Apply([&JpegBytes, this](FPixelStreamingPlayerId PlayerId, FPlayerContext& PlayerContext) {
+				if (PlayerContext.DataChannel)
+				{
+					FPixelStreamingInputMessage* Message = FPixelStreamingInputProtocol::FromStreamerProtocol.Find("FreezeFrame");
+					if(Message != nullptr)
+					{
+						PlayerContext.DataChannel->SendArbitraryData(Message->GetID(), JpegBytes);
+					}
+					else
+					{
+						UE_LOG(LogPixelStreaming, Error, TEXT("Could not find datachannel message with ID FreezeFrame."));
+					}
+				}
+			});
+			CachedJpegBytes = JpegBytes;
+		}
+		else
+		{
+			UE_LOG(LogPixelStreaming, Error, TEXT("JPEG image wrapper failed to accept frame data"));
+		}
+	}
+
+	void FFreezeFrame::SetupFreezeFrameCapture()
+	{
+		// Remove any existing binding
+		RemoveFreezeFrameBinding();
+
+		if(TSharedPtr<FPixelStreamingVideoInput> ConcreteVideoInput = VideoInput.Pin())
+		{
+			OnFrameCapturedForFreezeFrameHandle = ConcreteVideoInput->OnFrameCaptured.AddSP(this, &FFreezeFrame::FreezeFrameCapture);
+		}
+	}
+
+	void FFreezeFrame::RemoveFreezeFrameBinding()
+	{
+		if(!OnFrameCapturedForFreezeFrameHandle)
+		{
+			return;
+		}
+
+		if(TSharedPtr<FPixelStreamingVideoInput> ConcreteVideoInput = VideoInput.Pin())
+		{
+			ConcreteVideoInput->OnFrameCaptured.Remove(OnFrameCapturedForFreezeFrameHandle.GetValue());
+			OnFrameCapturedForFreezeFrameHandle.Reset();
+		}
+	}
+
+	void FFreezeFrame::FreezeFrameCapture()
+	{
+		TSharedPtr<FPixelStreamingVideoInput> Input = VideoInput.Pin();
+
+		if(!Input)
+		{
+			return;
+		}
+
+		if(Settings::IsCodecVPX())
+		{
+			// Request output format is I420 for VPX
+			TSharedPtr<IPixelCaptureOutputFrame> OutputFrame = Input->RequestFormat(PixelCaptureBufferFormat::FORMAT_I420);
+			if (OutputFrame)
+			{
+				// Can remove binding now we have got the output in the format we need to send a FF
+				RemoveFreezeFrameBinding();
+
+				FPixelCaptureOutputFrameI420* I420Frame = StaticCast<FPixelCaptureOutputFrameI420*>(OutputFrame.Get());
+				TSharedPtr<FPixelCaptureBufferI420> I420Buffer = I420Frame->GetI420Buffer();
+
+				const uint32_t NumBytes = webrtc::CalcBufferSize(webrtc::VideoType::kARGB, I420Frame->GetWidth(), I420Frame->GetHeight());
+				const uint32_t NumPixels = I420Frame->GetWidth() * I420Frame->GetHeight();
+				uint8_t* ARGBBuffer = new uint8_t[NumBytes];
+
+				// Convert I420 to ARGB
+				libyuv::ConvertFromI420(
+					I420Buffer->GetDataY(), I420Buffer->GetStrideY(),
+					I420Buffer->GetDataU(), I420Buffer->GetStrideUV(),
+					I420Buffer->GetDataV(), I420Buffer->GetStrideUV(),
+					ARGBBuffer, 0,
+					I420Buffer->GetWidth(), I420Buffer->GetHeight(),
+					libyuv::FOURCC_ARGB);
+
+				// We assume FColor packing is same ordering as the Buffer are copying from
+				TArray<FColor> PixelArr((FColor*)ARGBBuffer, NumPixels);
+				FIntRect Rect(0, 0, I420Frame->GetWidth(), I420Frame->GetHeight());
+				SendFreezeFrame(MoveTemp(PixelArr), Rect);
+			}
+		}
+		else
+		{
+			TSharedPtr<IPixelCaptureOutputFrame> OutputFrame = Input->RequestFormat(PixelCaptureBufferFormat::FORMAT_RHI);
+			if (OutputFrame)
+			{
+				// Can remove binding now we have got the output in the format we need to send a FF
+				RemoveFreezeFrameBinding();
+
+				TWeakPtr<FFreezeFrame> WeakSelf = AsShared();
+
+				ENQUEUE_RENDER_COMMAND(ReadSurfaceCommand)
+				([WeakSelf, OutputFrame](FRHICommandListImmediate& RHICmdList) {
+					if (auto ThisPtr = WeakSelf.Pin())
+					{
+						FPixelCaptureOutputFrameRHI* RHISourceFrame = StaticCast<FPixelCaptureOutputFrameRHI*>(OutputFrame.Get());
+
+						// Read the data out of the back buffer and send as a JPEG.
+						FIntRect Rect(0, 0, RHISourceFrame->GetWidth(), RHISourceFrame->GetHeight());
+						TArray<FColor> Data;
+
+						RHICmdList.ReadSurfaceData(RHISourceFrame->GetFrameTexture(), Rect, Data, FReadSurfaceDataFlags());
+						ThisPtr->SendFreezeFrame(MoveTemp(Data), Rect);
+					}
+				});
+			}
+		}
+
+	}
+
+} // namespace UE::PixelStreaming
\ No newline at end of file
diff --git a/Engine/Plugins/Media/PixelStreaming/Source/PixelStreaming/Private/FreezeFrame.h b/Engine/Plugins/Media/PixelStreaming/Source/PixelStreaming/Private/FreezeFrame.h
new file mode 100644
--- /dev/null
+++ b/Engine/Plugins/Media/PixelStreaming/Source/PixelStreaming/Private/FreezeFrame.h
@@ -0,0 +1,65 @@
+// Copyright Epic Games, Inc. All Rights Reserved.
+
+#pragma once
+
+#include "Containers/ContainersFwd.h"
+#include "PlayerContext.h"
+#include "PixelStreamingPlayerId.h"
+#include "PixelStreamingVideoInput.h"
+#include "Engine/Texture2D.h"
+#include "ThreadSafeMap.h"
+#include "Templates/SharedPointer.h"
+
+namespace UE::PixelStreaming
+{
+	/*
+	* The user of the Pixel Streaming application can trigger a freeze frame to happen at any time during stream.
+	* What a freeze frame will do is either:
+	* 1) Mode 1 - Send a known texture/image (effectively a still image, e.g. a loading screen).
+	* 2) Mode 2 - Send the current rendered frame out of UE as an image
+	* This freeze frame will displayed over the top of the video stream with no compression.
+	* So a crisp, high quality image, is guaranteed.
+	* While this is occuring the stream is not technically frozen, it still encodes and send frames;
+	* however, visually, the user is only shown this still. This technique can be used to hide large, known,
+	* cpu/gpu lags in the Pixel Streaming experience, for example during a level load or some expensive transition.
+	* The better solution is to not lag your application in the first place, but this can be a bandaid.
+	*/
+	class FFreezeFrame : public TSharedFromThis<FFreezeFrame>
+	{
+	public:
+		static TSharedPtr<FFreezeFrame> Create(TWeakPtr<TThreadSafeMap<FPixelStreamingPlayerId, FPlayerContext>> InPlayers);
+		virtual ~FFreezeFrame();
+		void SetVideoInput(TWeakPtr<FPixelStreamingVideoInput> VideoInput);
+		void SetPlayers(TWeakPtr<TThreadSafeMap<FPixelStreamingPlayerId, FPlayerContext>> InPlayers);
+
+		/**
+		 * @brief Start the freeze frame process, either mode 1 or 2 depending on if a texture is passed.
+		 * @param Texture - The still image to send as a freeze frame, pass nullptr to use send the backbuffer instead.
+		*/
+		void StartFreeze(UTexture2D* Texture);
+		void StopFreeze();
+		void SendCachedFreezeFrameTo(FPixelStreamingPlayerId PlayerId) const;
+
+	protected:
+		FFreezeFrame() = default;
+
+	private:
+		void SendFreezeFrame(TArray<FColor> RawData, const FIntRect& Rect);
+		void SetupFreezeFrameCapture();
+		void RemoveFreezeFrameBinding();
+		void FreezeFrameCapture();
+
+	private:
+		TWeakPtr<TThreadSafeMap<FPixelStreamingPlayerId, FPlayerContext>> WeakPlayers;
+
+		// Video input used to capture the frame if mode 2 is used.
+		TWeakPtr<FPixelStreamingVideoInput> VideoInput;
+
+		// When we send a freeze frame we retain the data so we send freeze frame to new peers if they join during a freeze frame.
+		TArray64<uint8> CachedJpegBytes;
+
+		// Delegate handle for when we bind to the OnFrameCaptured delegate of the video input - we unbind from this once the freeze frame is captured
+		TOptional<FDelegateHandle> OnFrameCapturedForFreezeFrameHandle;
+	};
+
+} // namespace UE::PixelStreaming
\ No newline at end of file
diff --git a/Engine/Plugins/Media/PixelStreaming/Source/PixelStreaming/Private/PixelStreamingModule.cpp b/Engine/Plugins/Media/PixelStreaming/Source/PixelStreaming/Private/PixelStreamingModule.cpp
--- a/Engine/Plugins/Media/PixelStreaming/Source/PixelStreaming/Private/PixelStreamingModule.cpp
+++ b/Engine/Plugins/Media/PixelStreaming/Source/PixelStreaming/Private/PixelStreamingModule.cpp
@@ -18,7 +18,6 @@
 #include "PixelStreamingUtils.h"
 #include "PixelStreamingCoderUtils.h"
 #include "Utils.h"
-#include "UtilsRender.h"
 
 #if PLATFORM_LINUX
 	#include "CudaModule.h"
diff --git a/Engine/Plugins/Media/PixelStreaming/Source/PixelStreaming/Private/SignallingConnectionObserver.cpp b/Engine/Plugins/Media/PixelStreaming/Source/PixelStreaming/Private/SignallingConnectionObserver.cpp
--- a/Engine/Plugins/Media/PixelStreaming/Source/PixelStreaming/Private/SignallingConnectionObserver.cpp
+++ b/Engine/Plugins/Media/PixelStreaming/Source/PixelStreaming/Private/SignallingConnectionObserver.cpp
@@ -56,7 +56,7 @@ namespace UE::PixelStreaming
 
 	void FPixelStreamingSignallingConnectionObserver::OnSignallingRemoteIceCandidate(FPixelStreamingPlayerId PlayerId, const FString& SdpMid, int SdpMLineIndex, const FString& Sdp)
 	{
-		if (FPlayerContext* PlayerContext = Streamer.Players.Find(PlayerId))
+		if (FPlayerContext* PlayerContext = Streamer.Players->Find(PlayerId))
 		{
 			PlayerContext->PeerConnection->AddRemoteIceCandidate(SdpMid, SdpMLineIndex, Sdp);
 		}
@@ -75,7 +75,7 @@ namespace UE::PixelStreaming
 
 	void FPixelStreamingSignallingConnectionObserver::OnSignallingSFUPeerDataChannels(FPixelStreamingPlayerId SFUId, FPixelStreamingPlayerId PlayerId, int32 SendStreamId, int32 RecvStreamId)
 	{
-		FPlayerContext* PlayerContext = Streamer.Players.Find(SFUId);
+		FPlayerContext* PlayerContext = Streamer.Players->Find(SFUId);
 		if (PlayerContext == nullptr)
 		{
 			UE_LOG(LogPixelStreaming, Error, TEXT("Trying to create data channels from SFU connection but no SFU connection found."));
diff --git a/Engine/Plugins/Media/PixelStreaming/Source/PixelStreaming/Private/Streamer.cpp b/Engine/Plugins/Media/PixelStreaming/Source/PixelStreaming/Private/Streamer.cpp
--- a/Engine/Plugins/Media/PixelStreaming/Source/PixelStreaming/Private/Streamer.cpp
+++ b/Engine/Plugins/Media/PixelStreaming/Source/PixelStreaming/Private/Streamer.cpp
@@ -2,7 +2,6 @@
 
 #include "Streamer.h"
 #include "IPixelStreamingModule.h"
-#include "PixelStreamingPrivate.h"
 #include "PixelStreamingDelegates.h"
 #include "PixelStreamingSignallingConnection.h"
 #include "PixelStreamingAudioDeviceModule.h"
@@ -18,10 +17,7 @@
 #include "Async/Async.h"
 #include "Engine/Texture2D.h"
 #include "RTCStatsCollector.h"
-#include "IImageWrapper.h"
-#include "IImageWrapperModule.h"
 #include "Framework/Application/SlateApplication.h"
-#include "UtilsRender.h"
 #include "PixelStreamingCodec.h"
 #include "Dom/JsonValue.h"
 #include "Serialization/JsonWriter.h"
@@ -31,9 +27,6 @@
 #include "PixelStreamingInputProtocol.h"
 #include "IPixelStreamingModule.h"
 #include "IPixelStreamingInputModule.h"
-#include "PixelCaptureBufferFormat.h"
-#include "PixelCaptureOutputFrameRHI.h"
-#include "PixelCaptureInputFrameRHI.h"
 #include "SignallingConnectionObserver.h"
 #include "ToStringExtensions.h"
 #include "Settings.h"
@@ -52,7 +45,9 @@ namespace UE::PixelStreaming
 	FStreamer::FStreamer(const FString& InStreamerId)
 		: StreamerId(InStreamerId)
 		, InputHandler(IPixelStreamingInputModule::Get().CreateInputHandler())
+		, Players(new TThreadSafeMap<FPixelStreamingPlayerId, FPlayerContext>())
 		, Module(IPixelStreamingModule::Get())
+		, FreezeFrame(FFreezeFrame::Create(Players))
 	{
 		VideoSourceGroup = FVideoSourceGroup::Create();
 		Observer = MakeShared<FPixelStreamingSignallingConnectionObserver>(*this);
@@ -72,7 +67,7 @@ namespace UE::PixelStreaming
 
 	void FStreamer::OnProtocolUpdated()
 	{
-		Players.Apply([this](FPixelStreamingPlayerId DataPlayerId, FPlayerContext& PlayerContext) {
+		Players->Apply([this](FPixelStreamingPlayerId DataPlayerId, FPlayerContext& PlayerContext) {
 			if (PlayerContext.DataChannel)
 			{
 				SendProtocol(DataPlayerId);
@@ -98,32 +93,7 @@ namespace UE::PixelStreaming
 	void FStreamer::SetVideoInput(TSharedPtr<FPixelStreamingVideoInput> Input)
 	{
 		VideoSourceGroup->SetVideoInput(Input);
-		TWeakPtr<FStreamer> WeakSelf = AsShared();
-		Input->OnFrameCaptured.AddLambda([WeakSelf, Input]() {
-			if (auto ThisPtr = WeakSelf.Pin())
-			{
-				TSharedPtr<IPixelCaptureOutputFrame> OutputFrame = Input->RequestFormat(PixelCaptureBufferFormat::FORMAT_RHI);
-				if (OutputFrame && ThisPtr->bCaptureNextBackBufferAndStream)
-				{
-					ThisPtr->bCaptureNextBackBufferAndStream = false;
-
-					ENQUEUE_RENDER_COMMAND(ReadSurfaceCommand)
-					([WeakSelf, OutputFrame](FRHICommandListImmediate& RHICmdList) {
-						if (auto ThisPtr = WeakSelf.Pin())
-						{
-							FPixelCaptureOutputFrameRHI* RHISourceFrame = StaticCast<FPixelCaptureOutputFrameRHI*>(OutputFrame.Get());
-
-							// Read the data out of the back buffer and send as a JPEG.
-							FIntRect Rect(0, 0, RHISourceFrame->GetWidth(), RHISourceFrame->GetHeight());
-							TArray<FColor> Data;
-
-							RHICmdList.ReadSurfaceData(RHISourceFrame->GetFrameTexture(), Rect, Data, FReadSurfaceDataFlags());
-							ThisPtr->SendFreezeFrame(MoveTemp(Data), Rect);
-						}
-					});
-				}
-			}
-		});
+		FreezeFrame->SetVideoInput(Input);
 	}
 
 	TWeakPtr<FPixelStreamingVideoInput> FStreamer::GetVideoInput()
@@ -284,59 +254,19 @@ namespace UE::PixelStreaming
 
 	void FStreamer::FreezeStream(UTexture2D* Texture)
 	{
-		if (Texture)
-		{
-			ENQUEUE_RENDER_COMMAND(ReadSurfaceCommand)
-			([this, Texture](FRHICommandListImmediate& RHICmdList) {
-				// A frame is supplied so immediately read its data and send as a JPEG.
-				FTextureRHIRef TextureRHI = Texture->GetResource() ? Texture->GetResource()->TextureRHI : nullptr;
-				if (!TextureRHI)
-				{
-					UE_LOG(LogPixelStreaming, Error, TEXT("Attempting freeze frame with texture %s with no texture RHI"), *Texture->GetName());
-					return;
-				}
-				uint32 Width = TextureRHI->GetDesc().Extent.X;
-				uint32 Height = TextureRHI->GetDesc().Extent.Y;
-
-				FTextureRHIRef DestTexture = CreateRHITexture(Width, Height);
-
-				FGPUFenceRHIRef CopyFence = GDynamicRHI->RHICreateGPUFence(*FString::Printf(TEXT("FreezeFrameFence")));
-
-				// Copy freeze frame texture to empty texture
-				CopyTexture(RHICmdList, TextureRHI, DestTexture, CopyFence);
-
-				TArray<FColor> Data;
-				FIntRect Rect(0, 0, Width, Height);
-				RHICmdList.ReadSurfaceData(DestTexture, Rect, Data, FReadSurfaceDataFlags());
-				SendFreezeFrame(MoveTemp(Data), Rect);
-			});
-		}
-		else
-		{
-			// A frame is not supplied, so we need to capture the back buffer at
-			// the next opportunity, and send as a JPEG.
-			bCaptureNextBackBufferAndStream = true;
-		}
+		FreezeFrame->StartFreeze(Texture);
 	}
 
 	void FStreamer::UnfreezeStream()
 	{
-		// Force a keyframe so when stream unfreezes if player has never received a h.264 frame before they can still connect.
+		// Force a keyframe so when stream unfreezes if player has never received a frame before they can still connect.
 		ForceKeyFrame();
-
-		Players.Apply([this](FPixelStreamingPlayerId PlayerId, FPlayerContext& PlayerContext) {
-			if (PlayerContext.DataChannel)
-			{
-				PlayerContext.DataChannel->SendMessage(FPixelStreamingInputProtocol::FromStreamerProtocol.Find("UnfreezeFrame")->GetID());
-			}
-		});
-
-		CachedJpegBytes.Empty();
+		FreezeFrame->StopFreeze();
 	}
 
 	void FStreamer::SendPlayerMessage(uint8 Type, const FString& Descriptor)
 	{
-		Players.Apply([&Type, &Descriptor](FPixelStreamingPlayerId PlayerId, FPlayerContext& PlayerContext) {
+		Players->Apply([&Type, &Descriptor](FPixelStreamingPlayerId PlayerId, FPlayerContext& PlayerContext) {
 			if (PlayerContext.DataChannel)
 			{
 				PlayerContext.DataChannel->SendMessage(Type, Descriptor);
@@ -350,7 +280,7 @@ namespace UE::PixelStreaming
 		// channels it might be a bad idea. At some point it would be good to take a snapshot of the
 		// keys in the map when we start, then one by one get the channel and send the data
 
-		Players.Apply([&ByteData, &MimeType, &FileExtension, this](FPixelStreamingPlayerId PlayerId, FPlayerContext& PlayerContext) {
+		Players->Apply([&ByteData, &MimeType, &FileExtension, this](FPixelStreamingPlayerId PlayerId, FPlayerContext& PlayerContext) {
 			if (PlayerContext.DataChannel)
 			{
 				// Send the mime type first
@@ -388,7 +318,7 @@ namespace UE::PixelStreaming
 	TArray<FPixelStreamingPlayerId> FStreamer::GetConnectedPlayers()
 	{
 		TArray<FPixelStreamingPlayerId> ConnectedPlayerIds;
-		Players.Apply([&ConnectedPlayerIds, this](FPixelStreamingPlayerId PlayerId, FPlayerContext& PlayerContext) {
+		Players->Apply([&ConnectedPlayerIds, this](FPixelStreamingPlayerId PlayerId, FPlayerContext& PlayerContext) {
 			ConnectedPlayerIds.Add(PlayerId);
 		});
 		return ConnectedPlayerIds;
@@ -401,7 +331,7 @@ namespace UE::PixelStreaming
 
 	IPixelStreamingAudioSink* FStreamer::GetPeerAudioSink(FPixelStreamingPlayerId PlayerId)
 	{
-		if (FPlayerContext* PlayerContext = Players.Find(PlayerId))
+		if (FPlayerContext* PlayerContext = Players->Find(PlayerId))
 		{
 			return PlayerContext->PeerConnection->GetAudioSink().Get();
 		}
@@ -411,7 +341,7 @@ namespace UE::PixelStreaming
 	IPixelStreamingAudioSink* FStreamer::GetUnlistenedAudioSink()
 	{
 		IPixelStreamingAudioSink* Result = nullptr;
-		Players.ApplyUntil([&Result](FPixelStreamingPlayerId PlayerId, FPlayerContext& PlayerContext) {
+		Players->ApplyUntil([&Result](FPixelStreamingPlayerId PlayerId, FPlayerContext& PlayerContext) {
 			if (PlayerContext.PeerConnection)
 			{
 				if (!PlayerContext.PeerConnection->GetAudioSink()->HasAudioConsumers())
@@ -492,7 +422,7 @@ namespace UE::PixelStreaming
 
 	void FStreamer::RefreshStreamBitrate()
 	{
-		Players.Apply([this](FPixelStreamingPlayerId PlayerId, FPlayerContext& PlayerContext) {
+		Players->Apply([this](FPixelStreamingPlayerId PlayerId, FPlayerContext& PlayerContext) {
 			if (PlayerContext.PeerConnection)
 			{
 				PlayerContext.PeerConnection->RefreshStreamBitrate();
@@ -502,12 +432,12 @@ namespace UE::PixelStreaming
 
 	void FStreamer::ForEachPlayer(const TFunction<void(FPixelStreamingPlayerId, FPlayerContext)>& Func)
 	{
-		Players.Apply(Func);
+		Players->Apply(Func);
 	}
 
 	bool FStreamer::CreateSession(FPixelStreamingPlayerId PlayerId)
 	{
-		if (FPlayerContext* PlayerContext = Players.Find(PlayerId))
+		if (FPlayerContext* PlayerContext = Players->Find(PlayerId))
 		{
 			if (PlayerContext->Config.IsSFU && SFUPlayerId != INVALID_PLAYER_ID)
 			{
@@ -559,10 +489,10 @@ namespace UE::PixelStreaming
 	{
 		if (VideoSourceGroup->GetVideoInput() != nullptr)
 		{
-			if (FPlayerContext* PlayerContext = Players.Find(PlayerId))
+			if (FPlayerContext* PlayerContext = Players->Find(PlayerId))
 			{
 				const bool AllowSimulcast = PlayerContext->Config.IsSFU;
-				
+
 				PlayerContext->PeerConnection->SetVideoSource(VideoSourceGroup->CreateVideoSource([this, PlayerId]() { return ShouldPeerGenerateFrames(PlayerId); }));
 				PlayerContext->PeerConnection->SetAudioSource(FPixelStreamingPeerConnection::GetApplicationAudioSource());
 				PlayerContext->PeerConnection->SetAudioSink(MakeShared<FAudioSink>());
@@ -579,7 +509,7 @@ namespace UE::PixelStreaming
 			return;
 		}
 
-		FPlayerContext& PlayerContext = Players.GetOrAdd(PlayerId);
+		FPlayerContext& PlayerContext = Players->GetOrAdd(PlayerId);
 		PlayerContext.Config = PlayerConfig;
 
 		// create peer connection
@@ -617,7 +547,7 @@ namespace UE::PixelStreaming
 	{
 		if (StatName == PixelStreamingStatNames::MeanQPPerSecond)
 		{
-			if (FPlayerContext* PlayerContext = Players.Find(PlayerId))
+			if (FPlayerContext* PlayerContext = Players->Find(PlayerId))
 			{
 				if (PlayerContext->DataChannel)
 				{
@@ -629,7 +559,7 @@ namespace UE::PixelStreaming
 
 	void FStreamer::OnOffer(FPixelStreamingPlayerId PlayerId, const FString& Sdp)
 	{
-		FPlayerContext& PlayerContext = Players.GetOrAdd(PlayerId);
+		FPlayerContext& PlayerContext = Players->GetOrAdd(PlayerId);
 
 		FPixelStreamingPlayerConfig Config;
 		Config.SupportsDataChannel = true;
@@ -660,7 +590,7 @@ namespace UE::PixelStreaming
 
 	void FStreamer::OnAnswer(FPixelStreamingPlayerId PlayerId, const FString& Sdp)
 	{
-		if (FPlayerContext* PlayerContext = Players.Find(PlayerId))
+		if (FPlayerContext* PlayerContext = Players->Find(PlayerId))
 		{
 			PlayerContext->PeerConnection->ReceiveAnswer(
 				Sdp,
@@ -684,7 +614,7 @@ namespace UE::PixelStreaming
 		// delete on the signalling thread which might be waiting for the players
 		// lock.
 		FPlayerContext PendingDeletePlayer;
-		if (FPlayerContext* PlayerContext = Players.Find(PlayerId))
+		if (FPlayerContext* PlayerContext = Players->Find(PlayerId))
 		{
 			// when a sfu is connected we only get disconnect messages.
 			// we dont get connect messages but we might get datachannel requests which can result
@@ -696,7 +626,7 @@ namespace UE::PixelStreaming
 			PendingDeletePlayer = *PlayerContext;
 		}
 
-		Players.Remove(PlayerId);
+		Players->Remove(PlayerId);
 
 		// delete webrtc objects here outside the lock
 		PendingDeletePlayer.DataChannel.Reset();
@@ -713,7 +643,7 @@ namespace UE::PixelStreaming
 			SetQualityController(INVALID_PLAYER_ID);
 
 			// find the first non sfu peer and give it quality controller status
-			Players.ApplyUntil([this](FPixelStreamingPlayerId PlayerId, FPlayerContext& PlayerContext) {
+			Players->ApplyUntil([this](FPixelStreamingPlayerId PlayerId, FPlayerContext& PlayerContext) {
 				if (PlayerContext.PeerConnection)
 				{
 					if (PlayerId != SFUPlayerId)
@@ -730,7 +660,7 @@ namespace UE::PixelStreaming
 		{
 			Delegates->OnClosedConnection.Broadcast(StreamerId, PlayerId, bWasQualityController);
 			Delegates->OnClosedConnectionNative.Broadcast(StreamerId, PlayerId, bWasQualityController);
-			if (Players.IsEmpty())
+			if (Players->IsEmpty())
 			{
 				Delegates->OnAllConnectionsClosed.Broadcast(StreamerId);
 				Delegates->OnAllConnectionsClosedNative.Broadcast(StreamerId);
@@ -751,7 +681,7 @@ namespace UE::PixelStreaming
 		}
 
 		VideoSourceGroup->RemoveAllVideoSources();
-		Players.Empty();
+		Players->Empty();
 		SFUPlayerId = INVALID_PLAYER_ID;
 		QualityControllingId = INVALID_PLAYER_ID;
 		InputControllingId = INVALID_PLAYER_ID;
@@ -764,7 +694,7 @@ namespace UE::PixelStreaming
 
 	void FStreamer::AddNewDataChannel(FPixelStreamingPlayerId PlayerId, TSharedPtr<FPixelStreamingDataChannel> NewChannel)
 	{
-		FPlayerContext& PlayerContext = Players.GetOrAdd(PlayerId);
+		FPlayerContext& PlayerContext = Players->GetOrAdd(PlayerId);
 		PlayerContext.DataChannel = NewChannel;
 
 		TWeakPtr<FStreamer> WeakStreamer = AsShared();
@@ -802,21 +732,21 @@ namespace UE::PixelStreaming
 
 		if (UPixelStreamingDelegates* Delegates = UPixelStreamingDelegates::GetPixelStreamingDelegates())
 		{
-			FPlayerContext* PlayerContext = Players.Find(PlayerId);
+			FPlayerContext* PlayerContext = Players->Find(PlayerId);
 			Delegates->OnDataChannelOpenNative.Broadcast(StreamerId, PlayerId, PlayerContext->DataChannel.Get());
 		}
 
 		// When data channel is open
 		SendProtocol(PlayerId);
 		// Try to send cached freeze frame (if we have one)
-		SendCachedFreezeFrameTo(PlayerId);
+		FreezeFrame->SendCachedFreezeFrameTo(PlayerId);
 		SendInitialSettings(PlayerId);
 		SendPeerControllerMessages(PlayerId);
 	}
 
 	void FStreamer::OnDataChannelClosed(FPixelStreamingPlayerId PlayerId)
 	{
-		if (FPlayerContext* PlayerContext = Players.Find(PlayerId))
+		if (FPlayerContext* PlayerContext = Players->Find(PlayerId))
 		{
 			PlayerContext->DataChannel = nullptr;
 
@@ -824,7 +754,7 @@ namespace UE::PixelStreaming
 			{
 				InputControllingId = INVALID_PLAYER_ID;
 				// just get the first channel we have and give it input control.
-				Players.ApplyUntil([this](FPixelStreamingPlayerId PlayerId, FPlayerContext& PlayerContext) {
+				Players->ApplyUntil([this](FPixelStreamingPlayerId PlayerId, FPlayerContext& PlayerContext) {
 					if (PlayerContext.DataChannel)
 					{
 						if (PlayerId != SFUPlayerId)
@@ -865,7 +795,7 @@ namespace UE::PixelStreaming
 		}
 		else if (Type == FPixelStreamingInputProtocol::ToStreamerProtocol.Find("TestEcho")->GetID())
 		{
-			if (FPlayerContext* PlayerContext = Players.Find(PlayerId))
+			if (FPlayerContext* PlayerContext = Players->Find(PlayerId))
 			{
 				if (PlayerContext->DataChannel)
 				{
@@ -934,7 +864,7 @@ namespace UE::PixelStreaming
 
 		const FString FullPayload = FString::Printf(TEXT("{ \"PixelStreaming\": %s, \"Encoder\": %s, \"WebRTC\": %s, \"ConfigOptions\": %s }"), *PixelStreamingPayload, *EncoderPayload, *WebRTCPayload, *ConfigPayload);
 
-		if (const FPlayerContext* PlayerContext = Players.Find(PlayerId))
+		if (const FPlayerContext* PlayerContext = Players->Find(PlayerId))
 		{
 			if (PlayerContext->DataChannel)
 			{
@@ -960,7 +890,7 @@ namespace UE::PixelStreaming
 				return;
 			}
 
-			if (const FPlayerContext* PlayerContext = Players.Find(PlayerId))
+			if (const FPlayerContext* PlayerContext = Players->Find(PlayerId))
 			{
 				if (PlayerContext->DataChannel)
 				{
@@ -976,7 +906,7 @@ namespace UE::PixelStreaming
 
 	void FStreamer::SendPeerControllerMessages(FPixelStreamingPlayerId PlayerId) const
 	{
-		if (const FPlayerContext* PlayerContext = Players.Find(PlayerId))
+		if (const FPlayerContext* PlayerContext = Players->Find(PlayerId))
 		{
 			if (PlayerContext->DataChannel)
 			{
@@ -1029,7 +959,7 @@ namespace UE::PixelStreaming
 					TransmissionTimeMs);
 			}
 
-			if (const FPlayerContext* PlayerContext = Players.Find(PlayerId))
+			if (const FPlayerContext* PlayerContext = Players->Find(PlayerId))
 			{
 				if (PlayerContext->DataChannel)
 				{
@@ -1039,47 +969,9 @@ namespace UE::PixelStreaming
 		});
 	}
 
-	void FStreamer::SendFreezeFrame(TArray<FColor> RawData, const FIntRect& Rect)
-	{
-		IImageWrapperModule& ImageWrapperModule = FModuleManager::GetModuleChecked<IImageWrapperModule>(TEXT("ImageWrapper"));
-		TSharedPtr<IImageWrapper> ImageWrapper = ImageWrapperModule.CreateImageWrapper(EImageFormat::JPEG);
-		bool bSuccess = ImageWrapper->SetRaw(RawData.GetData(), RawData.Num() * sizeof(FColor), Rect.Width(), Rect.Height(), ERGBFormat::BGRA, 8);
-		if (bSuccess)
-		{
-			// Compress to a JPEG of the maximum possible quality.
-			int32 Quality = Settings::CVarPixelStreamingFreezeFrameQuality.GetValueOnAnyThread();
-			const TArray64<uint8>& JpegBytes = ImageWrapper->GetCompressed(Quality);
-			Players.Apply([&JpegBytes, this](FPixelStreamingPlayerId PlayerId, FPlayerContext& PlayerContext) {
-				if (PlayerContext.DataChannel)
-				{
-					PlayerContext.DataChannel->SendArbitraryData(FPixelStreamingInputProtocol::FromStreamerProtocol.Find("FreezeFrame")->GetID(), JpegBytes);
-				}
-			});
-			CachedJpegBytes = JpegBytes;
-		}
-		else
-		{
-			UE_LOG(LogPixelStreaming, Error, TEXT("JPEG image wrapper failed to accept frame data"));
-		}
-	}
-
-	void FStreamer::SendCachedFreezeFrameTo(FPixelStreamingPlayerId PlayerId) const
-	{
-		if (CachedJpegBytes.Num() > 0)
-		{
-			if (const FPlayerContext* PlayerContext = Players.Find(PlayerId))
-			{
-				if (PlayerContext->DataChannel)
-				{
-					PlayerContext->DataChannel->SendArbitraryData(FPixelStreamingInputProtocol::FromStreamerProtocol.Find("FreezeFrame")->GetID(), CachedJpegBytes);
-				}
-			}
-		}
-	}
-
 	bool FStreamer::ShouldPeerGenerateFrames(FPixelStreamingPlayerId PlayerId) const
 	{
-		if (const FPlayerContext* PlayerContext = Players.Find(PlayerId))
+		if (const FPlayerContext* PlayerContext = Players->Find(PlayerId))
 		{
 			EPixelStreamingCodec Codec = PlayerContext->PeerConnection->GetNegotiatedVideoCodec();
 
@@ -1119,7 +1011,7 @@ namespace UE::PixelStreaming
 	void FStreamer::SetQualityController(FPixelStreamingPlayerId PlayerId)
 	{
 		QualityControllingId = PlayerId;
-		Players.Apply([this](FPixelStreamingPlayerId DataPlayerId, FPlayerContext& PlayerContext) {
+		Players->Apply([this](FPixelStreamingPlayerId DataPlayerId, FPlayerContext& PlayerContext) {
 			if (PlayerContext.DataChannel)
 			{
 				const uint8 IsController = DataPlayerId == QualityControllingId ? 1 : 0;
diff --git a/Engine/Plugins/Media/PixelStreaming/Source/PixelStreaming/Private/Streamer.h b/Engine/Plugins/Media/PixelStreaming/Source/PixelStreaming/Private/Streamer.h
--- a/Engine/Plugins/Media/PixelStreaming/Source/PixelStreaming/Private/Streamer.h
+++ b/Engine/Plugins/Media/PixelStreaming/Source/PixelStreaming/Private/Streamer.h
@@ -12,6 +12,7 @@
 #include "PixelStreamingSignallingConnection.h"
 #include "Templates/SharedPointer.h"
 #include "PlayerContext.h"
+#include "FreezeFrame.h"
 
 class IPixelStreamingModule;
 
@@ -110,8 +111,6 @@ namespace UE::PixelStreaming
 		void SendProtocol(FPixelStreamingPlayerId PlayerId) const;
 		void SendPeerControllerMessages(FPixelStreamingPlayerId PlayerId) const;
 		void SendLatencyReport(FPixelStreamingPlayerId PlayerId) const;
-		void SendFreezeFrame(TArray<FColor> RawData, const FIntRect& Rect);
-		void SendCachedFreezeFrameTo(FPixelStreamingPlayerId PlayerId) const;
 		bool ShouldPeerGenerateFrames(FPixelStreamingPlayerId PlayerId) const;
 
 		void SetQualityController(FPixelStreamingPlayerId PlayerId);
@@ -129,17 +128,13 @@ namespace UE::PixelStreaming
 
 		webrtc::PeerConnectionInterface::RTCConfiguration PeerConnectionConfig;
 
-		TThreadSafeMap<FPixelStreamingPlayerId, FPlayerContext> Players;
+		TSharedPtr<TThreadSafeMap<FPixelStreamingPlayerId, FPlayerContext>> Players;
 
 		FPixelStreamingPlayerId QualityControllingId = INVALID_PLAYER_ID;
 		FPixelStreamingPlayerId SFUPlayerId = INVALID_PLAYER_ID;
 		FPixelStreamingPlayerId InputControllingId = INVALID_PLAYER_ID;
 
 		bool bStreamingStarted = false;
-		bool bCaptureNextBackBufferAndStream = false;
-
-		// When we send a freeze frame we retain the data so we send freeze frame to new peers if they join during a freeze frame.
-		TArray64<uint8> CachedJpegBytes;
 
 		FPreConnectionEvent StreamingPreConnectionEvent;
 		FStreamingStartedEvent StreamingStartedEvent;
@@ -155,6 +150,8 @@ namespace UE::PixelStreaming
 
 		IPixelStreamingModule& Module;
 
+		TSharedPtr<FFreezeFrame> FreezeFrame;
+
 		TMap<FName, FString> ConfigOptions;
 	};
 } // namespace UE::PixelStreaming
diff --git a/Engine/Plugins/Media/PixelStreaming/Source/PixelStreaming/Private/Tests/CodecUtils.h b/Engine/Plugins/Media/PixelStreaming/Source/PixelStreaming/Private/Tests/CodecUtils.h
--- a/Engine/Plugins/Media/PixelStreaming/Source/PixelStreaming/Private/Tests/CodecUtils.h
+++ b/Engine/Plugins/Media/PixelStreaming/Source/PixelStreaming/Private/Tests/CodecUtils.h
@@ -34,7 +34,14 @@ namespace
 		ADD_LATENT_AUTOMATION_COMMAND(FConnectPlayerAfterStreamerConnectedOrTimeout(5.0, Streamer, Player, PlayerPort))
 		ADD_LATENT_AUTOMATION_COMMAND(FWaitForPlayerConnectedOrTimeout(5.0, Player, PlayerPort))
 		ADD_LATENT_AUTOMATION_COMMAND(FWaitForICEConnectedOrTimeout(5.0, Player))
-		ADD_LATENT_AUTOMATION_COMMAND(FSendSolidColorFrame(VideoInput, FrameConfig))
+
+		// Send 30 frames
+		for(int i = 0; i < 30; i++)
+		{
+			ADD_LATENT_AUTOMATION_COMMAND(FSendSolidColorFrame(VideoInput, FrameConfig))
+			ADD_LATENT_AUTOMATION_COMMAND(FWaitSeconds(0.033)) // send at 30fps interval
+		}
+
 		ADD_LATENT_AUTOMATION_COMMAND(FWaitForFrameReceived(5.0, VideoSink, FrameConfig))
 		ADD_LATENT_AUTOMATION_COMMAND(FCleanupAll(SignallingServer, Streamer, Player))
 	}
@@ -66,7 +73,14 @@ namespace
 		for (int Res = 2; Res < 512; Res *= 2)
 		{
 			FMockVideoFrameConfig FrameConfig = { Res /*Width*/, Res /*Height*/, 255 /*Y*/, 0 /*U*/, 255 /*V*/ };
-			ADD_LATENT_AUTOMATION_COMMAND(FSendSolidColorFrame(VideoInput, FrameConfig))
+
+			// Send 30 frames
+			for(int i = 0; i < 30; i++)
+			{
+				ADD_LATENT_AUTOMATION_COMMAND(FSendSolidColorFrame(VideoInput, FrameConfig))
+				ADD_LATENT_AUTOMATION_COMMAND(FWaitSeconds(0.033)) // send at 30fps interval
+			}
+
 			ADD_LATENT_AUTOMATION_COMMAND(FWaitForFrameReceived(5.0, VideoSink, FrameConfig))
 		}
 
diff --git a/Engine/Plugins/Media/PixelStreaming/Source/PixelStreaming/Private/Tests/TestUtils.cpp b/Engine/Plugins/Media/PixelStreaming/Source/PixelStreaming/Private/Tests/TestUtils.cpp
--- a/Engine/Plugins/Media/PixelStreaming/Source/PixelStreaming/Private/Tests/TestUtils.cpp
+++ b/Engine/Plugins/Media/PixelStreaming/Source/PixelStreaming/Private/Tests/TestUtils.cpp
@@ -28,6 +28,16 @@ namespace UE::PixelStreaming
 
 	/* ---------- Latent Automation Commands ----------- */
 
+	bool FWaitSeconds::Update()
+	{
+		double DeltaTime = FPlatformTime::Seconds() - StartTime;
+		if (DeltaTime > WaitSeconds)
+		{
+			return true;
+		}
+		return false;
+	}
+
 	bool FSendSolidColorFrame::Update()
 	{
 		TSharedPtr<FPixelCaptureBufferI420> Buffer = MakeShared<FPixelCaptureBufferI420>(FrameConfig.Width, FrameConfig.Height);
@@ -50,7 +60,6 @@ namespace UE::PixelStreaming
 
 		FPixelCaptureInputFrameI420 Frame(Buffer);
 		VideoInput->OnFrame(Frame);
-
 		return true;
 	}
 
diff --git a/Engine/Plugins/Media/PixelStreaming/Source/PixelStreaming/Private/Tests/TestUtils.h b/Engine/Plugins/Media/PixelStreaming/Source/PixelStreaming/Private/Tests/TestUtils.h
--- a/Engine/Plugins/Media/PixelStreaming/Source/PixelStreaming/Private/Tests/TestUtils.h
+++ b/Engine/Plugins/Media/PixelStreaming/Source/PixelStreaming/Private/Tests/TestUtils.h
@@ -46,6 +46,7 @@ namespace UE::PixelStreaming
 		FDelegateHandle MessageReceivedHandle;
 	};
 
+	DEFINE_LATENT_AUTOMATION_COMMAND_ONE_PARAMETER(FWaitSeconds, double, WaitSeconds);
 	DEFINE_LATENT_AUTOMATION_COMMAND_THREE_PARAMETER(FSendDataChannelMessageToPlayer, TSharedPtr<IPixelStreamingStreamer>, Streamer, uint8, Id, const FString, Body);
 	DEFINE_LATENT_AUTOMATION_COMMAND_THREE_PARAMETER(FSendDataChannelMessageToStreamer, TSharedPtr<FMockPlayer>, Player, uint8, Id, const FString, Body);
 	DEFINE_LATENT_AUTOMATION_COMMAND_THREE_PARAMETER(FSendCustomMessageToStreamer, TSharedPtr<FMockPlayer>, Player, uint8, Id, uint16, Body);
diff --git a/Engine/Plugins/Media/PixelStreaming/Source/PixelStreaming/Private/UtilsRender.h b/Engine/Plugins/Media/PixelStreaming/Source/PixelStreaming/Private/UtilsRender.h
deleted file mode 100644
