Fix: VCam not working on Windows with -Vulkan

diff --git a/Engine/Plugins/Media/PixelStreaming/Source/PixelStreaming/Private/PixelStreamingMediaIOCapture.cpp b/Engine/Plugins/Media/PixelStreaming/Source/PixelStreaming/Private/PixelStreamingMediaIOCapture.cpp
--- a/Engine/Plugins/Media/PixelStreaming/Source/PixelStreaming/Private/PixelStreamingMediaIOCapture.cpp
+++ b/Engine/Plugins/Media/PixelStreaming/Source/PixelStreaming/Private/PixelStreamingMediaIOCapture.cpp
@@ -162,7 +162,18 @@ ETextureCreateFlags UPixelStreamingMediaIOCapture::GetOutputTextureFlags() const
 #if PLATFORM_MAC
 	return TexCreate_CPUReadback;
 #else
-	return TexCreate_Shared | TexCreate_RenderTargetable | TexCreate_UAV;
+	ETextureCreateFlags Flags = TexCreate_RenderTargetable | TexCreate_UAV;
+	
+	if (RHIGetInterfaceType() == ERHIInterfaceType::Vulkan)
+	{
+		Flags |= TexCreate_External;
+	}
+	else if (RHIGetInterfaceType() == ERHIInterfaceType::D3D11 || RHIGetInterfaceType() == ERHIInterfaceType::D3D12)
+	{
+		Flags |= TexCreate_Shared;
+	}
+
+	return Flags;
 #endif
 }
 
