QOL: Remove deprecated functions

diff --git a/Engine/Plugins/Media/PixelStreaming/Source/PixelStreaming/Private/PixelStreamingModule.cpp b/Engine/Plugins/Media/PixelStreaming/Source/PixelStreaming/Private/PixelStreamingModule.cpp
--- a/Engine/Plugins/Media/PixelStreaming/Source/PixelStreaming/Private/PixelStreamingModule.cpp
+++ b/Engine/Plugins/Media/PixelStreaming/Source/PixelStreaming/Private/PixelStreamingModule.cpp
@@ -281,11 +281,6 @@ namespace UE::PixelStreaming
 		return StreamerKeys;
 	}
 
-	TSharedPtr<IPixelStreamingStreamer> FPixelStreamingModule::GetStreamer(const FString& StreamerId)
-	{
-		return FindStreamer(StreamerId);
-	}
-
 	TSharedPtr<IPixelStreamingStreamer> FPixelStreamingModule::FindStreamer(const FString& StreamerId)
 	{
 		FScopeLock Lock(&StreamersCS);
@@ -536,42 +531,6 @@ namespace UE::PixelStreaming
 	/**
 	 * End own methods
 	 */
-
-	/**
-	 * Deprecated methods
-	 */
-	const FPixelStreamingInputProtocol FPixelStreamingModule::GetProtocol()
-	{
-		return FPixelStreamingInputProtocol();
-	}
-
-	void FPixelStreamingModule::RegisterMessage(EPixelStreamingMessageDirection MessageDirection, const FString& MessageType, FPixelStreamingInputMessage Message, const TFunction<void(FString, FMemoryReader)>& Handler)
-	{
-		if (MessageDirection == EPixelStreamingMessageDirection::ToStreamer)
-		{
-			FPixelStreamingInputProtocol::ToStreamerProtocol.Add(MessageType, Message);
-			if (TSharedPtr<IPixelStreamingInputHandler> InputHandler = DefaultStreamer->GetInputHandler().Pin())
-			{
-				InputHandler->RegisterMessageHandler(MessageType, Handler);
-			}
-		}
-		else if (MessageDirection == EPixelStreamingMessageDirection::FromStreamer)
-		{
-			FPixelStreamingInputProtocol::FromStreamerProtocol.Add(MessageType, Message);
-		}
-	}
-
-	TFunction<void(FString, FMemoryReader)> FPixelStreamingModule::FindMessageHandler(const FString& MessageType)
-	{
-		if (TSharedPtr<IPixelStreamingInputHandler> InputHandler = DefaultStreamer->GetInputHandler().Pin())
-		{
-			return InputHandler->FindMessageHandler(MessageType);
-		}
-		return {};
-	}
-	/**
-	 * End deprecated methods
-	 */
 } // namespace UE::PixelStreaming
 
 IMPLEMENT_MODULE(UE::PixelStreaming::FPixelStreamingModule, PixelStreaming)
diff --git a/Engine/Plugins/Media/PixelStreaming/Source/PixelStreaming/Private/PixelStreamingModule.h b/Engine/Plugins/Media/PixelStreaming/Source/PixelStreaming/Private/PixelStreamingModule.h
--- a/Engine/Plugins/Media/PixelStreaming/Source/PixelStreaming/Private/PixelStreamingModule.h
+++ b/Engine/Plugins/Media/PixelStreaming/Source/PixelStreaming/Private/PixelStreamingModule.h
@@ -33,7 +33,6 @@ namespace UE::PixelStreaming
 		virtual void StopStreaming() override;
 		virtual TSharedPtr<IPixelStreamingStreamer> CreateStreamer(const FString& StreamerId) override;
 		virtual TArray<FString> GetStreamerIds() override;
-		virtual TSharedPtr<IPixelStreamingStreamer> GetStreamer(const FString& StreamerId) override;
 		virtual TSharedPtr<IPixelStreamingStreamer> FindStreamer(const FString& StreamerId) override;
 		virtual TSharedPtr<IPixelStreamingStreamer> DeleteStreamer(const FString& StreamerId) override;
 		void DeleteStreamer(TSharedPtr<IPixelStreamingStreamer> ToBeDeleted) override;
@@ -60,12 +59,6 @@ namespace UE::PixelStreaming
 		void ShutdownModule() override;
 		/** End IModuleInterface implementation */
 
-		/** Deprected methods from the IPixelStreamingModule. We're only keeping these here until we can fully remove them */
-		const FPixelStreamingInputProtocol GetProtocol() override;
-		void RegisterMessage(EPixelStreamingMessageDirection MessageDirection, const FString& MessageType, FPixelStreamingInputMessage Message, const TFunction<void(FString, FMemoryReader)>& Handler) override;
-		TFunction<void(FString, FMemoryReader)> FindMessageHandler(const FString& MessageType) override;
-		/** End deprecated methods */
-
 		// Own methods
 		void InitDefaultStreamer();
 		bool IsPlatformCompatible() const;
diff --git a/Engine/Plugins/Media/PixelStreaming/Source/PixelStreaming/Private/Streamer.cpp b/Engine/Plugins/Media/PixelStreaming/Source/PixelStreaming/Private/Streamer.cpp
--- a/Engine/Plugins/Media/PixelStreaming/Source/PixelStreaming/Private/Streamer.cpp
+++ b/Engine/Plugins/Media/PixelStreaming/Source/PixelStreaming/Private/Streamer.cpp
@@ -121,24 +121,6 @@ namespace UE::PixelStreaming
 		return InputHandler->GetTargetWindow();
 	}
 
-	void FStreamer::SetTargetScreenSize(TWeakPtr<FIntPoint> InTargetScreenSize)
-	{
-		// This method is marked as deprecated but still calls the deprecated method on the input handler. As such, we disable
-		// the warnings that arise from using the input handlers method
-		PRAGMA_DISABLE_DEPRECATION_WARNINGS
-		InputHandler->SetTargetScreenSize(InTargetScreenSize);
-		PRAGMA_ENABLE_DEPRECATION_WARNINGS
-	}
-
-	TWeakPtr<FIntPoint> FStreamer::GetTargetScreenSize()
-	{
-		// This method is marked as deprecated but still calls the deprecated method on the input handler. As such, we disable
-		// the warnings that arise from using the input handlers method
-        PRAGMA_DISABLE_DEPRECATION_WARNINGS
-		return InputHandler->GetTargetScreenSize();
-        PRAGMA_ENABLE_DEPRECATION_WARNINGS
-	}
-
 	void FStreamer::SetTargetScreenRect(TWeakPtr<FIntRect> InTargetScreenRect)
 	{
 		InputHandler->SetTargetScreenRect(InTargetScreenRect);
diff --git a/Engine/Plugins/Media/PixelStreaming/Source/PixelStreaming/Private/Streamer.h b/Engine/Plugins/Media/PixelStreaming/Source/PixelStreaming/Private/Streamer.h
--- a/Engine/Plugins/Media/PixelStreaming/Source/PixelStreaming/Private/Streamer.h
+++ b/Engine/Plugins/Media/PixelStreaming/Source/PixelStreaming/Private/Streamer.h
@@ -36,8 +36,6 @@ namespace UE::PixelStreaming
 		virtual TWeakPtr<SViewport> GetTargetViewport() override;
 		virtual void SetTargetWindow(TWeakPtr<SWindow> InTargetWindow) override;
 		virtual TWeakPtr<SWindow> GetTargetWindow() override;
-		virtual void SetTargetScreenSize(TWeakPtr<FIntPoint> InTargetScreenSize) override;
-		virtual TWeakPtr<FIntPoint> GetTargetScreenSize() override;
 		virtual void SetTargetScreenRect(TWeakPtr<FIntRect> InTargetScreenRect) override;
 		virtual TWeakPtr<FIntRect> GetTargetScreenRect() override;
 
diff --git a/Engine/Plugins/Media/PixelStreaming/Source/PixelStreaming/Public/IPixelStreamingModule.h b/Engine/Plugins/Media/PixelStreaming/Source/PixelStreaming/Public/IPixelStreamingModule.h
--- a/Engine/Plugins/Media/PixelStreaming/Source/PixelStreaming/Public/IPixelStreamingModule.h
+++ b/Engine/Plugins/Media/PixelStreaming/Source/PixelStreaming/Public/IPixelStreamingModule.h
@@ -94,14 +94,6 @@ public:
 	 */
 	virtual TArray<FString> GetStreamerIds() = 0;
 
-	/**
-	 * Get a streamer by an ID. 
-	 * // This method has been renamed as the module isn't responsible for a Streamer's lifetime, so renaming to Find as opposed to Get makes this more obvious
-	 * @return A pointer to the interface for a streamer. nullptr if the streamer isn't found
-	 */
-	UE_DEPRECATED(5.2, "GetStreamer() is deprecated. Please use FindStreamer instead.")
-	virtual TSharedPtr<IPixelStreamingStreamer> GetStreamer(const FString& StreamerId) = 0;
-
 	/**
 	 * Find a streamer by an ID.
 	 * @return A pointer to the interface for a streamer. nullptr if the streamer isn't found
@@ -121,32 +113,6 @@ public:
 	 */
 	virtual void DeleteStreamer(TSharedPtr<IPixelStreamingStreamer> ToBeDeleted) = 0;
 
-	/**
-	 * Get the protocol currently used by each peer.
-	 */
-	UE_DEPRECATED(5.2, "GetProtocol() is deprecated. Simply access the required direction with FPixelStreamingInputProtocol::ToStreamerProcol or FPixelStreamingInputProtocol::FromStreamerProcol")
-	virtual const FPixelStreamingInputProtocol GetProtocol() = 0;
-
-	/**
-	 * @brief Register a new message that peers can send and pixel streaming can receive
-	 *
-	 * @param MessageDirection The direction the message will travel. eg Streamer->Player or Player->Streamer
-	 * @param MessageType The human readable identifier (eg "TouchStarted")
-	 * @param Message The object used to define the structure of the message
-	 * @param Handler The handler for this message. This function will be executed whenever the corresponding message is received
-	 */
-	UE_DEPRECATED(5.2, "RegisterMessage(...) is no longer needed. Just add your message to the protocol using FPixelStreamingInputProtocol::Direction.Add(XXX);, and then add the handler to the Streamer's input handler")
-	virtual void RegisterMessage(EPixelStreamingMessageDirection MessageDirection, const FString& MessageType, FPixelStreamingInputMessage Message, const TFunction<void(FString, FMemoryReader)>& Handler) = 0;
-
-	/**
-	 * @brief Find the function to be called whenever the specified message is received.
-	 *
-	 * @param MessageType The human readable identifier for the message
-	 * @return TFunction<void(FMemoryReader)> The function called when this message is received.
-	 */
-	UE_DEPRECATED(5.2, "FindMessageHandler(...) has been moved from the PixelStreaming module to IPixelStreamingInputHandler. This object can be obtained from an (IPixelStreamingStreamer)->GetInputHandler()")
-	virtual TFunction<void(FString, FMemoryReader)> FindMessageHandler(const FString& MessageType) = 0;
-
 	/**
 	 * Sets the target FPS for Externally Consumed video Tracks
 	 * @param InFPS new FPS for the ExternalVideoSource to output at.
diff --git a/Engine/Plugins/Media/PixelStreaming/Source/PixelStreaming/Public/IPixelStreamingSignallingConnectionObserver.h b/Engine/Plugins/Media/PixelStreaming/Source/PixelStreaming/Public/IPixelStreamingSignallingConnectionObserver.h
--- a/Engine/Plugins/Media/PixelStreaming/Source/PixelStreaming/Public/IPixelStreamingSignallingConnectionObserver.h
+++ b/Engine/Plugins/Media/PixelStreaming/Source/PixelStreaming/Public/IPixelStreamingSignallingConnectionObserver.h
@@ -18,8 +18,6 @@ public:
 	// Streamer-only
 	virtual void OnSignallingSessionDescription(FPixelStreamingPlayerId PlayerId, webrtc::SdpType Type, const FString& Sdp) = 0;
 	virtual void OnSignallingRemoteIceCandidate(FPixelStreamingPlayerId PlayerId, const FString& SdpMid, int SdpMLineIndex, const FString& Sdp) = 0;
-	UE_DEPRECATED(5.2, "Please use OnSignallingPlayerConnected(FPixelStreamingPlayerId PlayerId, const FPixelStreamingPlayerConfig& PlayerConfig, bool bSendOffer). Using this method defaults bSendOffer to true! ")
-	virtual void OnSignallingPlayerConnected(FPixelStreamingPlayerId PlayerId, const FPixelStreamingPlayerConfig& PlayerConfig) { OnSignallingPlayerConnected(PlayerId, PlayerConfig, true /* bSendOffer */); };
 	virtual void OnSignallingPlayerConnected(FPixelStreamingPlayerId PlayerId, const FPixelStreamingPlayerConfig& PlayerConfig, bool bSendOffer) = 0;
 	virtual void OnSignallingPlayerDisconnected(FPixelStreamingPlayerId PlayerId) = 0;
 	virtual void OnSignallingSFUPeerDataChannels(FPixelStreamingPlayerId SFUId, FPixelStreamingPlayerId PlayerId, int32 SendStreamId, int32 RecvStreamId) = 0;
diff --git a/Engine/Plugins/Media/PixelStreaming/Source/PixelStreaming/Public/IPixelStreamingStreamer.h b/Engine/Plugins/Media/PixelStreaming/Source/PixelStreaming/Public/IPixelStreamingStreamer.h
--- a/Engine/Plugins/Media/PixelStreaming/Source/PixelStreaming/Public/IPixelStreamingStreamer.h
+++ b/Engine/Plugins/Media/PixelStreaming/Source/PixelStreaming/Public/IPixelStreamingStreamer.h
@@ -79,21 +79,6 @@ public:
 	 */
 	virtual TWeakPtr<SWindow> GetTargetWindow() = 0;
 
-	/**
-	 * @brief Set the target screen size for this streamer. This is used to when the streamer doesn't have a singular target window / viewport
-	 * and as such we just use the manual scale
-	 * @param InTargetScreenSize The target screen size
-	 */
-	UE_DEPRECATED(5.2, "SetTargetScreenSize() is deprecated. Please use SetTargetScreenRect() instead.")
-	virtual void SetTargetScreenSize(TWeakPtr<FIntPoint> InTargetScreenSize) = 0;
-
-	/**
-	 * @brief Get the target screen size for this streamer
-	 * @return The target screen size
-	 */
-	UE_DEPRECATED(5.2, "GetTargetScreenSize() is deprecated. Please use GetTargetScreenRect() instead.")
-	virtual TWeakPtr<FIntPoint> GetTargetScreenSize() = 0;
-
 	/**
 	 * @brief Set the target screen rectangle for this streamer. This is used to when the streamer doesn't have a singular target window / viewport
 	 * and as such we just use the manual scale
diff --git a/Engine/Plugins/Media/PixelStreaming/Source/PixelStreamingEditor/Private/PixelStreamingEditorModule.h b/Engine/Plugins/Media/PixelStreaming/Source/PixelStreamingEditor/Private/PixelStreamingEditorModule.h
--- a/Engine/Plugins/Media/PixelStreaming/Source/PixelStreamingEditor/Private/PixelStreamingEditorModule.h
+++ b/Engine/Plugins/Media/PixelStreaming/Source/PixelStreamingEditor/Private/PixelStreamingEditorModule.h
@@ -24,9 +24,6 @@ public:
 
 	virtual void StartStreaming(UE::EditorPixelStreaming::EStreamTypes InStreamType) override;
 	virtual void StopStreaming() override;
-	/** These two method have been deprecated. We're just keeping them around until they can be removed in 5.4 */
-	virtual void SetStreamType(UE::EditorPixelStreaming::EStreamTypes InStreamType) override{};
-	virtual UE::EditorPixelStreaming::EStreamTypes GetStreamType() override { return UE::EditorPixelStreaming::EStreamTypes::Editor; }
 
 	virtual void StartSignalling() override;
 	virtual void StopSignalling() override;
diff --git a/Engine/Plugins/Media/PixelStreaming/Source/PixelStreamingEditor/Public/IPixelStreamingEditorModule.h b/Engine/Plugins/Media/PixelStreaming/Source/PixelStreamingEditor/Public/IPixelStreamingEditorModule.h
--- a/Engine/Plugins/Media/PixelStreaming/Source/PixelStreamingEditor/Public/IPixelStreamingEditorModule.h
+++ b/Engine/Plugins/Media/PixelStreaming/Source/PixelStreamingEditor/Public/IPixelStreamingEditorModule.h
@@ -45,22 +45,6 @@ public:
 	 */
 	virtual void StopStreaming() = 0;
 
-	/**
-	 * Set the stream type object
-	 *
-	 * @param InStreamType The target stream type
-	 */
-	UE_DEPRECATED(5.2, "SetStreamType() has been deprecated. Stream types are configured by setting the appropriate video input on your streamer")
-	virtual void SetStreamType(UE::EditorPixelStreaming::EStreamTypes InStreamType) = 0;
-
-	/**
-	 * Get the stream type
-	 *
-	 * @return UE::EditorPixelStreaming::EStreamTypes The stream type currently streamed by the editor streamer
-	 */
-	UE_DEPRECATED(5.2, "GetStreamType() has been deprecated. You can now instead find what is being streamed by doing Streamer->GetVideoInput()->ToString()")
-	virtual UE::EditorPixelStreaming::EStreamTypes GetStreamType() = 0;
-
 	/**
 	 * Start the inbuilt C++ signalling server
 	 *
diff --git a/Engine/Plugins/Media/PixelStreaming/Source/PixelStreamingEditor/Public/PixelStreamingEditorUtils.h b/Engine/Plugins/Media/PixelStreaming/Source/PixelStreamingEditor/Public/PixelStreamingEditorUtils.h
--- a/Engine/Plugins/Media/PixelStreaming/Source/PixelStreamingEditor/Public/PixelStreamingEditorUtils.h
+++ b/Engine/Plugins/Media/PixelStreaming/Source/PixelStreamingEditor/Public/PixelStreamingEditorUtils.h
@@ -13,8 +13,7 @@ namespace UE::EditorPixelStreaming
 	enum class PIXELSTREAMINGEDITOR_API EStreamTypes : uint8
 	{
 		LevelEditorViewport = 0,
-		Editor = 1,
-		VCam UE_DEPRECATED(5.2, "EStreamTypes::VCam has been deprecated. Streaming from VCams should be started from the individual actor") = 2
+		Editor = 1
 	};
 
 	FString ToString(EStreamTypes StreamType);
diff --git a/Engine/Plugins/Media/PixelStreaming/Source/PixelStreamingInput/Private/PixelStreamingInputHandler.cpp b/Engine/Plugins/Media/PixelStreaming/Source/PixelStreamingInput/Private/PixelStreamingInputHandler.cpp
--- a/Engine/Plugins/Media/PixelStreaming/Source/PixelStreamingInput/Private/PixelStreamingInputHandler.cpp
+++ b/Engine/Plugins/Media/PixelStreaming/Source/PixelStreamingInput/Private/PixelStreamingInputHandler.cpp
@@ -311,16 +311,6 @@ namespace UE::PixelStreamingInput
 		return TargetWindow;
 	}
 
-	void FPixelStreamingInputHandler::SetTargetScreenSize(TWeakPtr<FIntPoint> InScreenSize)
-	{
-		TargetScreenSize = InScreenSize;
-	}
-
-	TWeakPtr<FIntPoint> FPixelStreamingInputHandler::GetTargetScreenSize()
-	{
-		return TargetScreenSize;
-	}
-
 	void FPixelStreamingInputHandler::SetTargetScreenRect(TWeakPtr<FIntRect> InScreenRect)
 	{
 		TargetScreenRect = InScreenRect;
diff --git a/Engine/Plugins/Media/PixelStreaming/Source/PixelStreamingInput/Private/PixelStreamingInputHandler.h b/Engine/Plugins/Media/PixelStreaming/Source/PixelStreamingInput/Private/PixelStreamingInputHandler.h
--- a/Engine/Plugins/Media/PixelStreaming/Source/PixelStreamingInput/Private/PixelStreamingInputHandler.h
+++ b/Engine/Plugins/Media/PixelStreaming/Source/PixelStreamingInput/Private/PixelStreamingInputHandler.h
@@ -48,9 +48,6 @@ namespace UE::PixelStreamingInput
 		virtual TWeakPtr<SWindow> GetTargetWindow() override;
 		virtual void SetTargetViewport(TWeakPtr<SViewport> InViewport) override;
 		virtual TWeakPtr<SViewport> GetTargetViewport() override;
-		/** These two are deprectated but we keep them around until they can be removed in 5.4 */
-		virtual void SetTargetScreenSize(TWeakPtr<FIntPoint> InScreenSize) override;
-		virtual TWeakPtr<FIntPoint> GetTargetScreenSize() override;
 		virtual void SetTargetScreenRect(TWeakPtr<FIntRect> InScreenRect) override;
 		virtual TWeakPtr<FIntRect> GetTargetScreenRect() override;
 		virtual bool IsFakingTouchEvents() const override { return bFakingTouchEvents; }
diff --git a/Engine/Plugins/Media/PixelStreaming/Source/PixelStreamingInput/Public/IPixelStreamingInputHandler.h b/Engine/Plugins/Media/PixelStreaming/Source/PixelStreamingInput/Public/IPixelStreamingInputHandler.h
--- a/Engine/Plugins/Media/PixelStreaming/Source/PixelStreamingInput/Public/IPixelStreamingInputHandler.h
+++ b/Engine/Plugins/Media/PixelStreaming/Source/PixelStreamingInput/Public/IPixelStreamingInputHandler.h
@@ -49,21 +49,6 @@ public:
 	 */
 	virtual TWeakPtr<SWindow> GetTargetWindow() = 0;
 
-	/**
-	 * @brief Set the target screen size for this streamer. This is used to when the streamer doesn't have a singular target window / viewport
-	 * and as such we just use the manual scale
-	 * @param InTargetWindow The target screen size
-	 */
-	UE_DEPRECATED(5.2, "SetTargetScreenSize() is deprecated. Please use SetTargetScreenRect() instead.")
-	virtual void SetTargetScreenSize(TWeakPtr<FIntPoint> InTargetScreenSize) = 0;
-
-	/**
-	 * @brief Get the target screen size for this streamer
-	 * @return The target screen size
-	 */
-	UE_DEPRECATED(5.2, "GetTargetScreenSize() is deprecated. Please use GetTargetScreenRect() instead.")
-	virtual TWeakPtr<FIntPoint> GetTargetScreenSize() = 0;
-
 	/**
 	 * @brief Set the target screen rectangle for this streamer. This is used to when the streamer doesn't have a singular target window / viewport
 	 * and as such we just use the manual scale
