Fix: Editor crash on shutdown when trying to restore settings

diff --git a/Engine/Plugins/Media/PixelStreaming/Source/PixelStreamingEditor/Private/PixelStreamingEditorModule.cpp b/Engine/Plugins/Media/PixelStreaming/Source/PixelStreamingEditor/Private/PixelStreamingEditorModule.cpp
--- a/Engine/Plugins/Media/PixelStreaming/Source/PixelStreamingEditor/Private/PixelStreamingEditorModule.cpp
+++ b/Engine/Plugins/Media/PixelStreaming/Source/PixelStreamingEditor/Private/PixelStreamingEditorModule.cpp
@@ -86,8 +86,24 @@ void FPixelStreamingEditorModule::InitEditorStreaming(IPixelStreamingModule& Mod
 	EditorStreamer = Module.CreateStreamer(EditorStreamerID);
 
 	// Bind to start/stop streaming so we disable/restore relevant editor settings
-	EditorStreamer->OnStreamingStarted().AddLambda([this](IPixelStreamingStreamer* Streamer){ DisableCPUThrottlingSetting(); });
-	EditorStreamer->OnStreamingStopped().AddLambda([this](IPixelStreamingStreamer* Streamer){ RestoreCPUThrottlingSetting(); });
+	EditorStreamer->OnStreamingStarted().AddLambda([this](IPixelStreamingStreamer* Streamer)
+	{ 
+		DisableCPUThrottlingSetting(); 
+	});
+	EditorStreamer->OnStreamingStopped().AddLambda([this](IPixelStreamingStreamer* Streamer)
+	{ 
+		if(!IsEngineExitRequested())
+		{
+			RestoreCPUThrottlingSetting(); 
+		}
+	});
+
+	/**
+	 * Called before the engine exits. Separate from OnPreExit as OnEnginePreExit occurs before shutting down any core modules.
+	*/
+	FCoreDelegates::OnEnginePreExit.AddLambda([this]() {
+		RestoreCPUThrottlingSetting();
+	});
 
 	// Give the editor streamer the default url if the user hasn't specified one when launching the editor
 	if (EditorStreamer->GetSignallingServerURL().IsEmpty())
