Fix: Encoders weren't being released after peer disconnect / stream disconnect

diff --git a/Engine/Plugins/Media/PixelStreaming/Source/PixelStreaming/Private/VideoEncoderFactorySingleLayer.cpp b/Engine/Plugins/Media/PixelStreaming/Source/PixelStreaming/Private/VideoEncoderFactorySingleLayer.cpp
--- a/Engine/Plugins/Media/PixelStreaming/Source/PixelStreaming/Private/VideoEncoderFactorySingleLayer.cpp
+++ b/Engine/Plugins/Media/PixelStreaming/Source/PixelStreaming/Private/VideoEncoderFactorySingleLayer.cpp
@@ -228,13 +228,15 @@ namespace UE::PixelStreaming
 		// Lock during deleting an encoder
 		FScopeLock Lock(&ActiveEncodersGuard);
 		ActiveEncoders.Remove(Encoder);
+
+		FreeUnusedEncoders();
 	}
 
 	void FVideoEncoderFactorySingleLayer::FreeUnusedEncoders()
 	{
 		// first clear unused encoders.
 		TArray<uint32> DeleteStreamIds;
-		HardwareEncoders.Apply([&DeleteStreamIds](uint32 EncoderStreamId, TSharedPtr<FVideoEncoderHardware> EncoderPtr) {
+		HardwareEncoders.Apply([&DeleteStreamIds](uint32 EncoderStreamId, TSharedPtr<FVideoEncoderHardware>& EncoderPtr) {
 			if (EncoderPtr.IsUnique())
 			{
 				DeleteStreamIds.Add(EncoderStreamId);
diff --git a/Engine/Plugins/Media/PixelStreaming/Source/PixelStreaming/Private/VideoEncoderSingleLayerHardware.cpp b/Engine/Plugins/Media/PixelStreaming/Source/PixelStreaming/Private/VideoEncoderSingleLayerHardware.cpp
--- a/Engine/Plugins/Media/PixelStreaming/Source/PixelStreaming/Private/VideoEncoderSingleLayerHardware.cpp
+++ b/Engine/Plugins/Media/PixelStreaming/Source/PixelStreaming/Private/VideoEncoderSingleLayerHardware.cpp
@@ -137,9 +137,9 @@ namespace UE::PixelStreaming
 		const FPixelCaptureOutputFrameRHI& RHILayer = StaticCast<const FPixelCaptureOutputFrameRHI&>(AdaptedLayer);
 		// TODO-TE
 #if WEBRTC_5414
-		rtc::scoped_refptr<FFrameBufferRHI> RHIBuffer = rtc::make_ref_counted<FFrameBufferRHI>(MakeShared<FVideoResourceRHI>(HardwareEncoder->GetDevice().ToSharedRef(), FVideoResourceRHI::FRawData{ RHILayer.GetFrameTexture(), nullptr, 0 }));
+		rtc::scoped_refptr<FFrameBufferRHI> RHIBuffer = rtc::make_ref_counted<FFrameBufferRHI>(MakeShared<FVideoResourceRHI>(HardwareEncoder.Pin()->GetDevice().ToSharedRef(), FVideoResourceRHI::FRawData{ RHILayer.GetFrameTexture(), nullptr, 0 }));
 #else
-		rtc::scoped_refptr<FFrameBufferRHI> RHIBuffer = new rtc::RefCountedObject<FFrameBufferRHI>(MakeShared<FVideoResourceRHI>(HardwareEncoder->GetDevice().ToSharedRef(), FVideoResourceRHI::FRawData{ RHILayer.GetFrameTexture(), nullptr, 0 }));
+		rtc::scoped_refptr<FFrameBufferRHI> RHIBuffer = new rtc::RefCountedObject<FFrameBufferRHI>(MakeShared<FVideoResourceRHI>(HardwareEncoder.Pin()->GetDevice().ToSharedRef(), FVideoResourceRHI::FRawData{ RHILayer.GetFrameTexture(), nullptr, 0 }));
 #endif
 		NewFrame.set_video_frame_buffer(RHIBuffer);
 		return NewFrame;
@@ -154,13 +154,13 @@ namespace UE::PixelStreaming
 			case EPixelStreamingCodec::H264:
 			{
 				FVideoEncoderConfigH264& VideoConfig = *StaticCast<FVideoEncoderConfigH264*>(InitialVideoConfig.Get());
-				HardwareEncoder = Factory.GetOrCreateHardwareEncoder(StreamId, VideoConfig).Pin();
+				HardwareEncoder = Factory.GetOrCreateHardwareEncoder(StreamId, VideoConfig);
 				break;
 			}
 			case EPixelStreamingCodec::AV1:
 			{
 				FVideoEncoderConfigAV1& VideoConfig = *StaticCast<FVideoEncoderConfigAV1*>(InitialVideoConfig.Get());
-				HardwareEncoder = Factory.GetOrCreateHardwareEncoder(StreamId, VideoConfig).Pin();
+				HardwareEncoder = Factory.GetOrCreateHardwareEncoder(StreamId, VideoConfig);
 				break;
 			}
 		}
@@ -184,7 +184,7 @@ namespace UE::PixelStreaming
 
 		// We late init here so we can pull the stream ID off the incoming frames and pull the correct encoder for the stream
 		// earlier locations do not have this information.
-		if (!HardwareEncoder)
+		if (!HardwareEncoder.IsValid())
 		{
 			InitialVideoConfig->Width = FrameWidth;
 			InitialVideoConfig->Height = FrameHeight;
@@ -194,7 +194,7 @@ namespace UE::PixelStreaming
 		// Update the encoding config using the incoming frame resolution
 		UpdateConfig(FrameWidth, FrameHeight);
 
-		if (HardwareEncoder)
+		if (TSharedPtr<FVideoEncoderHardware> const& PinnedHardwareEncoder = HardwareEncoder.Pin())
 		{
 			TRACE_CPUPROFILER_EVENT_SCOPE_ON_CHANNEL_STR("PixelStreaming Hardware Encoding", PixelStreamingChannel);
 
@@ -225,31 +225,31 @@ namespace UE::PixelStreaming
 			{
 				// If we have a wrapper format we need to also provide a descriptor of the raw format the data
 				// is transported in
-				FVideoDescriptor RawDescriptor = FVideoResourceRHI::GetDescriptorFrom(HardwareEncoder->GetDevice().ToSharedRef(), RHILayer.GetFrameTexture());
+				FVideoDescriptor RawDescriptor = FVideoResourceRHI::GetDescriptorFrom(PinnedHardwareEncoder->GetDevice().ToSharedRef(), RHILayer.GetFrameTexture());
 
 				// If we have a wrapper format for our resource, make sure the video descriptor specifies it
 				RHIBuffer = new rtc::RefCountedObject<FFrameBufferRHI>(MakeShared<FVideoResourceRHI>(
-					HardwareEncoder->GetDevice().ToSharedRef(),
+					PinnedHardwareEncoder->GetDevice().ToSharedRef(),
 					FVideoResourceRHI::FRawData{ RHILayer.GetFrameTexture(), nullptr, 0 },
 					FVideoDescriptor(static_cast<EVideoFormat>(WrapperData.Format), WrapperData.Width, WrapperData.Height, RawDescriptor)));
 			}
 			else
 			{
 				RHIBuffer = new rtc::RefCountedObject<FFrameBufferRHI>(MakeShared<FVideoResourceRHI>(
-					HardwareEncoder->GetDevice().ToSharedRef(),
+					PinnedHardwareEncoder->GetDevice().ToSharedRef(),
 					FVideoResourceRHI::FRawData{ RHILayer.GetFrameTexture(), nullptr, 0 }));
 			}
 
 			const bool bKeyframe = (frame_types && (*frame_types)[0] == webrtc::VideoFrameType::kVideoFrameKey) || Factory.ShouldForceKeyframe();
 
-			HardwareEncoder->SendFrame(RHIBuffer->GetVideoResource(), frame.timestamp_us(), bKeyframe);
+			PinnedHardwareEncoder->SendFrame(RHIBuffer->GetVideoResource(), frame.timestamp_us(), bKeyframe);
 
 			Factory.UnforceKeyFrame();
 
 			UpdateFrameMetadataPostEncode(*AdaptedLayer);
 
 			FVideoPacket Packet;
-			while (HardwareEncoder->ReceivePacket(Packet))
+			while (PinnedHardwareEncoder->ReceivePacket(Packet))
 			{
 				// Note: This works fine for 1:1 synchronous encoding, but if we ever need to relate frames to packets for async or gop this will need thought.
 				webrtc::EncodedImage Image;
@@ -337,17 +337,17 @@ namespace UE::PixelStreaming
 		InitialVideoConfig->Width = width;
 		InitialVideoConfig->Height = height;
 
-		if (HardwareEncoder)
+		if (TSharedPtr<FVideoEncoderHardware> const& PinnedHardwareEncoder = HardwareEncoder.Pin())
 		{
-			FVideoEncoderConfig VideoConfigMinimal = HardwareEncoder->GetMinimalConfig();
+			FVideoEncoderConfig VideoConfigMinimal = PinnedHardwareEncoder->GetMinimalConfig();
 			FVideoEncoderConfig* VideoConfig = &VideoConfigMinimal;
 
 			switch (Codec)
 			{
 				case EPixelStreamingCodec::H264:
-					if (HardwareEncoder->GetInstance()->Has<FVideoEncoderConfigH264>())
+					if (PinnedHardwareEncoder->GetInstance()->Has<FVideoEncoderConfigH264>())
 					{
-						FVideoEncoderConfigH264& VideoConfigH264 = HardwareEncoder->GetInstance()->Edit<FVideoEncoderConfigH264>();
+						FVideoEncoderConfigH264& VideoConfigH264 = PinnedHardwareEncoder->GetInstance()->Edit<FVideoEncoderConfigH264>();
 						VideoConfig = &VideoConfigH264;
 
 						VideoConfigH264.Profile = UE::PixelStreaming::Settings::GetH264Profile();
@@ -355,9 +355,9 @@ namespace UE::PixelStreaming
 
 					break;
 				case EPixelStreamingCodec::AV1:
-					if (HardwareEncoder->GetInstance()->Has<FVideoEncoderConfigAV1>())
+					if (PinnedHardwareEncoder->GetInstance()->Has<FVideoEncoderConfigAV1>())
 					{
-						FVideoEncoderConfigAV1& VideoConfigAV1 = HardwareEncoder->GetInstance()->Edit<FVideoEncoderConfigAV1>();
+						FVideoEncoderConfigAV1& VideoConfigAV1 = PinnedHardwareEncoder->GetInstance()->Edit<FVideoEncoderConfigAV1>();
 						VideoConfig = &VideoConfigAV1;
 					}
 
@@ -396,7 +396,7 @@ namespace UE::PixelStreaming
 			VideoConfig->Width = width;
 			VideoConfig->Height = height;
 
-			HardwareEncoder->SetMinimalConfig(*VideoConfig);
+			PinnedHardwareEncoder->SetMinimalConfig(*VideoConfig);
 		}
 	}
 
diff --git a/Engine/Plugins/Media/PixelStreaming/Source/PixelStreaming/Private/VideoEncoderSingleLayerHardware.h b/Engine/Plugins/Media/PixelStreaming/Source/PixelStreaming/Private/VideoEncoderSingleLayerHardware.h
--- a/Engine/Plugins/Media/PixelStreaming/Source/PixelStreaming/Private/VideoEncoderSingleLayerHardware.h
+++ b/Engine/Plugins/Media/PixelStreaming/Source/PixelStreaming/Private/VideoEncoderSingleLayerHardware.h
@@ -46,12 +46,12 @@ namespace UE::PixelStreaming
 		void UpdateFrameMetadataPrePacketization(IPixelCaptureOutputFrame& Frame);
 		void UpdateFrameMetadataPostPacketization(IPixelCaptureOutputFrame& Frame);
 		webrtc::VideoFrame WrapAdaptedFrame(const webrtc::VideoFrame& ExistingFrame, const IPixelCaptureOutputFrame& AdaptedLayer);
-
+		
 		FVideoEncoderFactorySingleLayer& Factory;
 		EPixelStreamingCodec Codec;
 
 		TUniquePtr<FVideoEncoderConfig> InitialVideoConfig;
-		TSharedPtr<FVideoEncoderHardware> HardwareEncoder;
+		TWeakPtr<FVideoEncoderHardware> HardwareEncoder;
 
 		// We store this so we can restore back to it if the user decides to use then stop using the PixelStreaming.Encoder.TargetBitrate CVar.
 		int32 WebRtcProposedTargetBitrate = 5000000;
