Fix: Re-starting editor streaming (start streaming -> stop streaming -> start streaming) would hard crash the editor

diff --git a/Engine/Plugins/Media/PixelStreaming/Source/PixelStreamingEditor/Private/PixelStreamingVideoInputBackBufferComposited.cpp b/Engine/Plugins/Media/PixelStreaming/Source/PixelStreamingEditor/Private/PixelStreamingVideoInputBackBufferComposited.cpp
--- a/Engine/Plugins/Media/PixelStreaming/Source/PixelStreamingEditor/Private/PixelStreamingVideoInputBackBufferComposited.cpp
+++ b/Engine/Plugins/Media/PixelStreaming/Source/PixelStreamingEditor/Private/PixelStreamingVideoInputBackBufferComposited.cpp
@@ -36,64 +36,71 @@ TSharedPtr<FPixelStreamingVideoInputBackBufferComposited> FPixelStreamingVideoIn
 	AsyncTask(ENamedThreads::GameThread, [WeakInput]() {
 		if (TSharedPtr<FPixelStreamingVideoInputBackBufferComposited> Input = WeakInput.Pin())
 		{
-			FSlateRenderer* Renderer = FSlateApplication::Get().GetRenderer();
-			Input->DelegateHandle = Renderer->OnBackBufferReadyToPresent().AddSP(Input.ToSharedRef(), &FPixelStreamingVideoInputBackBufferComposited::OnBackBufferReady);
+			FSlateApplication& SlateApplication = FSlateApplication::Get();
+			Input->OnBackBufferReadyToPresentHandle = SlateApplication.GetRenderer()->OnBackBufferReadyToPresent().AddSP(Input.ToSharedRef(), &FPixelStreamingVideoInputBackBufferComposited::OnBackBufferReady);
+			Input->OnPreTickHandle = SlateApplication.OnPreTick().AddSP(Input.ToSharedRef(), &FPixelStreamingVideoInputBackBufferComposited::OnPreTick);
 		} });
 
 	return NewInput;
 }
 
 FPixelStreamingVideoInputBackBufferComposited::FPixelStreamingVideoInputBackBufferComposited()
+	: SharedFrameRect(MakeShared<FIntRect>())
 {
-	SharedFrameRect = MakeShared<FIntRect>();
-	AsyncTask(ENamedThreads::GameThread, [this]() { FSlateApplication::Get().OnPreTick().AddLambda([this](float DeltaTime) {
-			FScopeLock Lock(&TopLevelWindowsCriticalSection);
-			TArray<TSharedRef<SWindow>> TopLevelSlateWindows;
-			FSlateApplication::Get().GetAllVisibleWindowsOrdered(TopLevelSlateWindows);
-
-			// We store all the necessary window information in structs. This prevents window information from updating
-			// underneath us while we composite and also means we aren't holding on to any shared refs between compositions
-			TArray<FTexturedWindow> TempWindows = TopLevelWindows;
-			TopLevelWindows.Empty();
-
-			for (uint8 i = 0; i < TopLevelSlateWindows.Num(); i++)
-			{
-				TSharedRef<SWindow> CurrentWindow = TopLevelSlateWindows[i];
-				FVector2D WindowExtent = CurrentWindow->GetPositionInScreen() - CurrentWindow->GetSizeInScreen();
-
-				// No need to keep track of "invalid" windows
-				if (CurrentWindow->GetOpacity() == 0.f
-					|| CurrentWindow->GetSizeInScreen() == FVector2D(0, 0)
-					|| CurrentWindow->GetPositionInScreen().X > 16384
-					|| CurrentWindow->GetPositionInScreen().X < -16384
-					|| CurrentWindow->GetPositionInScreen().Y > 16384
-					|| CurrentWindow->GetPositionInScreen().Y < -16384)
-				{
-					continue;
-				}
-
-				TopLevelWindows.Add(FTexturedWindow(CurrentWindow->GetPositionInScreen(), CurrentWindow->GetSizeInScreen(), CurrentWindow->GetOpacity(), CurrentWindow->GetType(), &CurrentWindow.Get()));
-
-				// HACK (william.belcher): This following section is a bit of a nasty work-around to the fact that when a modal is displayed, previous windows (eg the editor) won't
-				// trigger the OnBackBufferReady delegate. This means that we need to check if we have previously seen this top level window, and if we have we need to copy across
-				// its staging texture just in case we won't see it again until after the modal is closed
-				int32 Index = TempWindows.IndexOfByPredicate([&CurrentWindow](FTexturedWindow Window) {
-					return Window.GetOwningWindow() == &CurrentWindow.Get();
-				});
-
-				if (Index != INDEX_NONE)
-				{
-					// This isn't the first time we've seen this window, so copy across its staging texture
-					TopLevelWindows[TopLevelWindows.Num() - 1].SetTexture(TempWindows[Index].GetTexture());
-				}
-			} }); });
 }
 
 FPixelStreamingVideoInputBackBufferComposited::~FPixelStreamingVideoInputBackBufferComposited()
 {
 	if (!IsEngineExitRequested())
 	{
-		AsyncTask(ENamedThreads::GameThread, [HandleCopy = DelegateHandle]() { FSlateApplication::Get().GetRenderer()->OnBackBufferReadyToPresent().Remove(HandleCopy); });
+		AsyncTask(ENamedThreads::GameThread, [OnBackBufferReadyToPresentCopy = OnBackBufferReadyToPresentHandle, OnPreTickCopy = OnPreTickHandle]() { 
+			FSlateApplication& SlateApplication = FSlateApplication::Get();
+			SlateApplication.GetRenderer()->OnBackBufferReadyToPresent().Remove(OnBackBufferReadyToPresentCopy);
+			SlateApplication.OnPreTick().Remove(OnPreTickCopy);
+		});
+	}
+}
+
+void FPixelStreamingVideoInputBackBufferComposited::OnPreTick(float DeltaTime)
+{
+	FScopeLock Lock(&TopLevelWindowsCriticalSection);
+	TArray<TSharedRef<SWindow>> TopLevelSlateWindows;
+	FSlateApplication::Get().GetAllVisibleWindowsOrdered(TopLevelSlateWindows);
+
+	// We store all the necessary window information in structs. This prevents window information from updating
+	// underneath us while we composite and also means we aren't holding on to any shared refs between compositions
+	TArray<FTexturedWindow> TempWindows = TopLevelWindows;
+	TopLevelWindows.Empty();
+
+	for (TSharedRef<SWindow> CurrentWindow : TopLevelSlateWindows)
+	{
+		FVector2D WindowExtent = CurrentWindow->GetPositionInScreen() - CurrentWindow->GetSizeInScreen();
+
+		// No need to keep track of "invalid" windows
+		if (CurrentWindow->GetOpacity() == 0.f ||
+			CurrentWindow->GetSizeInScreen() == FVector2D(0, 0) ||
+			CurrentWindow->GetPositionInScreen().X > 16384 ||
+			CurrentWindow->GetPositionInScreen().X < -16384 ||
+			CurrentWindow->GetPositionInScreen().Y > 16384 ||
+			CurrentWindow->GetPositionInScreen().Y < -16384)
+		{
+			continue;
+		}
+
+		TopLevelWindows.Add(FTexturedWindow(CurrentWindow->GetPositionInScreen(), CurrentWindow->GetSizeInScreen(), CurrentWindow->GetOpacity(), CurrentWindow->GetType(), &CurrentWindow.Get()));
+
+		// HACK (william.belcher): This following section is a bit of a nasty work-around to the fact that when a modal is displayed, previous windows (eg the editor) won't
+		// trigger the OnBackBufferReady delegate. This means that we need to check if we have previously seen this top level window, and if we have we need to copy across
+		// its staging texture just in case we won't see it again until after the modal is closed
+		int32 Index = TempWindows.IndexOfByPredicate([&CurrentWindow](FTexturedWindow Window) {
+			return Window.GetOwningWindow() == &CurrentWindow.Get();
+		});
+
+		if (Index != INDEX_NONE)
+		{
+			// This isn't the first time we've seen this window, so copy across its staging texture
+			TopLevelWindows[TopLevelWindows.Num() - 1].SetTexture(TempWindows[Index].GetTexture());
+		}
 	}
 }
 
diff --git a/Engine/Plugins/Media/PixelStreaming/Source/PixelStreamingEditor/Public/PixelStreamingVideoInputBackBufferComposited.h b/Engine/Plugins/Media/PixelStreaming/Source/PixelStreamingEditor/Public/PixelStreamingVideoInputBackBufferComposited.h
--- a/Engine/Plugins/Media/PixelStreaming/Source/PixelStreamingEditor/Public/PixelStreamingVideoInputBackBufferComposited.h
+++ b/Engine/Plugins/Media/PixelStreaming/Source/PixelStreamingEditor/Public/PixelStreamingVideoInputBackBufferComposited.h
@@ -55,8 +55,10 @@ private:
     void CompositeWindows();
 
     void OnBackBufferReady(SWindow &SlateWindow, const FTextureRHIRef &FrameBuffer);
+    void OnPreTick(float DeltaTime);
 
-    FDelegateHandle DelegateHandle;
+    FDelegateHandle OnBackBufferReadyToPresentHandle;
+    FDelegateHandle OnPreTickHandle;
 
     TArray<FTexturedWindow> TopLevelWindows;
 
