Fix: Artefacting around editor streaming windows

diff --git a/Engine/Plugins/Media/PixelStreaming/Source/PixelStreamingEditor/Private/PixelStreamingVideoInputBackBufferComposited.cpp b/Engine/Plugins/Media/PixelStreaming/Source/PixelStreamingEditor/Private/PixelStreamingVideoInputBackBufferComposited.cpp
--- a/Engine/Plugins/Media/PixelStreaming/Source/PixelStreamingEditor/Private/PixelStreamingVideoInputBackBufferComposited.cpp
+++ b/Engine/Plugins/Media/PixelStreaming/Source/PixelStreamingEditor/Private/PixelStreamingVideoInputBackBufferComposited.cpp
@@ -253,7 +253,12 @@ void FPixelStreamingVideoInputBackBufferComposited::CompositeWindows()
 			const FSceneView& View = *ViewFamily.Views[0];
 
 			TShaderMapRef<FModifyAlphaSwizzleRgbaPS> PixelShader(GlobalShaderMap, PermutationVector);
-			FModifyAlphaSwizzleRgbaPS::FParameters* PixelShaderParameters = PixelShader->AllocateAndSetParameters(GraphBuilder, InputTexture, CompositedTexture);
+			FModifyAlphaSwizzleRgbaPS::FParameters* PixelShaderParameters = GraphBuilder.AllocParameters<FModifyAlphaSwizzleRgbaPS::FParameters>();
+			PixelShaderParameters->InputTexture = InputTexture;
+			PixelShaderParameters->InputSampler = TStaticSamplerState<SF_Point>::GetRHI();
+			PixelShaderParameters->RenderTargets[0] = FRenderTargetBinding{ CompositedTexture, ERenderTargetLoadAction::ELoad };
+			
+			
 			// Add screen pass to convert whatever format the editor produces to BGRA8
 			AddDrawScreenPass(GraphBuilder, RDG_EVENT_NAME("VideoInputBackBufferCompositedSwizzle"), View, OutputViewport, InputViewport, VertexShader, PixelShader, PixelShaderParameters);
 		}
