Fix: Pixel Streaming blueprint components running unnecessarily even when unused in scene

diff --git a/Engine/Plugins/Media/PixelStreaming/Source/PixelStreamingBlueprint/Internal/PixelStreamingStreamerVideoInputBackBuffer.h b/Engine/Plugins/Media/PixelStreaming/Source/PixelStreamingBlueprint/Internal/PixelStreamingStreamerVideoInputBackBuffer.h
--- a/Engine/Plugins/Media/PixelStreaming/Source/PixelStreamingBlueprint/Internal/PixelStreamingStreamerVideoInputBackBuffer.h
+++ b/Engine/Plugins/Media/PixelStreaming/Source/PixelStreamingBlueprint/Internal/PixelStreamingStreamerVideoInputBackBuffer.h
@@ -9,4 +9,7 @@ UCLASS(NotBlueprintType, NotBlueprintable, Category = "PixelStreaming", META = (
 class PIXELSTREAMINGBLUEPRINT_API UPixelStreamingStreamerVideoInputBackBuffer : public UPixelStreamingStreamerVideoInput
 {
 	GENERATED_UCLASS_BODY()
+
+public:
+	virtual TSharedPtr<FPixelStreamingVideoInput> GetVideoInput() override;
 };
diff --git a/Engine/Plugins/Media/PixelStreaming/Source/PixelStreamingBlueprint/Internal/PixelStreamingStreamerVideoInputMediaCapture.h b/Engine/Plugins/Media/PixelStreaming/Source/PixelStreamingBlueprint/Internal/PixelStreamingStreamerVideoInputMediaCapture.h
--- a/Engine/Plugins/Media/PixelStreaming/Source/PixelStreamingBlueprint/Internal/PixelStreamingStreamerVideoInputMediaCapture.h
+++ b/Engine/Plugins/Media/PixelStreaming/Source/PixelStreamingBlueprint/Internal/PixelStreamingStreamerVideoInputMediaCapture.h
@@ -10,4 +10,7 @@ UCLASS(NotBlueprintType, NotBlueprintable, Category = "PixelStreaming", META = (
 class PIXELSTREAMINGBLUEPRINT_API UPixelStreamingStreamerVideoInputMediaCapture : public UPixelStreamingStreamerVideoInput
 {
 	GENERATED_UCLASS_BODY()
+
+public:
+	virtual TSharedPtr<FPixelStreamingVideoInput> GetVideoInput() override;
 };
diff --git a/Engine/Plugins/Media/PixelStreaming/Source/PixelStreamingBlueprint/Private/PixelStreamingStreamerVideoInputBackBuffer.cpp b/Engine/Plugins/Media/PixelStreaming/Source/PixelStreamingBlueprint/Private/PixelStreamingStreamerVideoInputBackBuffer.cpp
--- a/Engine/Plugins/Media/PixelStreaming/Source/PixelStreamingBlueprint/Private/PixelStreamingStreamerVideoInputBackBuffer.cpp
+++ b/Engine/Plugins/Media/PixelStreaming/Source/PixelStreamingBlueprint/Private/PixelStreamingStreamerVideoInputBackBuffer.cpp
@@ -7,15 +7,24 @@
 UPixelStreamingStreamerVideoInputBackBuffer::UPixelStreamingStreamerVideoInputBackBuffer(const FObjectInitializer& ObjectInitializer)
 	: Super(ObjectInitializer)
 {
-	// detect if we're in PIE mode or not
-	bool IsGame = false;
-	FParse::Bool(FCommandLine::Get(), TEXT("game"), IsGame);
-	if (GIsEditor && !IsGame)
-	{
-		VideoInput = FPixelStreamingVideoInputPIEViewport::Create();
-	}
-	else
+}
+
+TSharedPtr<FPixelStreamingVideoInput> UPixelStreamingStreamerVideoInputBackBuffer::GetVideoInput()
+{
+	if (!VideoInput)
 	{
-		VideoInput = FPixelStreamingVideoInputBackBuffer::Create();
+		// detect if we're in PIE mode or not
+		bool IsGame = false;
+		FParse::Bool(FCommandLine::Get(), TEXT("game"), IsGame);
+		if (GIsEditor && !IsGame)
+		{
+			VideoInput = FPixelStreamingVideoInputPIEViewport::Create();
+		}
+		else
+		{
+			VideoInput = FPixelStreamingVideoInputBackBuffer::Create();
+		}
 	}
+	return VideoInput;
 }
+
diff --git a/Engine/Plugins/Media/PixelStreaming/Source/PixelStreamingBlueprint/Private/PixelStreamingStreamerVideoInputMediaCapture.cpp b/Engine/Plugins/Media/PixelStreaming/Source/PixelStreamingBlueprint/Private/PixelStreamingStreamerVideoInputMediaCapture.cpp
--- a/Engine/Plugins/Media/PixelStreaming/Source/PixelStreamingBlueprint/Private/PixelStreamingStreamerVideoInputMediaCapture.cpp
+++ b/Engine/Plugins/Media/PixelStreaming/Source/PixelStreamingBlueprint/Private/PixelStreamingStreamerVideoInputMediaCapture.cpp
@@ -7,15 +7,24 @@
 UPixelStreamingStreamerVideoInputMediaCapture::UPixelStreamingStreamerVideoInputMediaCapture(const FObjectInitializer& ObjectInitializer)
 	: Super(ObjectInitializer)
 {
-	// detect if we're in PIE mode or not
-	bool IsGame = false;
-	FParse::Bool(FCommandLine::Get(), TEXT("game"), IsGame);
-	if (GIsEditor && !IsGame)
-	{
-		VideoInput = FPixelStreamingVideoInputPIEViewport::Create();
-	}
-	else
+}
+
+TSharedPtr<FPixelStreamingVideoInput> UPixelStreamingStreamerVideoInputMediaCapture::GetVideoInput()
+{
+	if (!VideoInput)
 	{
-		VideoInput = FPixelStreamingVideoInputMediaCapture::CreateActiveViewportCapture();
+		// detect if we're in PIE mode or not
+		bool IsGame = false;
+		FParse::Bool(FCommandLine::Get(), TEXT("game"), IsGame);
+		if (GIsEditor && !IsGame)
+		{
+			VideoInput = FPixelStreamingVideoInputPIEViewport::Create();
+		}
+		else
+		{
+			VideoInput = FPixelStreamingVideoInputMediaCapture::CreateActiveViewportCapture();
+		}
 	}
+	return VideoInput;
 }
+
