Add: Enable FlexFec support to PixelStreaming

diff --git a/Engine/Plugins/Media/PixelStreaming/Source/PixelStreaming/Private/PixelStreamingPeerConnection.cpp b/Engine/Plugins/Media/PixelStreaming/Source/PixelStreaming/Private/PixelStreamingPeerConnection.cpp
--- a/Engine/Plugins/Media/PixelStreaming/Source/PixelStreaming/Private/PixelStreamingPeerConnection.cpp
+++ b/Engine/Plugins/Media/PixelStreaming/Source/PixelStreaming/Private/PixelStreamingPeerConnection.cpp
@@ -913,6 +913,11 @@ void InitializeFieldTrials()
 		FieldTrials += TEXT("WebRTC-FrameDropper/Disabled/");
 	}
 
+	if (Settings::CVarPixelStreamingWebRTCEnableFlexFec.GetValueOnAnyThread())
+	{
+		FieldTrials += TEXT("WebRTC-FlexFEC-03-Advertised/Enabled/WebRTC-FlexFEC-03/Enabled/");
+	}
+
 	// Parse "WebRTC-Video-Pacing/" field trial
 	{
 		float OutPacingFactor = -1.0f;
diff --git a/Engine/Plugins/Media/PixelStreaming/Source/PixelStreaming/Private/Settings.cpp b/Engine/Plugins/Media/PixelStreaming/Source/PixelStreaming/Private/Settings.cpp
--- a/Engine/Plugins/Media/PixelStreaming/Source/PixelStreaming/Private/Settings.cpp
+++ b/Engine/Plugins/Media/PixelStreaming/Source/PixelStreaming/Private/Settings.cpp
@@ -179,6 +179,12 @@ namespace UE::PixelStreaming::Settings
 		TEXT("Whether put audio and video in the same stream (which will make WebRTC try to sync them)."),
 		ECVF_Default);
 
+	TAutoConsoleVariable<bool> CVarPixelStreamingWebRTCEnableFlexFec(
+		TEXT("PixelStreaming.WebRTC.EnableFlexFec"),
+		false,
+		TEXT("Signals support for Flexible Forward Error Correction to WebRTC."),
+		ECVF_Default);
+
 	TAutoConsoleVariable<bool> CVarPixelStreamingWebRTCDisableStats(
 		TEXT("PixelStreaming.WebRTC.DisableStats"),
 		false,
@@ -801,6 +807,7 @@ namespace UE::PixelStreaming::Settings
 		CommandLineParseOption(TEXT("PixelStreamingWebRTCDisableFrameDropper"), CVarPixelStreamingWebRTCDisableFrameDropper);
 		CommandLineParseOption(TEXT("PixelStreamingSendPlayerIdAsInteger"), CVarSendPlayerIdAsInteger);
 		CommandLineParseOption(TEXT("PixelStreamingWebRTCUseLegacyAudioDevice"), CVarPixelStreamingWebRTCUseLegacyAudioDevice);
+		CommandLineParseOption(TEXT("PixelStreamingWebRTCEnableFlexFec"), CVarPixelStreamingWebRTCEnableFlexFec);
 		CommandLineParseOption(TEXT("PixelStreamingDisableLatencyTester"), CVarPixelStreamingDisableLatencyTester);
 		CommandLineParseOption(TEXT("PixelStreamingVPXUseCompute"), CVarPixelStreamingVPXUseCompute);
 		CommandLineParseOption(TEXT("PixelStreamingExperimentalAudioInput"), CVarPixelStreamingExperimentalAudioInput);
diff --git a/Engine/Plugins/Media/PixelStreaming/Source/PixelStreaming/Private/Settings.h b/Engine/Plugins/Media/PixelStreaming/Source/PixelStreaming/Private/Settings.h
--- a/Engine/Plugins/Media/PixelStreaming/Source/PixelStreaming/Private/Settings.h
+++ b/Engine/Plugins/Media/PixelStreaming/Source/PixelStreaming/Private/Settings.h
@@ -47,6 +47,7 @@ namespace UE::PixelStreaming::Settings
 	extern TAutoConsoleVariable<bool> CVarPixelStreamingWebRTCDisableTransmitVideo;
 	extern TAutoConsoleVariable<bool> CVarPixelStreamingWebRTCDisableAudioSync;
 	extern TAutoConsoleVariable<bool> CVarPixelStreamingWebRTCUseLegacyAudioDevice;
+	extern TAutoConsoleVariable<bool> CVarPixelStreamingWebRTCEnableFlexFec;
 	extern TAutoConsoleVariable<bool> CVarPixelStreamingWebRTCDisableStats;
 	extern TAutoConsoleVariable<float> CVarPixelStreamingWebRTCAudioGain;
 	extern TAutoConsoleVariable<bool> CVarPixelStreamingWebRTCNegotiateCodecs;
diff --git a/Engine/Plugins/Media/PixelStreaming/Source/PixelStreaming/Private/VideoEncoderFactorySingleLayer.cpp b/Engine/Plugins/Media/PixelStreaming/Source/PixelStreaming/Private/VideoEncoderFactorySingleLayer.cpp
--- a/Engine/Plugins/Media/PixelStreaming/Source/PixelStreaming/Private/VideoEncoderFactorySingleLayer.cpp
+++ b/Engine/Plugins/Media/PixelStreaming/Source/PixelStreaming/Private/VideoEncoderFactorySingleLayer.cpp
@@ -125,6 +125,13 @@ namespace UE::PixelStreaming
 			{
 				UE_LOG(LogPixelStreaming, Error, TEXT("Selected codec was not a supported codec."));
 			}
+
+			// We can still support flex-fec even if we're not negotiating codecs
+			if (Settings::CVarPixelStreamingWebRTCEnableFlexFec.GetValueOnAnyThread())
+			{
+				SupportedFormats.push_back(webrtc::SdpVideoFormat(cricket::kFlexfecCodecName));
+			}
+
 			return SupportedFormats;
 		}
 
@@ -157,6 +164,11 @@ namespace UE::PixelStreaming
 					}
 				}
 			}
+
+			if (Settings::CVarPixelStreamingWebRTCEnableFlexFec.GetValueOnAnyThread())
+			{
+				SupportedFormats.push_back(webrtc::SdpVideoFormat(cricket::kFlexfecCodecName));
+			}
 		}
 
 		return SupportedFormats;
