Adjust logic to consider running in PIE with the timer manager. Use the editor timer manager in cases where we are running in editor.

diff --git a/Engine/Plugins/Media/PixelStreaming/Source/PixelStreaming/PixelStreaming.Build.cs b/Engine/Plugins/Media/PixelStreaming/Source/PixelStreaming/PixelStreaming.Build.cs
--- a/Engine/Plugins/Media/PixelStreaming/Source/PixelStreaming/PixelStreaming.Build.cs
+++ b/Engine/Plugins/Media/PixelStreaming/Source/PixelStreaming/PixelStreaming.Build.cs
@@ -100,6 +100,13 @@ namespace UnrealBuildTool.Rules
                 "NVML",
                 "AudioPlatformConfiguration"
             });
+            if (Target.bBuildEditor == true)
+            {
+                PrivateDependencyModuleNames.AddRange(new string[]
+                {
+                    "UnrealEd"
+                });
+            }
 
             PrivateDefinitions.Add("PIXELSTREAMING_DUMP_ENCODING=0");
 
diff --git a/Engine/Plugins/Media/PixelStreaming/Source/PixelStreaming/Private/PixelStreamingSignallingConnection.cpp b/Engine/Plugins/Media/PixelStreaming/Source/PixelStreaming/Private/PixelStreamingSignallingConnection.cpp
--- a/Engine/Plugins/Media/PixelStreaming/Source/PixelStreaming/Private/PixelStreamingSignallingConnection.cpp
+++ b/Engine/Plugins/Media/PixelStreaming/Source/PixelStreaming/Private/PixelStreamingSignallingConnection.cpp
@@ -13,6 +13,12 @@
 #include "ToStringExtensions.h"
 #include "GenericPlatform/GenericPlatformHttp.h"
 #include "PixelStreamingModule.h"
+#include "Engine/GameEngine.h"
+#include "Engine/GameInstance.h"
+
+#if WITH_EDITOR
+#include "Editor.h"
+#endif
 
 DECLARE_LOG_CATEGORY_EXTERN(LogPixelStreamingSS, Log, VeryVerbose);
 DEFINE_LOG_CATEGORY(LogPixelStreamingSS);
@@ -806,60 +812,97 @@ bool FPixelStreamingSignallingConnection::GetPlayerIdJson(const FJsonObjectPtr&
 	return false;
 }
 
+namespace UE::PixelStreamingSignallingConnection::Private
+{
+FTimerManager* GetTimerManager()
+{
+#if WITH_EDITOR
+	if (GEditor)
+	{
+		// In editor use the editor manager
+		if (GEditor->IsTimerManagerValid())
+		{
+			return &GEditor->GetTimerManager().Get();
+		}
+	}
+	else
+#endif
+	{
+		// Otherwise we should always have a game instance
+		const TIndirectArray<FWorldContext>& WorldContexts = GEngine->GetWorldContexts();
+		for (const FWorldContext& WorldContext : WorldContexts)
+		{
+			if (WorldContext.WorldType == EWorldType::Game && WorldContext.OwningGameInstance)
+			{
+				return &WorldContext.OwningGameInstance->GetTimerManager();
+			}
+		}
+	}
+
+	return nullptr;
+}
+} // namespace UE::PixelStreamingSignallingConnection::Private
+
 void FPixelStreamingSignallingConnection::StartKeepAliveTimer()
 {
-	// GWorld dereferencing needs to happen on the game thread
+	// Dereferencing needs to happen on the game thread
 	// we dont need to wait since its just setting the timer
 	UE::PixelStreaming::DoOnGameThread([this]() {
-		if (GWorld && !GWorld->GetTimerManager().IsTimerActive(TimerHandle_KeepAlive))
+		using namespace UE::PixelStreamingSignallingConnection::Private;
+		FTimerManager* TimerManager = GetTimerManager();
+		if (TimerManager && !TimerManager->IsTimerActive(TimerHandle_KeepAlive))
 		{
-			GWorld->GetTimerManager().SetTimer(TimerHandle_KeepAlive, FTimerDelegate::CreateRaw(this, &FPixelStreamingSignallingConnection::KeepAlive), KEEP_ALIVE_INTERVAL, true);
+			TimerManager->SetTimer(TimerHandle_KeepAlive, FTimerDelegate::CreateRaw(this, &FPixelStreamingSignallingConnection::KeepAlive), KEEP_ALIVE_INTERVAL, true);
 		}
 	});
 }
 
 void FPixelStreamingSignallingConnection::StopKeepAliveTimer()
 {
-	// GWorld dereferencing needs to happen on the game thread
+	// Dereferencing needs to happen on the game thread
 	// we need to wait because if we're destructing this object we dont
 	// want to call the callback mid/post destruction
 	UE::PixelStreaming::DoOnGameThreadAndWait(MAX_uint32, [this]() {
-		if (GWorld)
+		using namespace UE::PixelStreamingSignallingConnection::Private;
+		FTimerManager* TimerManager = GetTimerManager();
+		if (TimerManager)
 		{
-			GWorld->GetTimerManager().ClearTimer(TimerHandle_KeepAlive);
+			TimerManager->ClearTimer(TimerHandle_KeepAlive);
 		}
 	});
 }
 
 void FPixelStreamingSignallingConnection::StartReconnectTimer()
 {
-	// GWorld dereferencing needs to happen on the game thread
+	// Dereferencing needs to happen on the game thread
 	UE::PixelStreaming::DoOnGameThread([this]() {
 		if (IsEngineExitRequested())
 		{
 			return;
 		}
-
-		if (GWorld && !GWorld->GetTimerManager().IsTimerActive(TimerHandle_Reconnect))
+		using namespace UE::PixelStreamingSignallingConnection::Private;
+		FTimerManager* TimerManager = GetTimerManager();
+		if (TimerManager && !TimerManager->IsTimerActive(TimerHandle_Reconnect))
 		{
 			float ReconnectInterval = UE::PixelStreaming::Settings::CVarPixelStreamingSignalingReconnectInterval.GetValueOnAnyThread();
-			GWorld->GetTimerManager().SetTimer(TimerHandle_Reconnect, FTimerDelegate::CreateRaw(this, &FPixelStreamingSignallingConnection::Connect, Url, true), ReconnectInterval, true);
+			TimerManager->SetTimer(TimerHandle_Reconnect, FTimerDelegate::CreateRaw(this, &FPixelStreamingSignallingConnection::Connect, Url, true), ReconnectInterval, true);
 		}
 	});
 }
 
 void FPixelStreamingSignallingConnection::StopReconnectTimer()
 {
-	// GWorld dereferencing needs to happen on the game thread
+	// Dereferencing needs to happen on the game thread
 	UE::PixelStreaming::DoOnGameThread([this]() {
 		if (IsEngineExitRequested())
 		{
 			return;
 		}
-
-		if (GWorld)
+		using namespace UE::PixelStreamingSignallingConnection::Private;
+		FTimerManager* TimerManager = GetTimerManager();
+		if (TimerManager)
 		{
-			GWorld->GetTimerManager().ClearTimer(TimerHandle_Reconnect);
+			TimerManager->ClearTimer(TimerHandle_Reconnect);
 		}
 	});
 }
