Fix: Pixel Streaming XR streamed stereo image not converging in the headset.

- Fix hard coded projection matrix, FoV, eye offsets, IPD, aspect ratio. Values are now derived from remote device eye transforms and eye projection matrices.

diff --git a/Engine/Plugins/Media/PixelStreaming/Source/PixelStreamingHMD/Private/PixelStreamingHMD.cpp b/Engine/Plugins/Media/PixelStreaming/Source/PixelStreamingHMD/Private/PixelStreamingHMD.cpp
--- a/Engine/Plugins/Media/PixelStreaming/Source/PixelStreamingHMD/Private/PixelStreamingHMD.cpp
+++ b/Engine/Plugins/Media/PixelStreaming/Source/PixelStreamingHMD/Private/PixelStreamingHMD.cpp
@@ -14,12 +14,14 @@
 #include "Widgets/SWindow.h"
 #include "SceneView.h"
 
+DEFINE_LOG_CATEGORY(LogPixelStreamingHMD);
+
 FPixelStreamingHMD::FPixelStreamingHMD(const FAutoRegister& AutoRegister)
 	: FHeadMountedDisplayBase(nullptr)
 	, FHMDSceneViewExtension(AutoRegister)
 	, CurHmdTransform(FTransform::Identity)
 	, WorldToMeters(100.0f)
-	, InterpupillaryDistance(0.064f)
+	, InterpupillaryDistance(0.0315f)
 	, bStereoEnabled(true)
 {
 }
@@ -28,6 +30,105 @@ FPixelStreamingHMD::~FPixelStreamingHMD()
 {
 }
 
+void FPixelStreamingHMD::SetEyeViews(FTransform Left, FMatrix LeftProj, FTransform Right, FMatrix RightProj)
+{
+	CurLeftEyeTransform = Left;
+	CurRightEyeTransform = Right;
+	CurLeftEyeProjMatrix = LeftProj;
+	CurRightEyeProjMatrix = RightProj;
+
+	// Calculate HMD roation as the rotation quaternion in between both eyes
+	FQuat HMDRotation = FQuat::Slerp(Left.GetRotation(), Right.GetRotation(), 0.5);
+
+	// Calculate HMD position as the position between both eyes
+	FVector RightEyeLoc = Right.GetLocation();
+	FVector LeftEyeLoc = Left.GetLocation();
+	FVector Dir = RightEyeLoc - LeftEyeLoc;
+	float IPD = Dir.Size();
+	// Get the location half way between the two eye locations
+	FVector HMDLocation = LeftEyeLoc + (Dir * 0.5);
+	FTransform HMDTransform = FTransform(HMDRotation, HMDLocation, FVector::OneVector);
+
+	// Set the HMD transform
+	SetTransform(HMDTransform);
+
+	// Set the IPD (in meters)
+	SetInterpupillaryDistance(IPD / 100.0f);
+
+	// Calculate the horizontal and vertical FoV from the projection matrix (left and right eye will have same FoVs)
+	HFoVRads = 2.0f * FMath::Atan(1.0f / CurLeftEyeProjMatrix.M[0][0]);
+	VFoVRads = 2.0f * FMath::Atan(1.0f / CurLeftEyeProjMatrix.M[1][1]);
+
+	// Extract the left/right eye projection offsets
+	CurLeftEyeProjOffsetX = -CurLeftEyeProjMatrix.M[0][2];
+	CurLeftEyeProjOffsetY = CurLeftEyeProjMatrix.M[1][2];
+	CurRightEyeProjOffsetX = -CurRightEyeProjMatrix.M[0][2];
+	CurRightEyeProjOffsetY = CurRightEyeProjMatrix.M[1][2];
+
+	// Extract near and farclip planes
+    NearClip = CurLeftEyeProjMatrix.M[3][2] / (CurLeftEyeProjMatrix.M[2][2] - 1);
+    FarClip = CurLeftEyeProjMatrix.M[3][2] / (CurLeftEyeProjMatrix.M[2][2] + 1);
+	SetClippingPlanes(NearClip, FarClip);
+
+	// Calculate target aspect ratio from the projection matrix (left and right eye will have same aspect ratio)
+	//TargetAspectRatio = CurLeftEyeProjMatrix.M[1][1] / CurLeftEyeProjMatrix.M[0][0];
+	TargetAspectRatio = tan(HFoVRads * 0.5f) / tan(VFoVRads * 0.5f);
+
+	TSharedPtr<SWindow> TargetWindow = GEngine->GameViewport->GetWindow();
+	FVector2f SizeInScreen = TargetWindow->GetSizeInScreen();
+	const float InWidth = SizeInScreen.X / 2.f;
+	const float InHeight = SizeInScreen.Y;
+	const float AspectRatio = InWidth / InHeight;
+
+	// If current resolution does not match remote device aspect ratio, we will change resolution to match aspect ratio (though we rate limit res change to every 5s)
+	if(UE::PixelStreamingHMD::Settings::CVarPixelStreamingHMDMatchAspectRatio.GetValueOnAnyThread() &&
+		FMath::Abs(AspectRatio - TargetAspectRatio) > 0.01 &&
+		!bReceivedTransforms)
+	{
+		int TargetHeight = InHeight;
+		int TargetWidth = InHeight * TargetAspectRatio * 2.0f;
+		UE_LOG(LogPixelStreamingHMD, Warning, TEXT("XR Pixel Streaming streaming resolution not matching remote device aspect ratio. Changing resolution to %dx%d"), TargetWidth, TargetHeight);
+		FString ChangeResCommand = FString::Printf(TEXT("r.SetRes %dx%d"), TargetWidth, TargetHeight);
+		GEngine->Exec(GEngine->GetWorld(), *ChangeResCommand);
+	}
+
+	// If we know we are doing XR update some CVars for Pixel Streaming to optimise for it.
+	if(!bReceivedTransforms)
+	{
+		// Couple engine's render rate and streaming rate
+		if (IConsoleVariable* DecoupleFramerateCVar = IConsoleManager::Get().FindConsoleVariable(TEXT("PixelStreaming.DecoupleFramerate")))
+		{
+			DecoupleFramerateCVar->Set(false);
+		}
+
+		// Set the rate at which we will stream
+		if (IConsoleVariable* StreamFPS = IConsoleManager::Get().FindConsoleVariable(TEXT("PixelStreaming.WebRTC.Fps")))
+		{
+			StreamFPS->Set(90);
+		}
+
+		// Set the MinQP to bound quality
+		if (IConsoleVariable* MinQP = IConsoleManager::Get().FindConsoleVariable(TEXT("PixelStreaming.Encoder.MinQP")))
+		{
+			MinQP->Set(15);
+		}
+
+		// Necessary for coupled framerate
+		if (IConsoleVariable* CaptureUseFence = IConsoleManager::Get().FindConsoleVariable(TEXT("PixelStreaming.CaptureUseFence")))
+		{
+			CaptureUseFence->Set(false);
+		}
+
+		// Disable keyframes interval, only send them as needed
+		if (IConsoleVariable* KeyframeInterval = IConsoleManager::Get().FindConsoleVariable(TEXT("PixelStreaming.Encoder.KeyframeInterval")))
+		{
+			KeyframeInterval->Set(0);
+		}
+	}
+
+	bReceivedTransforms = true;
+}
+
 float FPixelStreamingHMD::GetWorldToMetersScale() const
 {
 	return WorldToMeters;
@@ -45,7 +146,7 @@ void FPixelStreamingHMD::EnableHMD(bool enable)
 
 bool FPixelStreamingHMD::GetHMDMonitorInfo(MonitorInfo& MonitorDesc)
 {
-	MonitorDesc.MonitorName = "";
+	MonitorDesc.MonitorName = "PixelStreamingHMD";
 	MonitorDesc.MonitorId = 0;
 	MonitorDesc.DesktopX = MonitorDesc.DesktopY = MonitorDesc.ResolutionX = MonitorDesc.ResolutionY = 0;
 	return false;
@@ -53,8 +154,8 @@ bool FPixelStreamingHMD::GetHMDMonitorInfo(MonitorInfo& MonitorDesc)
 
 void FPixelStreamingHMD::GetFieldOfView(float& OutHFOVInDegrees, float& OutVFOVInDegrees) const
 {
-	OutHFOVInDegrees = 0.0f;
-	OutVFOVInDegrees = 0.0f;
+	OutHFOVInDegrees = FMath::RadiansToDegrees(HFoVRads);
+	OutVFOVInDegrees = FMath::RadiansToDegrees(VFoVRads);
 }
 
 bool FPixelStreamingHMD::EnumerateTrackedDevices(TArray<int32>& OutDevices, EXRTrackedDeviceType Type)
@@ -101,55 +202,7 @@ void FPixelStreamingHMD::ResetOrientationAndPosition(float yaw)
 
 void FPixelStreamingHMD::DrawDistortionMesh_RenderThread(struct FHeadMountedDisplayPassContext& Context, const FIntPoint& TextureSize)
 {
-	float ClipSpaceQuadZ = 0.0f;
-	FMatrix QuadTexTransform = FMatrix::Identity;
-	FMatrix QuadPosTransform = FMatrix::Identity;
-	const FSceneView& View = Context.View;
-	const FIntRect SrcRect = View.UnscaledViewRect;
-
-	FRHICommandListImmediate& RHICmdList = Context.RHICmdList;
-	const FSceneViewFamily& ViewFamily = *(View.Family);
-	FIntPoint ViewportSize = ViewFamily.RenderTarget->GetSizeXY();
-	RHICmdList.SetViewport(0, 0, 0.0f, ViewportSize.X, ViewportSize.Y, 1.0f);
-
-	static const uint32 NumVerts = 4;
-	static const uint32 NumTris = 2;
-
-	static const FDistortionVertex MeshVerts[][NumVerts] = {
-		{
-			// left eye
-			{ FVector2f(-0.9f, -0.9f), FVector2f(0.0f, 1.0f), FVector2f(0.0f, 1.0f), FVector2f(0.0f, 1.0f), 1.0f, 0.0f },
-			{ FVector2f(-0.1f, -0.9f), FVector2f(0.5f, 1.0f), FVector2f(0.5f, 1.0f), FVector2f(0.5f, 1.0f), 1.0f, 0.0f },
-			{ FVector2f(-0.1f, 0.9f), FVector2f(0.5f, 0.0f), FVector2f(0.5f, 0.0f), FVector2f(0.5f, 0.0f), 1.0f, 0.0f },
-			{ FVector2f(-0.9f, 0.9f), FVector2f(0.0f, 0.0f), FVector2f(0.0f, 0.0f), FVector2f(0.0f, 0.0f), 1.0f, 0.0f },
-		},
-		{
-			// right eye
-			{ FVector2f(0.1f, -0.9f), FVector2f(0.5f, 1.0f), FVector2f(0.5f, 1.0f), FVector2f(0.5f, 1.0f), 1.0f, 0.0f },
-			{ FVector2f(0.9f, -0.9f), FVector2f(1.0f, 1.0f), FVector2f(1.0f, 1.0f), FVector2f(1.0f, 1.0f), 1.0f, 0.0f },
-			{ FVector2f(0.9f, 0.9f), FVector2f(1.0f, 0.0f), FVector2f(1.0f, 0.0f), FVector2f(1.0f, 0.0f), 1.0f, 0.0f },
-			{ FVector2f(0.1f, 0.9f), FVector2f(0.5f, 0.0f), FVector2f(0.5f, 0.0f), FVector2f(0.5f, 0.0f), 1.0f, 0.0f },
-		}
-	};
-
-	FRHIResourceCreateInfo CreateInfo(TEXT("FPixelStreamingHMD"));
-	FBufferRHIRef VertexBufferRHI = RHICmdList.CreateVertexBuffer(sizeof(FDistortionVertex) * NumVerts, BUF_Volatile, CreateInfo);
-	void* VoidPtr = RHICmdList.LockBuffer(VertexBufferRHI, 0, sizeof(FDistortionVertex) * NumVerts, RLM_WriteOnly);
-	FPlatformMemory::Memcpy(VoidPtr, MeshVerts[View.StereoViewIndex], sizeof(FDistortionVertex) * NumVerts);
-	RHICmdList.UnlockBuffer(VertexBufferRHI);
-
-	static const uint16 Indices[] = { 0, 1, 2, 0, 2, 3 };
-
-	FBufferRHIRef IndexBufferRHI = RHICmdList.CreateIndexBuffer(sizeof(uint16), sizeof(uint16) * 6, BUF_Volatile, CreateInfo);
-	void* VoidPtr2 = RHICmdList.LockBuffer(IndexBufferRHI, 0, sizeof(uint16) * 6, RLM_WriteOnly);
-	FPlatformMemory::Memcpy(VoidPtr2, Indices, sizeof(uint16) * 6);
-	RHICmdList.UnlockBuffer(IndexBufferRHI);
-
-	RHICmdList.SetStreamSource(0, VertexBufferRHI, 0);
-	RHICmdList.DrawIndexedPrimitive(IndexBufferRHI, 0, 0, NumVerts, 0, NumTris, 1);
-
-	IndexBufferRHI.SafeRelease();
-	VertexBufferRHI.SafeRelease();
+	// Note: Left intentionally blank as we do not want to do any distortion on the UE side, the device will distort the image for us.
 }
 
 bool FPixelStreamingHMD::IsStereoEnabled() const
@@ -171,35 +224,59 @@ void FPixelStreamingHMD::AdjustViewRect(int32 ViewIndex, int32& X, int32& Y, uin
 
 void FPixelStreamingHMD::CalculateStereoViewOffset(const int32 ViewIndex, FRotator& ViewRotation, const float InWorldToMeters, FVector& ViewLocation)
 {
-	if (ViewIndex != INDEX_NONE)
+	if(ViewIndex == INDEX_NONE)
 	{
-		float EyeOffset = 3.20000005f;
-		const float PassOffset = (ViewIndex == EStereoscopicEye::eSSE_LEFT_EYE) ? -EyeOffset : EyeOffset;
+		return;
+	}
+
+	// If not received any transforms yet, just do default offset of half IPD
+	if (!bReceivedTransforms)
+	{
+		float IPDCentimeters = InterpupillaryDistance * 100.0f;
+		const float PassOffset = (ViewIndex == EStereoscopicEye::eSSE_LEFT_EYE) ? -IPDCentimeters * 0.5f : IPDCentimeters * 0.5f;
 		ViewLocation += ViewRotation.Quaternion().RotateVector(FVector(0, PassOffset, 0));
 	}
+	else
+	{
+		FTransform& EyeTransform = (ViewIndex == EStereoscopicEye::eSSE_LEFT_EYE) ? CurLeftEyeTransform : CurRightEyeTransform;
+		FVector LocationOffset = EyeTransform.GetLocation() - CurHmdTransform.GetLocation();
+		ViewLocation += LocationOffset;
+
+		FQuat DeltaRot = EyeTransform.GetRotation() * CurHmdTransform.GetRotation().Inverse();
+		ViewRotation += DeltaRot.Rotator();
+	}
 }
 
 FMatrix FPixelStreamingHMD::GetStereoProjectionMatrix(const int32 ViewIndex) const
 {
-	const float ProjectionCenterOffset = 0.151976421f;
-	const float PassProjectionOffset = (ViewIndex == EStereoscopicEye::eSSE_LEFT_EYE) ? ProjectionCenterOffset : -ProjectionCenterOffset;
+	const float PassProjectionOffsetX = (ViewIndex == EStereoscopicEye::eSSE_LEFT_EYE) ? CurLeftEyeProjOffsetX : CurRightEyeProjOffsetX;
+	const float PassProjectionOffsetY = (ViewIndex == EStereoscopicEye::eSSE_LEFT_EYE) ? CurLeftEyeProjOffsetY : CurRightEyeProjOffsetY;
 
-	const float HalfFov = 2.19686294f / 2.f;
-	TSharedPtr<SWindow> TargetWindow = GEngine->GameViewport->GetWindow();
-	FVector2f SizeInScreen = TargetWindow->GetSizeInScreen();
-	const float InWidth = SizeInScreen.X / 2.f;
-	const float InHeight = SizeInScreen.Y;
-	const float XS = 1.0f / tan(HalfFov);
-	const float YS = InWidth / tan(HalfFov) / InHeight;
+	const float HFoVOverride = UE::PixelStreamingHMD::Settings::CVarPixelStreamingHMDHFOV.GetValueOnAnyThread();
+	const float VFoVOverride = UE::PixelStreamingHMD::Settings::CVarPixelStreamingHMDVFOV.GetValueOnAnyThread();
+	// FoV's are either passed in from the remote device or taken from the FoV override CVars.
+	const float HalfVFov = VFoVOverride > 0.0f ? FMath::DegreesToRadians(VFoVOverride) * 0.5f : VFoVRads * 0.5f;
+	const float HalfHFov = HFoVOverride > 0.0f ? FMath::DegreesToRadians(HFoVOverride) * 0.5f : HFoVRads * 0.5f;
 
 	const float InNearZ = GNearClippingPlane;
-	return FMatrix(
-			   FPlane(XS, 0.0f, 0.0f, 0.0f),
-			   FPlane(0.0f, YS, 0.0f, 0.0f),
-			   FPlane(0.0f, 0.0f, 0.0f, 1.0f),
-			   FPlane(0.0f, 0.0f, InNearZ, 0.0f))
-
-		* FTranslationMatrix(FVector(PassProjectionOffset, 0, 0));
+	const float TanHalfHFov = tan(HalfHFov);
+	const float TanHalfVFov = tan(HalfVFov);
+	const float XS = 1.0f / TanHalfHFov;
+	const float YS = 1.0f / TanHalfVFov;
+
+	// Apply eye off-center translation
+	const FTranslationMatrix TranslationMatrix = FTranslationMatrix(FVector(PassProjectionOffsetX, PassProjectionOffsetY, 0));
+	const float ZNear = GNearClippingPlane_RenderThread;
+
+	FMatrix ProjMatrix = FMatrix(
+		FPlane(XS, 0.0f, 0.0f, 0.0f),
+		FPlane(0.0f, YS, 0.0f, 0.0f),
+		FPlane(0.0f, 0.0f, 0.0f, 1.0f),
+		FPlane(0.0f, 0.0f, ZNear, 0.0f)
+	);
+
+	const FMatrix OutMatrix = ProjMatrix * TranslationMatrix;
+	return OutMatrix;
 }
 
 void FPixelStreamingHMD::GetEyeRenderParams_RenderThread(const FHeadMountedDisplayPassContext& Context, FVector2D& EyeToSrcUVScaleValue, FVector2D& EyeToSrcUVOffsetValue) const
@@ -211,7 +288,8 @@ void FPixelStreamingHMD::GetEyeRenderParams_RenderThread(const FHeadMountedDispl
 void FPixelStreamingHMD::SetupViewFamily(FSceneViewFamily& InViewFamily)
 {
 	InViewFamily.EngineShowFlags.MotionBlur = 0;
-	InViewFamily.EngineShowFlags.HMDDistortion = true;
+	// Note: We do not want to apply any distortion on the UE side.
+	InViewFamily.EngineShowFlags.HMDDistortion = false;
 	InViewFamily.EngineShowFlags.StereoRendering = IsStereoEnabled();
 
 	if (UWorld* World = GWorld)
diff --git a/Engine/Plugins/Media/PixelStreaming/Source/PixelStreamingHMD/Private/Settings.cpp b/Engine/Plugins/Media/PixelStreaming/Source/PixelStreamingHMD/Private/Settings.cpp
--- a/Engine/Plugins/Media/PixelStreaming/Source/PixelStreamingHMD/Private/Settings.cpp
+++ b/Engine/Plugins/Media/PixelStreaming/Source/PixelStreamingHMD/Private/Settings.cpp
@@ -27,15 +27,54 @@ namespace UE::PixelStreamingHMD::Settings
 		}
 	}
 
+	template <typename T>
+	void CommandLineParseValue(const TCHAR* Match, TAutoConsoleVariable<T>& CVar)
+	{
+		T Value;
+		if (FParse::Value(FCommandLine::Get(), Match, Value))
+		{
+			CVar->Set(Value, ECVF_SetByCommandline);
+		}
+	};
+
+	inline void CommandLineParseValue(const TCHAR* Match, TAutoConsoleVariable<FString>& CVar, bool bStopOnSeparator = false)
+	{
+		FString Value;
+		if (FParse::Value(FCommandLine::Get(), Match, Value, bStopOnSeparator))
+		{
+			CVar->Set(*Value, ECVF_SetByCommandline);
+		}
+	};
+
 	TAutoConsoleVariable<bool> CVarPixelStreamingEnableHMD(
 			TEXT("PixelStreaming.HMD.Enable"),
 			false,
 			TEXT("Enables HMD specific functionality for Pixel Streaming. Namely input handling and stereoscopic rendering"),
 			ECVF_Default);
-	
+
+	TAutoConsoleVariable<bool> CVarPixelStreamingHMDMatchAspectRatio(
+			TEXT("PixelStreaming.HMD.MatchAspectRatio"),
+			true,
+			TEXT("If true automatically resize the rendering resolution to match the aspect ratio determined by the HFoV and VFoV"),
+			ECVF_Default);
+
+	TAutoConsoleVariable<float> CVarPixelStreamingHMDHFOV(
+			TEXT("PixelStreaming.HMD.HFOV"),
+			-1.0f,
+			TEXT("Overrides the horizontal field of view for HMD rendering, values are in degrees and values less than 0.0f disable the override."),
+			ECVF_Default);
+
+	TAutoConsoleVariable<float> CVarPixelStreamingHMDVFOV(
+			TEXT("PixelStreaming.HMD.VFOV"),
+			-1.0f,
+			TEXT("Overrides the vertical field of view for HMD rendering, values are in degrees and values less than 0.0f disable the override."),
+			ECVF_Default);
 
 	void InitialiseSettings()
 	{
 		CommandLineParseOption(TEXT("PixelStreamingEnableHMD"), CVarPixelStreamingEnableHMD);
+		CommandLineParseOption(TEXT("PixelStreamingHMDMatchAspectRatio"), CVarPixelStreamingHMDMatchAspectRatio);
+		CommandLineParseValue(TEXT("PixelStreamingHMDHFOV="), CVarPixelStreamingHMDHFOV);
+		CommandLineParseValue(TEXT("PixelStreamingHMDVFOV="), CVarPixelStreamingHMDVFOV);
 	}
 } // UE::PixelStreamingHMD::Settings
\ No newline at end of file
diff --git a/Engine/Plugins/Media/PixelStreaming/Source/PixelStreamingHMD/Private/Settings.h b/Engine/Plugins/Media/PixelStreaming/Source/PixelStreamingHMD/Private/Settings.h
--- a/Engine/Plugins/Media/PixelStreaming/Source/PixelStreamingHMD/Private/Settings.h
+++ b/Engine/Plugins/Media/PixelStreaming/Source/PixelStreamingHMD/Private/Settings.h
@@ -13,6 +13,9 @@ namespace UE::PixelStreamingHMD::Settings
 
 	// Begin CVars
 	extern TAutoConsoleVariable<bool> CVarPixelStreamingEnableHMD;
+	extern TAutoConsoleVariable<bool> CVarPixelStreamingHMDMatchAspectRatio;
+	extern TAutoConsoleVariable<float> CVarPixelStreamingHMDHFOV;
+	extern TAutoConsoleVariable<float> CVarPixelStreamingHMDVFOV;
 	// End CVars
 
 } // namespace UE::PixelStreamingHMD::Settings
diff --git a/Engine/Plugins/Media/PixelStreaming/Source/PixelStreamingHMD/Public/PixelStreamingHMD.h b/Engine/Plugins/Media/PixelStreaming/Source/PixelStreamingHMD/Public/PixelStreamingHMD.h
--- a/Engine/Plugins/Media/PixelStreaming/Source/PixelStreamingHMD/Public/PixelStreamingHMD.h
+++ b/Engine/Plugins/Media/PixelStreaming/Source/PixelStreamingHMD/Public/PixelStreamingHMD.h
@@ -7,6 +7,8 @@
 #include "SceneViewExtension.h"
 #include "XRTrackingSystemBase.h"
 
+DECLARE_LOG_CATEGORY_EXTERN(LogPixelStreamingHMD, Log, All);
+
 class APlayerController;
 class FSceneView;
 class FSceneViewFamily;
@@ -48,9 +50,6 @@ public:
 	virtual class TSharedPtr<class IStereoRendering, ESPMode::ThreadSafe> GetStereoRenderingDevice() override
 	{
 		return SharedThis(this);
-		// TODO (william.belcher): We may want to move the stereo rendering aspect to the front end using WebGL. In that case,
-		// uncomment the line below and delete the line above
-		// return nullptr;
 	}
 
 	protected :
@@ -66,6 +65,7 @@ public:
 	virtual bool GetHMDMonitorInfo(MonitorInfo&) override;
 	virtual void GetFieldOfView(float& OutHFOVInDegrees, float& OutVFOVInDegrees) const override;
 	virtual bool IsChromaAbCorrectionEnabled() const override;
+	virtual bool GetHMDDistortionEnabled(EShadingPath ShadingPath) const override { return false; }
 	virtual void DrawDistortionMesh_RenderThread(struct FHeadMountedDisplayPassContext& Context, const FIntPoint& TextureSize) override;
 
 	/** IStereoRendering interface */
@@ -98,10 +98,26 @@ public:
 	bool IsInitialized() const { return true; }
 
 	void SetTransform(FTransform Transform) { CurHmdTransform = Transform; }
+	void SetEyeViews(FTransform Left, FMatrix LeftProj, FTransform Right, FMatrix RightProj);
 
 private:
 	FTransform CurHmdTransform;
+	FTransform CurLeftEyeTransform;
+	FTransform CurRightEyeTransform;
+	FMatrix CurLeftEyeProjMatrix;
+	FMatrix CurRightEyeProjMatrix;
 	float WorldToMeters;
 	float InterpupillaryDistance;
+	float HFoVRads = FMath::DegreesToRadians(90.0f);
+	float VFoVRads = FMath::DegreesToRadians(90.0f);
+	float CurLeftEyeProjOffsetX = 0.0f;
+	float CurLeftEyeProjOffsetY = 0.0f;
+	float CurRightEyeProjOffsetX = 0.0f;
+	float CurRightEyeProjOffsetY = 0.0f;
+	float TargetAspectRatio = 9.0f/16.0f;
+	float NearClip = 10.0f;
+	float FarClip = 10000.0f;
 	bool bStereoEnabled;
+	bool bReceivedTransforms = false;
+	double LastResChangeSeconds = 0.0f;
 };
\ No newline at end of file
diff --git a/Engine/Plugins/Media/PixelStreaming/Source/PixelStreamingInput/Private/PixelStreamingInputHandler.cpp b/Engine/Plugins/Media/PixelStreaming/Source/PixelStreamingInput/Private/PixelStreamingInputHandler.cpp
--- a/Engine/Plugins/Media/PixelStreaming/Source/PixelStreamingInput/Private/PixelStreamingInputHandler.cpp
+++ b/Engine/Plugins/Media/PixelStreaming/Source/PixelStreamingInput/Private/PixelStreamingInputHandler.cpp
@@ -73,6 +73,7 @@ namespace UE::PixelStreamingInput
 		RegisterMessageHandler("MouseWheel", [this](FString SourceId, FMemoryReader Ar) { HandleOnMouseWheel(Ar); });
 		RegisterMessageHandler("MouseDouble", [this](FString SourceId, FMemoryReader Ar) { HandleOnMouseDoubleClick(Ar); });
 
+		RegisterMessageHandler("XREyeViews", [this](FString SourceId, FMemoryReader Ar) { HandleOnXREyeViews(Ar); });
 		RegisterMessageHandler("XRHMDTransform", [this](FString SourceId, FMemoryReader Ar) { HandleOnXRHMDTransform(Ar); });
 		RegisterMessageHandler("XRControllerTransform", [this](FString SourceId, FMemoryReader Ar) { HandleOnXRControllerTransform(Ar); });
 		RegisterMessageHandler("XRButtonPressed", [this](FString SourceId, FMemoryReader Ar) { HandleOnXRButtonPressed(Ar); });
@@ -883,79 +884,84 @@ namespace UE::PixelStreamingInput
 		}
 	}
 
-	/**
-	 * XR Handling
-	 */
-	void FPixelStreamingInputHandler::HandleOnXRHMDTransform(FMemoryReader Ar)
+	FMatrix FPixelStreamingInputHandler::ExtractWebXRMatrix(FMemoryReader& Ar)
 	{
-		// The buffer contains the transform matrix stored as 16 floats
-		FMatrix HMDMatrix;
+		FMatrix OutMat;
 		for (int32 Row = 0; Row < 4; ++Row)
 
 		{
-			float Col0, Col1, Col2, Col3;
+			float Col0 = 0.0f, Col1 = 0.0f, Col2 = 0.0f, Col3 = 0.0f;
 			Ar << Col0 << Col1 << Col2 << Col3;
-			HMDMatrix.M[Row][0] = Col0;
-			HMDMatrix.M[Row][1] = Col1;
-			HMDMatrix.M[Row][2] = Col2;
-			HMDMatrix.M[Row][3] = Col3;
+			OutMat.M[Row][0] = Col0;
+			OutMat.M[Row][1] = Col1;
+			OutMat.M[Row][2] = Col2;
+			OutMat.M[Row][3] = Col3;
 		}
-		HMDMatrix.DiagnosticCheckNaN();
-		/**
-		 * Converts the 'Y up' 'right handed' WebXR coordinate system transform to Unreal's 'Z up'
-		 * 'left handed' coordinate system.
-		 *
-		 * Ignores scale.
-		 *
-		 * HMD Coordinate space (https://developer.mozilla.org/en-US/docs/Web/API/WebXR_Device_API/Geometry)
-		 */
+		OutMat.DiagnosticCheckNaN();
+		return OutMat;
+	}
+
+	FTransform FPixelStreamingInputHandler::WebXRMatrixToUETransform(FMatrix Mat)
+	{
 		// Rows and columns are swapped between raw mat and FMat
 		FMatrix UEMatrix = FMatrix(
-			FPlane(HMDMatrix.M[0][0], HMDMatrix.M[1][0], HMDMatrix.M[2][0], HMDMatrix.M[3][0]),
-			FPlane(HMDMatrix.M[0][1], HMDMatrix.M[1][1], HMDMatrix.M[2][1], HMDMatrix.M[3][1]),
-			FPlane(HMDMatrix.M[0][2], HMDMatrix.M[1][2], HMDMatrix.M[2][2], HMDMatrix.M[3][2]),
-			FPlane(HMDMatrix.M[0][3], HMDMatrix.M[1][3], HMDMatrix.M[2][3], HMDMatrix.M[3][3]));
+			FPlane(Mat.M[0][0], Mat.M[1][0], Mat.M[2][0], Mat.M[3][0]),
+			FPlane(Mat.M[0][1], Mat.M[1][1], Mat.M[2][1], Mat.M[3][1]),
+			FPlane(Mat.M[0][2], Mat.M[1][2], Mat.M[2][2], Mat.M[3][2]),
+			FPlane(Mat.M[0][3], Mat.M[1][3], Mat.M[2][3], Mat.M[3][3]));
 		// Extract & convert translation
 		FVector Translation = FVector(-UEMatrix.M[3][2], UEMatrix.M[3][0], UEMatrix.M[3][1]) * 100.0f;
 		// Extract & convert rotation
 		FQuat RawRotation(UEMatrix);
 		FQuat Rotation(-RawRotation.Z, RawRotation.X, RawRotation.Y, -RawRotation.W);
+		return FTransform(Rotation, Translation, FVector(Mat.GetScaleVector(1.0f)));
+	}
+
+	/**
+	 * XR Handling
+	 */
+	void FPixelStreamingInputHandler::HandleOnXREyeViews(FMemoryReader Ar)
+	{
+		// The `Ar` buffer contains the left eye transform matrix stored as 16 floats
+		FTransform LeftEyeTransform = WebXRMatrixToUETransform(ExtractWebXRMatrix(Ar));
+
+		// The `Ar` buffer contains the left eye projection matrix stored as 16 floats
+		FMatrix LeftEyeProjectionMatrix = ExtractWebXRMatrix(Ar);
+
+		// The `Ar` buffer contains the right eye transform matrix stored as 16 floats
+		FTransform RightEyeTransform = WebXRMatrixToUETransform(ExtractWebXRMatrix(Ar));
+
+		// The `Ar` buffer contains the right eye projection matrix stored as 16 floats
+		FMatrix RightEyeProjectionMatrix = ExtractWebXRMatrix(Ar);
 
 		if (FPixelStreamingHMD* HMD = IPixelStreamingHMDModule::Get().GetPixelStreamingHMD(); HMD != nullptr)
 		{
-			HMD->SetTransform(FTransform(Rotation, Translation, FVector(HMDMatrix.GetScaleVector(1.0f))));
+			HMD->SetEyeViews(LeftEyeTransform, LeftEyeProjectionMatrix, RightEyeTransform, RightEyeProjectionMatrix);
 		}
 	}
 
-	void FPixelStreamingInputHandler::HandleOnXRControllerTransform(FMemoryReader Ar)
+	void FPixelStreamingInputHandler::HandleOnXRHMDTransform(FMemoryReader Ar)
 	{
-		// The buffer contains the transform matrix stored as 16 floats
-		FMatrix ControllerMatrix;
-		for (int32 Row = 0; Row < 4; ++Row)
+		// The `Ar` buffer contains the transform matrix stored as 16 floats
+		FTransform HMDTransform = WebXRMatrixToUETransform(ExtractWebXRMatrix(Ar));
+
+		if (FPixelStreamingHMD* HMD = IPixelStreamingHMDModule::Get().GetPixelStreamingHMD(); HMD != nullptr)
 		{
-			float Col0, Col1, Col2, Col3;
-			Ar << Col0 << Col1 << Col2 << Col3;
-			// Rows and columns are swapped between raw mat and FMat
-			ControllerMatrix.M[Row][0] = Col0;
-			ControllerMatrix.M[Row][1] = Col1;
-			ControllerMatrix.M[Row][2] = Col2;
-			ControllerMatrix.M[Row][3] = Col3;
+			HMD->SetTransform(HMDTransform);
 		}
-		ControllerMatrix.DiagnosticCheckNaN();
-		FMatrix UEMatrix = FMatrix(
-			FPlane(ControllerMatrix.M[0][0], ControllerMatrix.M[1][0], ControllerMatrix.M[2][0], ControllerMatrix.M[3][0]),
-			FPlane(ControllerMatrix.M[0][1], ControllerMatrix.M[1][1], ControllerMatrix.M[2][1], ControllerMatrix.M[3][1]),
-			FPlane(ControllerMatrix.M[0][2], ControllerMatrix.M[1][2], ControllerMatrix.M[2][2], ControllerMatrix.M[3][2]),
-			FPlane(ControllerMatrix.M[0][3], ControllerMatrix.M[1][3], ControllerMatrix.M[2][3], ControllerMatrix.M[3][3]));
-		// Extract & convert translation
-		FVector Translation = FVector(-UEMatrix.M[3][2], UEMatrix.M[3][0], UEMatrix.M[3][1]) * 100.0f;
-		// Extract & convert rotation
-		FQuat RawRotation(UEMatrix);
-		FQuat Rotation(-RawRotation.Z, RawRotation.X, RawRotation.Y, -RawRotation.W);
-		EControllerHand Handedness;
+	}
+
+	void FPixelStreamingInputHandler::HandleOnXRControllerTransform(FMemoryReader Ar)
+	{
+		// The `Ar` buffer contains the transform matrix stored as 16 floats
+		FTransform ControllerTransform = WebXRMatrixToUETransform(ExtractWebXRMatrix(Ar));
+
+		// The `Ar` buffer contains a UInt8 for the handedness
+		EControllerHand Handedness = EControllerHand::Left;
 		Ar << Handedness;
+
 		FPixelStreamingXRController Controller;
-		Controller.Transform = FTransform(Rotation, Translation, FVector(ControllerMatrix.GetScaleVector(1.0f)));
+		Controller.Transform = ControllerTransform;
 		Controller.Handedness = Handedness;
 		XRControllers.Add(Handedness, Controller);
 	}
@@ -1066,7 +1072,7 @@ namespace UE::PixelStreamingInput
 
 	void FPixelStreamingInputHandler::HandleOnXRSystem(FMemoryReader Ar)
 	{
-		uint8 ActiveSystem;
+		uint8 ActiveSystem = (uint8)EPixelStreamingXRSystem::Unknown;
 		Ar << ActiveSystem;
 		IPixelStreamingHMDModule::Get().SetActiveXRSystem(static_cast<XRSystem>(ActiveSystem));
 	}
@@ -1395,7 +1401,7 @@ namespace UE::PixelStreamingInput
 
 			if (bEditable)
 			{
-				FVector2D NormalizedLocation;
+				FVector2D NormalizedLocation = FVector2D::ZeroVector;
 				TSharedPtr<SWindow> ApplicationWindow = TargetWindow.Pin();
 				if (ApplicationWindow.IsValid())
 				{
diff --git a/Engine/Plugins/Media/PixelStreaming/Source/PixelStreamingInput/Private/PixelStreamingInputHandler.h b/Engine/Plugins/Media/PixelStreaming/Source/PixelStreamingInput/Private/PixelStreamingInputHandler.h
--- a/Engine/Plugins/Media/PixelStreaming/Source/PixelStreamingInput/Private/PixelStreamingInputHandler.h
+++ b/Engine/Plugins/Media/PixelStreaming/Source/PixelStreamingInput/Private/PixelStreamingInputHandler.h
@@ -96,6 +96,7 @@ namespace UE::PixelStreamingInput
 		/**
 		 * XR handling
 		 */
+		virtual void HandleOnXREyeViews(FMemoryReader Ar);
 		virtual void HandleOnXRHMDTransform(FMemoryReader Ar);
 		virtual void HandleOnXRControllerTransform(FMemoryReader Ar);
 		virtual void HandleOnXRButtonPressed(FMemoryReader Ar);
@@ -121,6 +122,19 @@ namespace UE::PixelStreamingInput
 		 */
 		void PopulateDefaultCommandHandlers();
 
+		/**
+		 * Extract 4x4 WebXR ordered matrix and convert to FMatrix.
+		*/
+		FMatrix ExtractWebXRMatrix(FMemoryReader& Ar);
+
+		/**
+		 * Converts the 'Y up' 'right handed' WebXR coordinate system transform to Unreal's 'Z up'
+		 * 'left handed' coordinate system. Note: Ignores scale.
+		 * Assumes WebXR conforms to the following: (https://developer.mozilla.org/en-US/docs/Web/API/WebXR_Device_API/Geometry)
+		 * @return A 4x4 z-up transform matrix for use with UE.
+		*/
+		FTransform WebXRMatrixToUETransform(FMatrix Mat);
+
 		FIntPoint ConvertFromNormalizedScreenLocation(const FVector2D& ScreenLocation, bool bIncludeOffset = true);
 		FWidgetPath FindRoutingMessageWidget(const FVector2D& Location) const;
 		FKey TranslateMouseButtonToKey(const EMouseButtons::Type Button);
diff --git a/Engine/Plugins/Media/PixelStreaming/Source/PixelStreamingInput/Private/PixelStreamingInputModule.cpp b/Engine/Plugins/Media/PixelStreaming/Source/PixelStreamingInput/Private/PixelStreamingInputModule.cpp
--- a/Engine/Plugins/Media/PixelStreaming/Source/PixelStreamingInput/Private/PixelStreamingInputModule.cpp
+++ b/Engine/Plugins/Media/PixelStreaming/Source/PixelStreamingInput/Private/PixelStreamingInputModule.cpp
@@ -105,13 +105,35 @@ namespace UE::PixelStreamingInput
 
 		// XR Input Messages.
 		// clang-format off
+		FPixelStreamingInputProtocol::ToStreamerProtocol.Add("XREyeViews", FPixelStreamingInputMessage(109,		{	// Left eye transform (4x4 matrix)
+																													EType::Float, EType::Float, EType::Float, EType::Float,
+																													EType::Float, EType::Float, EType::Float, EType::Float,
+																													EType::Float, EType::Float, EType::Float, EType::Float,
+																													EType::Float, EType::Float, EType::Float, EType::Float,
+																													// Left eye perspective projection (4x4 matrix)
+																													EType::Float, EType::Float, EType::Float, EType::Float,
+																													EType::Float, EType::Float, EType::Float, EType::Float,
+																													EType::Float, EType::Float, EType::Float, EType::Float,
+																													EType::Float, EType::Float, EType::Float, EType::Float,
+																													// Right eye transform (4x4 matrix)
+																													EType::Float, EType::Float, EType::Float, EType::Float,
+																													EType::Float, EType::Float, EType::Float, EType::Float,
+																													EType::Float, EType::Float, EType::Float, EType::Float,
+																													EType::Float, EType::Float, EType::Float, EType::Float,
+																													// Right eye perspective projection (4x4 matrix)
+																													EType::Float, EType::Float, EType::Float, EType::Float,
+																													EType::Float, EType::Float, EType::Float, EType::Float,
+																													EType::Float, EType::Float, EType::Float, EType::Float,
+																													EType::Float, EType::Float, EType::Float, EType::Float,
+																												}));
+
 		FPixelStreamingInputProtocol::ToStreamerProtocol.Add("XRHMDTransform", FPixelStreamingInputMessage(110, {	// 4x4 Transform
 																													EType::Float, EType::Float, EType::Float, EType::Float,
 																													EType::Float, EType::Float, EType::Float, EType::Float,
 																													EType::Float, EType::Float, EType::Float, EType::Float,
 																													EType::Float, EType::Float, EType::Float, EType::Float,
 																												}));
-		
+
 		FPixelStreamingInputProtocol::ToStreamerProtocol.Add("XRControllerTransform", FPixelStreamingInputMessage(111, {// 4x4 Transform
 																														EType::Float, EType::Float, EType::Float, EType::Float, 
 																														EType::Float, EType::Float, EType::Float, EType::Float, 
@@ -120,7 +142,7 @@ namespace UE::PixelStreamingInput
 																														// Handedness (L, R, Any)
 																														EType::Uint8 
 																														}));
-		
+
 		FPixelStreamingInputProtocol::ToStreamerProtocol.Add("XRButtonPressed", FPixelStreamingInputMessage(112,// Handedness,   ButtonIdx,      IsRepeat
 																												{ EType::Uint8, EType::Uint8, EType::Uint8 }));
 		FPixelStreamingInputProtocol::ToStreamerProtocol.Add("XRButtonTouched", FPixelStreamingInputMessage(113,// Handedness,   ButtonIdx,      IsRepeat
