Fix for PixelStreaming plugin calling ACCodecsCore plugin without a platform allow list causing compile issues on all platforms the PixelStreaming plugin does not support.

#rnx

diff --git a/Engine/Plugins/Media/PixelStreaming/PixelStreaming.uplugin b/Engine/Plugins/Media/PixelStreaming/PixelStreaming.uplugin
--- a/Engine/Plugins/Media/PixelStreaming/PixelStreaming.uplugin
+++ b/Engine/Plugins/Media/PixelStreaming/PixelStreaming.uplugin
@@ -77,7 +77,8 @@
 		},
 		{
 			"Name": "AVCodecsCore",
-			"Enabled": true
+			"Enabled": true,
+			"PlatformAllowList": [ "Win64", "Linux", "Mac" ]
 		},
 		{
 			"Name": "NVCodecs",
