Allowing the pixel streaming component to use the pixelstreamingurl cli argument instead of the given textbox entry.

diff --git a/Engine/Plugins/Media/PixelStreaming/Source/PixelStreamingBlueprint/Private/PixelStreamingStreamerComponent.cpp b/Engine/Plugins/Media/PixelStreaming/Source/PixelStreamingBlueprint/Private/PixelStreamingStreamerComponent.cpp
--- a/Engine/Plugins/Media/PixelStreaming/Source/PixelStreamingBlueprint/Private/PixelStreamingStreamerComponent.cpp
+++ b/Engine/Plugins/Media/PixelStreaming/Source/PixelStreamingBlueprint/Private/PixelStreamingStreamerComponent.cpp
@@ -54,7 +54,16 @@ void UPixelStreamingStreamerComponent::StartStreaming()
 
 	Streamer->SetCoupleFramerate(CoupleFramerate);
 	Streamer->SetStreamFPS(StreamFPS);
-	Streamer->SetSignallingServerURL(SignallingServerURL);
+	if (UsePixelStreamingURL)
+	{
+		FString ServerURL;
+		FParse::Value(FCommandLine::Get(), TEXT("PixelStreamingURL="), ServerURL);
+		Streamer->SetSignallingServerURL(ServerURL);
+	}
+	else
+	{
+		Streamer->SetSignallingServerURL(SignallingServerURL);
+	}
 	Streamer->StartStreaming();
 }
 
diff --git a/Engine/Plugins/Media/PixelStreaming/Source/PixelStreamingBlueprint/Private/PixelStreamingStreamerComponent.h b/Engine/Plugins/Media/PixelStreaming/Source/PixelStreamingBlueprint/Private/PixelStreamingStreamerComponent.h
--- a/Engine/Plugins/Media/PixelStreaming/Source/PixelStreamingBlueprint/Private/PixelStreamingStreamerComponent.h
+++ b/Engine/Plugins/Media/PixelStreaming/Source/PixelStreamingBlueprint/Private/PixelStreamingStreamerComponent.h
@@ -58,6 +58,9 @@ public:
 	UPROPERTY(EditAnywhere, Category = "PixelStreaming")
 	FString SignallingServerURL = "ws://localhost:8888";
 
+	UPROPERTY(EditAnywhere, Category = "PixelStreaming")
+	bool UsePixelStreamingURL = false;
+
 	UPROPERTY(EditAnywhere, Category = "PixelStreaming")
 	int32 StreamFPS = 60;
 
