Fix crash in WebRTC when no audio device is found

diff --git a/Engine/Plugins/Media/PixelStreaming/Source/PixelStreaming/Private/AudioCapturer.cpp b/Engine/Plugins/Media/PixelStreaming/Source/PixelStreaming/Private/AudioCapturer.cpp
--- a/Engine/Plugins/Media/PixelStreaming/Source/PixelStreaming/Private/AudioCapturer.cpp
+++ b/Engine/Plugins/Media/PixelStreaming/Source/PixelStreaming/Private/AudioCapturer.cpp
@@ -123,7 +123,9 @@ namespace UE::PixelStreaming
 	int32 FAudioCapturer::Init()
 	{
 		if (bInitialized)
+		{
 			return 0;
+		}
 
 		{
 			FScopeLock Lock(&DeviceBufferCS);
@@ -132,54 +134,38 @@ namespace UE::PixelStreaming
 			DeviceBuffer = MakeUnique<webrtc::AudioDeviceBuffer>(m_taskQueueFactory.get());
 		}
 
-		// subscribe to audio data
-		if (!GEngine)
+		FAudioDeviceHandle AudioDevice = GEngine ? GEngine->GetMainAudioDevice() : FAudioDeviceHandle();
+		if (AudioDevice)
 		{
-			return -1;
-		}
-
-		FAudioDeviceHandle AudioDevice = GEngine->GetMainAudioDevice();
-		if (!AudioDevice)
-		{
-			UE_LOG(LogPixelStreamingAudioCapturer, Warning, TEXT("No audio device"));
-			return -1;
+			AudioDevice->RegisterSubmixBufferListener(Listener->AsShared(), AudioDevice->GetMainSubmixObject());
 		}
 
 		bInitialized = true;
-		AudioDevice->RegisterSubmixBufferListener(Listener->AsShared(), AudioDevice->GetMainSubmixObject());
-
 		UE_LOG(LogPixelStreamingAudioCapturer, Verbose, TEXT("Init"));
-
 		return 0;
 	}
 
 	int32 FAudioCapturer::Terminate()
 	{
 		if (!bInitialized)
-			return 0;
-
-		// unsubscribe from audio data
-		if (!GEngine)
 		{
-			return -1;
+			return 0;
 		}
 
-		FAudioDeviceHandle AudioDevice = GEngine->GetMainAudioDevice();
-		if (!AudioDevice)
+		// UE-218577: GetMainAudioDevice may have changed since Init was called. Introduce variable (not straightforward due to multi-thread destruction of FAudioCapturer, see ticket).
+		FAudioDeviceHandle AudioDevice = GEngine ? GEngine->GetMainAudioDevice() : FAudioDeviceHandle();
+		if (AudioDevice)
 		{
-			return -1;
+			AudioDevice->UnregisterSubmixBufferListener(Listener->AsShared(), AudioDevice->GetMainSubmixObject());
 		}
-
-		AudioDevice->UnregisterSubmixBufferListener(Listener->AsShared(), AudioDevice->GetMainSubmixObject());
+		
 		bInitialized = false;
-
 		{
 			FScopeLock Lock(&DeviceBufferCS);
 			DeviceBuffer.Reset();
 		}
 
 		UE_LOG(LogPixelStreamingAudioCapturer, Verbose, TEXT("Terminate"));
-
 		return 0;
 	}
 
