XR Pixel Streaming fix convergence on Meta Quest 3.
- Fix: Perspective warping on MQ3 caused by inverted vertical offset in perspective matric
- Fix: Convergence at eye peripheries was wrong, fixed was to remove eye translation/rotation in stereo rendering
- QOL: Only send XR eye views when needed, send HMD the rest of the time.
- QOL: Expose more CVar in Pixel Streaming XR.
- Fix: Handling of XR axes, touches, and button presses would sometimes cause no input or stuck input.

diff --git a/Engine/Plugins/Media/PixelStreaming/Source/PixelStreamingHMD/Private/PixelStreamingHMD.cpp b/Engine/Plugins/Media/PixelStreaming/Source/PixelStreamingHMD/Private/PixelStreamingHMD.cpp
--- a/Engine/Plugins/Media/PixelStreaming/Source/PixelStreamingHMD/Private/PixelStreamingHMD.cpp
+++ b/Engine/Plugins/Media/PixelStreaming/Source/PixelStreamingHMD/Private/PixelStreamingHMD.cpp
@@ -21,7 +21,7 @@ FPixelStreamingHMD::FPixelStreamingHMD(const FAutoRegister& AutoRegister)
 	, FHMDSceneViewExtension(AutoRegister)
 	, CurHmdTransform(FTransform::Identity)
 	, WorldToMeters(100.0f)
-	, InterpupillaryDistance(0.0315f)
+	, InterpupillaryDistance(0.0f)
 	, bStereoEnabled(true)
 {
 }
@@ -30,48 +30,50 @@ FPixelStreamingHMD::~FPixelStreamingHMD()
 {
 }
 
-void FPixelStreamingHMD::SetEyeViews(FTransform Left, FMatrix LeftProj, FTransform Right, FMatrix RightProj)
+void FPixelStreamingHMD::SetEyeViews(FTransform Left, FMatrix LeftProj, FTransform Right, FMatrix RightProj, FTransform HMD)
 {
-	CurLeftEyeTransform = Left;
-	CurRightEyeTransform = Right;
-	CurLeftEyeProjMatrix = LeftProj;
-	CurRightEyeProjMatrix = RightProj;
-
-	// Calculate HMD roation as the rotation quaternion in between both eyes
-	FQuat HMDRotation = FQuat::Slerp(Left.GetRotation(), Right.GetRotation(), 0.5);
-
-	// Calculate HMD position as the position between both eyes
-	FVector RightEyeLoc = Right.GetLocation();
-	FVector LeftEyeLoc = Left.GetLocation();
-	FVector Dir = RightEyeLoc - LeftEyeLoc;
-	float IPD = Dir.Size();
-	// Get the location half way between the two eye locations
-	FVector HMDLocation = LeftEyeLoc + (Dir * 0.5);
-	FTransform HMDTransform = FTransform(HMDRotation, HMDLocation, FVector::OneVector);
+	// This is our intialization message we can use this set the base position if it hasn't been set.
+	if(BasePosition == FVector::ZeroVector)
+	{
+		SetBasePosition(HMD.GetLocation());
+	}
 
 	// Set the HMD transform
-	SetTransform(HMDTransform);
+	SetTransform(HMD);
+
+	// Make left and right relative the HMD
+	FTransform HMDInv = HMD.Inverse();
+	FTransform LeftRelative = Left * HMDInv;
+	FTransform RightRelative = Right * HMDInv;
 
+	// Undo rotation of HMD, then find relative positional offset between eyes and HMD
+	LeftEyePosOffset = LeftRelative.GetLocation();
+	RightEyePosOffset = RightRelative.GetLocation();
+
+	float IPD = FVector::Dist(RightEyePosOffset, LeftEyePosOffset);
 	// Set the IPD (in meters)
 	SetInterpupillaryDistance(IPD / 100.0f);
 
+	// Calculate left/right view orientation relative to HMD
+	LeftEyeRotOffset = LeftRelative.GetRotation();
+	RightEyeRotOffset = RightRelative.GetRotation();
+
 	// Calculate the horizontal and vertical FoV from the projection matrix (left and right eye will have same FoVs)
-	HFoVRads = 2.0f * FMath::Atan(1.0f / CurLeftEyeProjMatrix.M[0][0]);
-	VFoVRads = 2.0f * FMath::Atan(1.0f / CurLeftEyeProjMatrix.M[1][1]);
+	HFoVRads = 2.0f * FMath::Atan(1.0f / LeftProj.M[0][0]);
+	VFoVRads = 2.0f * FMath::Atan(1.0f / LeftProj.M[1][1]);
 
 	// Extract the left/right eye projection offsets
-	CurLeftEyeProjOffsetX = -CurLeftEyeProjMatrix.M[0][2];
-	CurLeftEyeProjOffsetY = CurLeftEyeProjMatrix.M[1][2];
-	CurRightEyeProjOffsetX = -CurRightEyeProjMatrix.M[0][2];
-	CurRightEyeProjOffsetY = CurRightEyeProjMatrix.M[1][2];
+	CurLeftEyeProjOffsetX  = -LeftProj.M[0][2]; // 0.242512569
+	CurLeftEyeProjOffsetY  = -LeftProj.M[1][2]; // 0.193187475
+	CurRightEyeProjOffsetX = -RightProj.M[0][2]; // -0.242512569
+	CurRightEyeProjOffsetY = -RightProj.M[1][2]; // 0.193187475
 
 	// Extract near and farclip planes
-    NearClip = CurLeftEyeProjMatrix.M[3][2] / (CurLeftEyeProjMatrix.M[2][2] - 1);
-    FarClip = CurLeftEyeProjMatrix.M[3][2] / (CurLeftEyeProjMatrix.M[2][2] + 1);
+    NearClip = LeftProj.M[3][2] / (LeftProj.M[2][2] - 1);
+    FarClip = LeftProj.M[3][2] / (LeftProj.M[2][2] + 1);
 	SetClippingPlanes(NearClip, FarClip);
 
 	// Calculate target aspect ratio from the projection matrix (left and right eye will have same aspect ratio)
-	//TargetAspectRatio = CurLeftEyeProjMatrix.M[1][1] / CurLeftEyeProjMatrix.M[0][0];
 	TargetAspectRatio = tan(HFoVRads * 0.5f) / tan(VFoVRads * 0.5f);
 
 	TSharedPtr<SWindow> TargetWindow = GEngine->GameViewport->GetWindow();
@@ -158,6 +160,28 @@ void FPixelStreamingHMD::GetFieldOfView(float& OutHFOVInDegrees, float& OutVFOVI
 	OutVFOVInDegrees = FMath::RadiansToDegrees(VFoVRads);
 }
 
+bool FPixelStreamingHMD::GetRelativeEyePose(int32 DeviceId, int32 ViewIndex, FQuat& OutOrientation, FVector& OutPosition)
+{
+	if (DeviceId != IXRTrackingSystem::HMDDeviceId)
+	{
+		return false;
+	}
+
+	if(UE::PixelStreamingHMD::Settings::CVarPixelStreamingHMDApplyEyePosition.GetValueOnAnyThread())
+	{
+		// If not using override IPD get the actual translation of each eye from the HMD transform and apply that.
+		OutPosition = (ViewIndex == EStereoscopicEye::eSSE_LEFT_EYE) ? LeftEyePosOffset : RightEyePosOffset;
+	}
+
+	// Apply eye rotation if this enabled (default: true)
+	if(UE::PixelStreamingHMD::Settings::CVarPixelStreamingHMDApplyEyeRotation.GetValueOnAnyThread())
+	{
+		OutOrientation = (ViewIndex == EStereoscopicEye::eSSE_LEFT_EYE) ? LeftEyeRotOffset : RightEyeRotOffset;
+	}
+
+	return false;
+}
+
 bool FPixelStreamingHMD::EnumerateTrackedDevices(TArray<int32>& OutDevices, EXRTrackedDeviceType Type)
 {
 	if (Type == EXRTrackedDeviceType::Any || Type == EXRTrackedDeviceType::HeadMountedDisplay)
@@ -229,28 +253,57 @@ void FPixelStreamingHMD::CalculateStereoViewOffset(const int32 ViewIndex, FRotat
 		return;
 	}
 
+	float OverrideIPD = UE::PixelStreamingHMD::Settings::CVarPixelStreamingHMDIPD.GetValueOnAnyThread();
+
 	// If not received any transforms yet, just do default offset of half IPD
 	if (!bReceivedTransforms)
 	{
-		float IPDCentimeters = InterpupillaryDistance * 100.0f;
+		float IPDCentimeters = OverrideIPD > 0.0f ? OverrideIPD : InterpupillaryDistance * 100.0f;
 		const float PassOffset = (ViewIndex == EStereoscopicEye::eSSE_LEFT_EYE) ? -IPDCentimeters * 0.5f : IPDCentimeters * 0.5f;
 		ViewLocation += ViewRotation.Quaternion().RotateVector(FVector(0, PassOffset, 0));
 	}
 	else
 	{
-		FTransform& EyeTransform = (ViewIndex == EStereoscopicEye::eSSE_LEFT_EYE) ? CurLeftEyeTransform : CurRightEyeTransform;
-		FVector LocationOffset = EyeTransform.GetLocation() - CurHmdTransform.GetLocation();
-		ViewLocation += LocationOffset;
+		if(OverrideIPD > 0.0f)
+		{
+			// If using override IPD only translate along the horizontal plane.
+			const float EyeTranslationOffset = (ViewIndex == EStereoscopicEye::eSSE_LEFT_EYE) ? -OverrideIPD * 0.5f : OverrideIPD * 0.5f;
+			ViewLocation += ViewRotation.Quaternion().RotateVector(FVector(0, EyeTranslationOffset, 0));
+		}
+		else if(UE::PixelStreamingHMD::Settings::CVarPixelStreamingHMDApplyEyePosition.GetValueOnAnyThread())
+		{
+			// If not using override IPD get the actual translation of each eye from the HMD transform and apply that.
+			ViewLocation += (ViewIndex == EStereoscopicEye::eSSE_LEFT_EYE) ? LeftEyePosOffset : RightEyePosOffset;
+		}
 
-		FQuat DeltaRot = EyeTransform.GetRotation() * CurHmdTransform.GetRotation().Inverse();
-		ViewRotation += DeltaRot.Rotator();
+		// Apply eye rotation if this enabled (default: true)
+		if(UE::PixelStreamingHMD::Settings::CVarPixelStreamingHMDApplyEyeRotation.GetValueOnAnyThread())
+		{
+			ViewRotation += (ViewIndex == EStereoscopicEye::eSSE_LEFT_EYE) ? LeftEyeRotOffset.Rotator() : RightEyeRotOffset.Rotator();
+		}
 	}
 }
 
 FMatrix FPixelStreamingHMD::GetStereoProjectionMatrix(const int32 ViewIndex) const
 {
-	const float PassProjectionOffsetX = (ViewIndex == EStereoscopicEye::eSSE_LEFT_EYE) ? CurLeftEyeProjOffsetX : CurRightEyeProjOffsetX;
-	const float PassProjectionOffsetY = (ViewIndex == EStereoscopicEye::eSSE_LEFT_EYE) ? CurLeftEyeProjOffsetY : CurRightEyeProjOffsetY;
+
+	float ProjOffsetX = (ViewIndex == EStereoscopicEye::eSSE_LEFT_EYE) ? CurLeftEyeProjOffsetX : CurRightEyeProjOffsetX;
+	float ProjOffsetY = (ViewIndex == EStereoscopicEye::eSSE_LEFT_EYE) ? CurLeftEyeProjOffsetY : CurRightEyeProjOffsetY;
+
+	// Check OverrideProjectOffset X & Y if they have been set by the user, use them instead of the values from WebXR
+	{
+		float OverrideProjectionOffsetX = UE::PixelStreamingHMD::Settings::CVarPixelStreamingHMDProjectionOffsetX.GetValueOnAnyThread();
+		float OverrideProjectionOffsetY = UE::PixelStreamingHMD::Settings::CVarPixelStreamingHMDProjectionOffsetY.GetValueOnAnyThread();
+
+		if(OverrideProjectionOffsetX >= 0.0f)
+		{
+			ProjOffsetX = (ViewIndex == EStereoscopicEye::eSSE_LEFT_EYE) ? OverrideProjectionOffsetX : -OverrideProjectionOffsetX;
+		}
+		if(OverrideProjectionOffsetY >= 0.0f)
+		{
+			ProjOffsetY = OverrideProjectionOffsetY;
+		}
+	}
 
 	const float HFoVOverride = UE::PixelStreamingHMD::Settings::CVarPixelStreamingHMDHFOV.GetValueOnAnyThread();
 	const float VFoVOverride = UE::PixelStreamingHMD::Settings::CVarPixelStreamingHMDVFOV.GetValueOnAnyThread();
@@ -265,17 +318,17 @@ FMatrix FPixelStreamingHMD::GetStereoProjectionMatrix(const int32 ViewIndex) con
 	const float YS = 1.0f / TanHalfVFov;
 
 	// Apply eye off-center translation
-	const FTranslationMatrix TranslationMatrix = FTranslationMatrix(FVector(PassProjectionOffsetX, PassProjectionOffsetY, 0));
+	const FTranslationMatrix OffCenterProjection = FTranslationMatrix(FVector(ProjOffsetX, ProjOffsetY, 0));
 	const float ZNear = GNearClippingPlane_RenderThread;
 
 	FMatrix ProjMatrix = FMatrix(
 		FPlane(XS, 0.0f, 0.0f, 0.0f),
 		FPlane(0.0f, YS, 0.0f, 0.0f),
 		FPlane(0.0f, 0.0f, 0.0f, 1.0f),
-		FPlane(0.0f, 0.0f, ZNear, 0.0f)
+		FPlane(0.0f, 0.0, ZNear, 0.0f)
 	);
 
-	const FMatrix OutMatrix = ProjMatrix * TranslationMatrix;
+	const FMatrix OutMatrix = ProjMatrix * OffCenterProjection;
 	return OutMatrix;
 }
 
diff --git a/Engine/Plugins/Media/PixelStreaming/Source/PixelStreamingHMD/Private/Settings.cpp b/Engine/Plugins/Media/PixelStreaming/Source/PixelStreamingHMD/Private/Settings.cpp
--- a/Engine/Plugins/Media/PixelStreaming/Source/PixelStreamingHMD/Private/Settings.cpp
+++ b/Engine/Plugins/Media/PixelStreaming/Source/PixelStreamingHMD/Private/Settings.cpp
@@ -58,6 +58,18 @@ namespace UE::PixelStreamingHMD::Settings
 			TEXT("If true automatically resize the rendering resolution to match the aspect ratio determined by the HFoV and VFoV"),
 			ECVF_Default);
 
+	TAutoConsoleVariable<bool> CVarPixelStreamingHMDApplyEyePosition(
+			TEXT("PixelStreaming.HMD.ApplyEyePosition"),
+			false,
+			TEXT("If true automatically position each eye's rendering by whatever amount WebXR reports for each left-right XRView. If false do no eye positioning. Default: false"),
+			ECVF_Default);
+
+	TAutoConsoleVariable<bool> CVarPixelStreamingHMDApplyEyeRotation(
+			TEXT("PixelStreaming.HMD.ApplyEyeRotation"),
+			false,
+			TEXT("If true automatically rotate each eye's rendering by whatever amount WebXR reports for each left-right XRView. If false do no eye rotation. Default: false"),
+			ECVF_Default);
+
 	TAutoConsoleVariable<float> CVarPixelStreamingHMDHFOV(
 			TEXT("PixelStreaming.HMD.HFOV"),
 			-1.0f,
@@ -70,11 +82,33 @@ namespace UE::PixelStreamingHMD::Settings
 			TEXT("Overrides the vertical field of view for HMD rendering, values are in degrees and values less than 0.0f disable the override."),
 			ECVF_Default);
 
+	TAutoConsoleVariable<float> CVarPixelStreamingHMDIPD(
+			TEXT("PixelStreaming.HMD.IPD"),
+			-1.0f,
+			TEXT("Overrides the HMD IPD (interpupillary distance), values are in centimeters and values less than 0.0f disable the override."),
+			ECVF_Default);
+
+	TAutoConsoleVariable<float> CVarPixelStreamingHMDProjectionOffsetX(
+			TEXT("PixelStreaming.HMD.ProjectionOffsetX"),
+			-1.0f,
+			TEXT("Overrides the left/right eye projection matrix x-offset, values are in clip space and values less than 0.0f disable the override."),
+			ECVF_Default);
+
+	TAutoConsoleVariable<float> CVarPixelStreamingHMDProjectionOffsetY(
+			TEXT("PixelStreaming.HMD.ProjectionOffsetY"),
+			-1.0f,
+			TEXT("Overrides the left-right eye projection matrix y-offset, values are in clip space and values less than 0.0f disable the override."),
+			ECVF_Default);
+
 	void InitialiseSettings()
 	{
 		CommandLineParseOption(TEXT("PixelStreamingEnableHMD"), CVarPixelStreamingEnableHMD);
 		CommandLineParseOption(TEXT("PixelStreamingHMDMatchAspectRatio"), CVarPixelStreamingHMDMatchAspectRatio);
+		CommandLineParseOption(TEXT("PixelStreamingHMDApplyEyeRotation"), CVarPixelStreamingHMDApplyEyeRotation);
 		CommandLineParseValue(TEXT("PixelStreamingHMDHFOV="), CVarPixelStreamingHMDHFOV);
 		CommandLineParseValue(TEXT("PixelStreamingHMDVFOV="), CVarPixelStreamingHMDVFOV);
+		CommandLineParseValue(TEXT("PixelStreamingHMDVIPD="), CVarPixelStreamingHMDIPD);
+		CommandLineParseValue(TEXT("PixelStreamingHMDProjectionOffsetX="), CVarPixelStreamingHMDProjectionOffsetX);
+		CommandLineParseValue(TEXT("PixelStreamingHMDProjectionOffsetY="), CVarPixelStreamingHMDProjectionOffsetY);
 	}
 } // UE::PixelStreamingHMD::Settings
\ No newline at end of file
diff --git a/Engine/Plugins/Media/PixelStreaming/Source/PixelStreamingHMD/Private/Settings.h b/Engine/Plugins/Media/PixelStreaming/Source/PixelStreamingHMD/Private/Settings.h
--- a/Engine/Plugins/Media/PixelStreaming/Source/PixelStreamingHMD/Private/Settings.h
+++ b/Engine/Plugins/Media/PixelStreaming/Source/PixelStreamingHMD/Private/Settings.h
@@ -14,8 +14,13 @@ namespace UE::PixelStreamingHMD::Settings
 	// Begin CVars
 	extern TAutoConsoleVariable<bool> CVarPixelStreamingEnableHMD;
 	extern TAutoConsoleVariable<bool> CVarPixelStreamingHMDMatchAspectRatio;
+	extern TAutoConsoleVariable<bool> CVarPixelStreamingHMDApplyEyePosition;
+	extern TAutoConsoleVariable<bool> CVarPixelStreamingHMDApplyEyeRotation;
 	extern TAutoConsoleVariable<float> CVarPixelStreamingHMDHFOV;
 	extern TAutoConsoleVariable<float> CVarPixelStreamingHMDVFOV;
+	extern TAutoConsoleVariable<float> CVarPixelStreamingHMDIPD;
+	extern TAutoConsoleVariable<float> CVarPixelStreamingHMDProjectionOffsetX;
+	extern TAutoConsoleVariable<float> CVarPixelStreamingHMDProjectionOffsetY;
 	// End CVars
 
 } // namespace UE::PixelStreamingHMD::Settings
diff --git a/Engine/Plugins/Media/PixelStreaming/Source/PixelStreamingHMD/Public/PixelStreamingHMD.h b/Engine/Plugins/Media/PixelStreaming/Source/PixelStreamingHMD/Public/PixelStreamingHMD.h
--- a/Engine/Plugins/Media/PixelStreaming/Source/PixelStreamingHMD/Public/PixelStreamingHMD.h
+++ b/Engine/Plugins/Media/PixelStreaming/Source/PixelStreamingHMD/Public/PixelStreamingHMD.h
@@ -28,6 +28,9 @@ public:
 	}
 
 	int32 GetXRSystemFlags() const { return EXRSystemFlags::IsHeadMounted; }
+	virtual bool GetRelativeEyePose(int32 DeviceId, int32 ViewIndex, FQuat& OutOrientation, FVector& OutPosition) override;
+	virtual void SetBasePosition(const FVector& InBasePosition) override { BasePosition = InBasePosition; };
+	virtual FVector GetBasePosition() const override { return BasePosition; }
 
 	virtual bool EnumerateTrackedDevices(TArray<int32>& OutDevices, EXRTrackedDeviceType Type = EXRTrackedDeviceType::Any) override;
 
@@ -54,8 +57,7 @@ public:
 
 	protected :
 		/** FXRTrackingSystemBase protected interface */
-		virtual float
-		GetWorldToMetersScale() const override;
+		virtual float GetWorldToMetersScale() const override;
 
 public:
 	/** IHeadMountedDisplay interface */
@@ -98,14 +100,15 @@ public:
 	bool IsInitialized() const { return true; }
 
 	void SetTransform(FTransform Transform) { CurHmdTransform = Transform; }
-	void SetEyeViews(FTransform Left, FMatrix LeftProj, FTransform Right, FMatrix RightProj);
+	void SetEyeViews(FTransform Left, FMatrix LeftProj, FTransform Right, FMatrix RightProj, FTransform HMD);
 
 private:
+	FVector BasePosition = FVector::ZeroVector;
 	FTransform CurHmdTransform;
-	FTransform CurLeftEyeTransform;
-	FTransform CurRightEyeTransform;
-	FMatrix CurLeftEyeProjMatrix;
-	FMatrix CurRightEyeProjMatrix;
+	FVector LeftEyePosOffset;
+	FVector RightEyePosOffset;
+	FQuat LeftEyeRotOffset;
+	FQuat RightEyeRotOffset;
 	float WorldToMeters;
 	float InterpupillaryDistance;
 	float HFoVRads = FMath::DegreesToRadians(90.0f);
diff --git a/Engine/Plugins/Media/PixelStreaming/Source/PixelStreamingInput/Private/PixelStreamingInputHandler.cpp b/Engine/Plugins/Media/PixelStreaming/Source/PixelStreamingInput/Private/PixelStreamingInputHandler.cpp
--- a/Engine/Plugins/Media/PixelStreaming/Source/PixelStreamingInput/Private/PixelStreamingInputHandler.cpp
+++ b/Engine/Plugins/Media/PixelStreaming/Source/PixelStreamingInput/Private/PixelStreamingInputHandler.cpp
@@ -78,6 +78,7 @@ namespace UE::PixelStreamingInput
 		RegisterMessageHandler("XRControllerTransform", [this](FString SourceId, FMemoryReader Ar) { HandleOnXRControllerTransform(Ar); });
 		RegisterMessageHandler("XRButtonPressed", [this](FString SourceId, FMemoryReader Ar) { HandleOnXRButtonPressed(Ar); });
 		RegisterMessageHandler("XRButtonTouched", [this](FString SourceId, FMemoryReader Ar) { HandleOnXRButtonTouched(Ar); });
+		RegisterMessageHandler("XRButtonTouchReleased", [this](FString SourceId, FMemoryReader Ar) { HandleOnXRButtonTouchReleased(Ar); });
 		RegisterMessageHandler("XRButtonReleased", [this](FString SourceId, FMemoryReader Ar) { HandleOnXRButtonReleased(Ar); });
 		RegisterMessageHandler("XRAnalog", [this](FString SourceId, FMemoryReader Ar) { HandleOnXRAnalog(Ar); });
 		RegisterMessageHandler("XRSystem", [this](FString SourceId, FMemoryReader Ar) { HandleOnXRSystem(Ar); });
@@ -264,7 +265,7 @@ namespace UE::PixelStreamingInput
 			return;
 		}
 #endif
-		
+
 		TouchIndicesProcessedThisFrame.Reset();
 
 		FMessage Message;
@@ -273,7 +274,7 @@ namespace UE::PixelStreamingInput
 			FMemoryReader Ar(Message.Data);
 			(*Message.Handler)(Message.SourceId, Ar);
 		}
-		
+
 		ProcessLatestAnalogInputFromThisTick();
 		BroadcastActiveTouchMoveEvents();
 	}
@@ -577,8 +578,23 @@ namespace UE::PixelStreamingInput
 	{
 		const TPayloadThreeParam<uint8, uint8, double> Payload(Ar);
 		const FInputDeviceId ControllerId = FInputDeviceId::CreateFromInternalId((int32)Payload.Param1);
+
+		FKeyId KeyId = Payload.Param2;
+		double AxisValue = Payload.Param3;
+
+		FKey* AnalogKeyPtr = FPixelStreamingInputConverter::GamepadInputToFKey.Find(MakeTuple(KeyId, Action::Axis));
+		if(!AnalogKeyPtr)
+		{
+			return;
+		}
+
+		FAnalogValue AnalogValue;
+		AnalogValue.Value = AxisValue;
+		// Only send axes values continuously in the case of gamepad triggers
+		AnalogValue.bKeepUnlessZero = (KeyId == 5 || KeyId == 6);
+
 		// Overwrite the last data: every tick only process the latest
-		AnalogEventsReceivedThisTick.FindOrAdd(ControllerId).FindOrAdd(Payload.Param2) = Payload.Param3;
+		AnalogEventsReceivedThisTick.FindOrAdd(ControllerId).FindOrAdd(AnalogKeyPtr) = AnalogValue;
 	}
 
 	void FPixelStreamingInputHandler::HandleOnControllerButtonPressed(FMemoryReader Ar)
@@ -909,12 +925,22 @@ namespace UE::PixelStreamingInput
 			FPlane(Mat.M[0][1], Mat.M[1][1], Mat.M[2][1], Mat.M[3][1]),
 			FPlane(Mat.M[0][2], Mat.M[1][2], Mat.M[2][2], Mat.M[3][2]),
 			FPlane(Mat.M[0][3], Mat.M[1][3], Mat.M[2][3], Mat.M[3][3]));
-		// Extract & convert translation
+
+		// Extract scale vector and reorder coordinates to be UE coordinate system.
+		FVector ScaleVectorRaw = UEMatrix.GetScaleVector();
+		// Note: We do not invert Z scaling here because we already handle that when we rebuild translation/rot below.
+		FVector ScaleVector = FVector(ScaleVectorRaw.Z, ScaleVectorRaw.X, ScaleVectorRaw.Y);
+
+		// Temporarily remove scaling component as we need rotation axes to be unit length for proper quat conversion
+		UEMatrix.RemoveScaling();
+
+		// Extract & convert translation component to UE coordinate syste,
 		FVector Translation = FVector(-UEMatrix.M[3][2], UEMatrix.M[3][0], UEMatrix.M[3][1]) * 100.0f;
-		// Extract & convert rotation
+
+		// Extract & convert rotation component to UE coordinate system
 		FQuat RawRotation(UEMatrix);
 		FQuat Rotation(-RawRotation.Z, RawRotation.X, RawRotation.Y, -RawRotation.W);
-		return FTransform(Rotation, Translation, FVector(Mat.GetScaleVector(1.0f)));
+		return FTransform(Rotation, Translation, ScaleVector);
 	}
 
 	/**
@@ -934,9 +960,12 @@ namespace UE::PixelStreamingInput
 		// The `Ar` buffer contains the right eye projection matrix stored as 16 floats
 		FMatrix RightEyeProjectionMatrix = ExtractWebXRMatrix(Ar);
 
+		// The `Ar` buffer contains the right eye projection matrix stored as 16 floats
+		FTransform HMDTransform = WebXRMatrixToUETransform(ExtractWebXRMatrix(Ar));
+
 		if (FPixelStreamingHMD* HMD = IPixelStreamingHMDModule::Get().GetPixelStreamingHMD(); HMD != nullptr)
 		{
-			HMD->SetEyeViews(LeftEyeTransform, LeftEyeProjectionMatrix, RightEyeTransform, RightEyeProjectionMatrix);
+			HMD->SetEyeViews(LeftEyeTransform, LeftEyeProjectionMatrix, RightEyeTransform, RightEyeProjectionMatrix, HMDTransform);
 		}
 	}
 
@@ -987,7 +1016,7 @@ namespace UE::PixelStreamingInput
 		UE_LOG(LogPixelStreamingInputHandler, Verbose, TEXT("XR_TOUCHED: ControllerId = %d; KeyName = %s; IsRepeat = %s;"), ControllerId.GetId(), *ButtonPtr->ToString(), bIsRepeat ? TEXT("True") : TEXT("False"));
 	}
 
-	void FPixelStreamingInputHandler::HandleOnXRButtonPressed(FMemoryReader Ar)
+	void FPixelStreamingInputHandler::HandleOnXRButtonTouchReleased(FMemoryReader Ar)
 	{
 		TPayloadThreeParam<uint8, uint8, uint8> Payload(Ar);
 		IPlatformInputDeviceMapper& DeviceMapper = IPlatformInputDeviceMapper::Get();
@@ -998,24 +1027,53 @@ namespace UE::PixelStreamingInput
 		uint8 ButtonIdx = Payload.Param2;
 		bool bIsRepeat = Payload.Param3 != 0;
 
-		FKey* ButtonPtr = FPixelStreamingInputConverter::XRInputToFKey.Find(MakeTuple(System, Handedness, ButtonIdx, Action::Click));
+		FKey* ButtonPtr = FPixelStreamingInputConverter::XRInputToFKey.Find(MakeTuple(System, Handedness, ButtonIdx, Action::Touch));
 		if (ButtonPtr == nullptr)
 		{
 			return;
 		}
-		MessageHandler->OnControllerButtonPressed(ButtonPtr->GetFName(), DeviceMapper.GetPrimaryPlatformUser(), ControllerId, bIsRepeat);
-		UE_LOG(LogPixelStreamingInputHandler, Verbose, TEXT("XR_PRESSED: ControllerId = %d; KeyName = %s; IsRepeat = %s;"), ControllerId.GetId(), *ButtonPtr->ToString(), bIsRepeat ? TEXT("True") : TEXT("False"));
+
+		MessageHandler->OnControllerButtonReleased(ButtonPtr->GetFName(), DeviceMapper.GetPrimaryPlatformUser(), ControllerId, bIsRepeat);
+		UE_LOG(LogPixelStreamingInputHandler, Verbose, TEXT("XR_TOUCH_RELEASED: ControllerId = %d; KeyName = %s; IsRepeat = %s;"), ControllerId.GetId(), *ButtonPtr->ToString(), bIsRepeat ? TEXT("True") : TEXT("False"));
+	}
+
+	void FPixelStreamingInputHandler::HandleOnXRButtonPressed(FMemoryReader Ar)
+	{
+		TPayloadFourParam<uint8, uint8, uint8, double> Payload(Ar);
+		IPlatformInputDeviceMapper& DeviceMapper = IPlatformInputDeviceMapper::Get();
+		FInputDeviceId ControllerId = DeviceMapper.GetDefaultInputDevice();
+
+		XRSystem System = IPixelStreamingHMDModule::Get().GetActiveXRSystem();
+		EControllerHand Handedness = static_cast<EControllerHand>(Payload.Param1);
+		uint8 ButtonIdx = Payload.Param2;
+		bool bIsRepeat = Payload.Param3 != 0;
+
+		FKey* ButtonPtr = FPixelStreamingInputConverter::XRInputToFKey.Find(MakeTuple(System, Handedness, ButtonIdx, Action::Click));
 
 		// Try and see if there is an axis associated with this button (usually the case for triggers)
-		// if so, trigger an axis movement with value 1.0
 		FKey* AxisPtr = FPixelStreamingInputConverter::XRInputToFKey.Find(MakeTuple(System, Handedness, ButtonIdx, Action::Axis));
-		if (AxisPtr == nullptr)
+
+		// Only send button click if there is no axis associated with this button or this the first press of the axis
+		if (ButtonPtr)
 		{
-			return;
+			MessageHandler->OnControllerButtonPressed(ButtonPtr->GetFName(), DeviceMapper.GetPrimaryPlatformUser(), ControllerId, bIsRepeat);
+			UE_LOG(LogPixelStreamingInputHandler, Verbose, TEXT("XR_PRESSED: ControllerId = %d; KeyName = %s; IsRepeat = %s"), ControllerId.GetId(), *ButtonPtr->ToString(), bIsRepeat ? TEXT("True") : TEXT("False"));
+		}
+
+		// If we have axis associate with this press then set axis value to the button press value
+		if(AxisPtr)
+		{
+			// Trigger axes are not robust to some inputs missing across frames.
+			// So to protect against this case (where PS is too slow to transmit them)
+			// we must set the `bKeepUnlessZero` flag, which allows the input to be continuously passed
+			// until a release is fired.
+			FAnalogValue AnalogValue;
+			AnalogValue.bKeepUnlessZero = true;
+			AnalogValue.Value = Payload.Param4;
+			AnalogValue.bIsRepeat = false;
+			UE_LOG(LogPixelStreamingInputHandler, Verbose, TEXT("XR_ANALOG: ControllerId = %d; KeyName = %s; IsRepeat = False; AnalogValue = %.4f; [Queued for Tick()]"), ControllerId.GetId(), *AxisPtr->ToString(), AnalogValue.Value);
+			AnalogEventsReceivedThisTick.FindOrAdd(ControllerId).FindOrAdd(AxisPtr) = AnalogValue;
 		}
-		float AnalogValue = 1.0f;
-		MessageHandler->OnControllerAnalog(AxisPtr->GetFName(), DeviceMapper.GetPrimaryPlatformUser(), ControllerId, AnalogValue);
-		UE_LOG(LogPixelStreamingInputHandler, Verbose, TEXT("XR_ANALOG: ControllerId = %d; KeyName = %s; AnalogValue = %.4f;"), ControllerId.GetId(), *AxisPtr->ToString(), AnalogValue);
 	}
 
 	void FPixelStreamingInputHandler::HandleOnXRButtonReleased(FMemoryReader Ar)
@@ -1029,25 +1087,29 @@ namespace UE::PixelStreamingInput
 		uint8 ButtonIdx = Payload.Param2;
 		bool bIsRepeat = Payload.Param3 != 0;
 
-		FKey* ButtonPtr = FPixelStreamingInputConverter::XRInputToFKey.Find(MakeTuple(System, Handedness, ButtonIdx, Action::Click));
-		if (ButtonPtr == nullptr)
-		{
-			return;
-		}
-		MessageHandler->OnControllerButtonReleased(ButtonPtr->GetFName(), DeviceMapper.GetPrimaryPlatformUser(), ControllerId, bIsRepeat);
-		UE_LOG(LogPixelStreamingInputHandler, Verbose, TEXT("XR_RELEASED: ControllerId = %d; KeyName = %s; IsRepeat = %s;"), ControllerId.GetId(), *ButtonPtr->ToString(), bIsRepeat ? TEXT("True") : TEXT("False"));
-
 		// Try and see if there is an axis associated with this button (usually the case for triggers)
-		// if so, trigger an axis movement with value 0.0
 		FKey* AxisPtr = FPixelStreamingInputConverter::XRInputToFKey.Find(MakeTuple(System, Handedness, ButtonIdx, Action::Axis));
-		if (AxisPtr == nullptr)
+
+		// If we have axis associate with this release then set axis value to 0.0
+		if (AxisPtr)
 		{
-			return;
+			// In the case of an axes release, we should clear any analog value that is being
+			// applied across ticks.
+			FAnalogValue AnalogValue;
+			AnalogValue.bKeepUnlessZero = true;
+			AnalogValue.Value = 0.0;
+			AnalogValue.bIsRepeat = false;
+			UE_LOG(LogPixelStreamingInputHandler, Verbose, TEXT("XR_ANALOG: ControllerId = %d; KeyName = %s; IsRepeat = False; AnalogValue = %.4f; [Queued for Tick()]"), ControllerId.GetId(), *AxisPtr->ToString(), AnalogValue.Value);
+			AnalogEventsReceivedThisTick.FindOrAdd(ControllerId).FindOrAdd(AxisPtr) = AnalogValue;
 		}
 
-		float AnalogValue = 0.0f;
-		MessageHandler->OnControllerAnalog(AxisPtr->GetFName(), DeviceMapper.GetPrimaryPlatformUser(), ControllerId, AnalogValue);
-		UE_LOG(LogPixelStreamingInputHandler, Verbose, TEXT("XR_ANALOG: ControllerId = %d; KeyName = %s; AnalogValue = %.4f;"), ControllerId.GetId(), *AxisPtr->ToString(), AnalogValue);
+		// Do the actual release after the analog trigger, as the release can cancel any further inputs
+		FKey* ButtonPtr = FPixelStreamingInputConverter::XRInputToFKey.Find(MakeTuple(System, Handedness, ButtonIdx, Action::Click));
+		if (ButtonPtr)
+		{
+			MessageHandler->OnControllerButtonReleased(ButtonPtr->GetFName(), DeviceMapper.GetPrimaryPlatformUser(), ControllerId, bIsRepeat);
+			UE_LOG(LogPixelStreamingInputHandler, Verbose, TEXT("XR_RELEASED: ControllerId = %d; KeyName = %s; IsRepeat = %s;"), ControllerId.GetId(), *ButtonPtr->ToString(), bIsRepeat ? TEXT("True") : TEXT("False"));
+		}
 	}
 
 	void FPixelStreamingInputHandler::HandleOnXRAnalog(FMemoryReader Ar)
@@ -1059,15 +1121,22 @@ namespace UE::PixelStreamingInput
 		XRSystem System = IPixelStreamingHMDModule::Get().GetActiveXRSystem();
 		EControllerHand Handedness = static_cast<EControllerHand>(Payload.Param1);
 		Action InputAction = static_cast<Action>(Payload.Param2 % 2);
-		FKey* ButtonPtr = FPixelStreamingInputConverter::XRInputToFKey.Find(MakeTuple(System, Handedness, Payload.Param2, InputAction));
-		if (ButtonPtr == nullptr)
+		int AxisIndex = Payload.Param2;
+		FKey* AnalogKeyPtr = FPixelStreamingInputConverter::XRInputToFKey.Find(MakeTuple(System, Handedness, AxisIndex, InputAction));
+		if (AnalogKeyPtr == nullptr)
 		{
 			return;
 		}
 
-		float AnalogValue = (float)Payload.Param3;
-		// UE_LOG(LogPixelStreamingInputHandler, Verbose, TEXT("XR_ANALOG: ControllerId = %d; KeyName = %s; AnalogValue = %.4f;"), ControllerId.GetId(), *Button.ToString(), AnalogValue);
-		MessageHandler->OnControllerAnalog(ButtonPtr->GetFName(), DeviceMapper.GetPrimaryPlatformUser(), ControllerId, AnalogValue);
+		// This codepath is used for XR joysticks, which seems to be robust to temporary drops in input transmission
+		// so we can safely set `bKeepUnlessZero` to false. However, if we use this for more than joysticks we will have to conditionally set this.
+		FAnalogValue AnalogValue;
+		AnalogValue.bKeepUnlessZero = false;
+		AnalogValue.Value = Payload.Param3;
+		AnalogValue.bIsRepeat = false;
+
+		UE_LOG(LogPixelStreamingInputHandler, Verbose, TEXT("XR_ANALOG: ControllerId = %d; KeyName = %s; IsRepeat = False; AnalogValue = %.4f; [Queued for Tick()]"), ControllerId.GetId(), *AnalogKeyPtr->ToString(), AnalogValue.Value);
+		AnalogEventsReceivedThisTick.FindOrAdd(ControllerId).FindOrAdd(AnalogKeyPtr) = AnalogValue;
 	}
 
 	void FPixelStreamingInputHandler::HandleOnXRSystem(FMemoryReader Ar)
@@ -1279,37 +1348,62 @@ namespace UE::PixelStreamingInput
 
 	void FPixelStreamingInputHandler::ProcessLatestAnalogInputFromThisTick()
 	{
-		for (auto AnalogInputIt = AnalogEventsReceivedThisTick.CreateConstIterator(); AnalogInputIt; ++AnalogInputIt)
+		for (auto AnalogInputIt = AnalogEventsReceivedThisTick.CreateIterator(); AnalogInputIt; ++AnalogInputIt)
 		{
-			for (auto FKeyIt = AnalogInputIt->Value.CreateConstIterator(); FKeyIt; ++FKeyIt)
+			for (auto FKeyIt = AnalogInputIt->Value.CreateIterator(); FKeyIt; ++FKeyIt)
 			{
-				ProcessAnalog(AnalogInputIt->Key, FKeyIt->Key, FKeyIt->Value);
-			}
-		}
-		AnalogEventsReceivedThisTick.Reset();
-	}
+				const FInputDeviceId& ControllerId = AnalogInputIt->Key;
+				FKey* Key = FKeyIt->Key;
+				FAnalogValue AnalogValue = FKeyIt->Value;
+				bool bIsRepeat = AnalogValue.bIsRepeat;
 
-	void FPixelStreamingInputHandler::ProcessAnalog(const FInputDeviceId& ControllerId, FKeyId Key, FAnalogValue AnalogValue)
-	{
-		const FKey* AxisPtr = FPixelStreamingInputConverter::GamepadInputToFKey.Find(MakeTuple(Key, Action::Axis));
-		if (!AxisPtr)
-		{
-			return;
+				if (!Key)
+				{
+					return;
+				}
+
+				// Pass an analog input along the engine's input processing system
+				FSlateApplication& SlateApplication = FSlateApplication::Get();
+				const FAnalogInputEvent AnalogInputEvent(
+					*Key,															/* InKey */
+					SlateApplication.GetPlatformApplication()->GetModifierKeys(),	/* InModifierKeys */
+					ControllerId,													/* InDeviceId */
+					bIsRepeat,														/* bInIsRepeat*/
+					0,																/* InCharacterCode */
+					0,																/* InKeyCode */
+					AnalogValue.Value,												/* InAnalogValue */
+					// TODO (william.belcher): This user idx should be the playerId
+					0 /* InUserIndex */
+				);
+
+				bool bHandled = FSlateApplication::Get().ProcessAnalogInputEvent(AnalogInputEvent);
+				UE_LOG(LogPixelStreamingInputHandler, Verbose, TEXT("TICKED ANALOG Input: ControllerId = %d; KeyName = %s; IsRepeat = %s; AnalogValue = %.4f; Handled = %s; [Queued for Tick()]"), ControllerId.GetId(), *Key->ToString(), bIsRepeat ? TEXT("True") : TEXT("False"), AnalogValue.Value, bHandled ? TEXT("True") : TEXT("False"));
+
+				// Remove current analog key unless it has the special `bKeepUnlessZero` flag set.
+				// This flag is used to continuously apply input values across ticks because
+				// Pixel Streaming may not have transmitted an axis value in time for the next tick.
+				// But in all ordinary cases where this flag is not set, the stored analog value should
+				// be dropped from the map so the input for the axis (e.g. joystick) is only applied the frame
+				// it is received. The `bKeepUnlessZero` is used for trigger axes, where a temporary drop in
+				// input triggers UE into thinking a full press/release should occur.
+				if(!AnalogValue.bKeepUnlessZero)
+				{
+					FKeyIt.RemoveCurrent();
+				}
+				else if(AnalogValue.bKeepUnlessZero && AnalogValue.Value == 0.0)
+				{
+					// HACK: If we have zero, send it again next frame to ensure we trigger a release internally
+					// Without this release does not seem to get processed for axes inputs
+					FKeyIt->Value.bIsRepeat = true;
+					FKeyIt->Value.bKeepUnlessZero = false;
+				}
+				else
+				{
+					// We are resending the same input, signal this is the case on UE side
+					FKeyIt->Value.bIsRepeat = true;
+				}
+			}
 		}
-		
-		FSlateApplication& SlateApplication = FSlateApplication::Get();
-		const FAnalogInputEvent AnalogInputEvent(
-			*AxisPtr,														/* InKey */
-			SlateApplication.GetPlatformApplication()->GetModifierKeys(),	/* InModifierKeys */
-			ControllerId,													/* InDeviceId */
-			false,													/* bInIsRepeat*/
-			0,													/* InCharacterCode */
-			0,														/* InKeyCode */
-			AnalogValue,													/* InAnalogValue */
-			// TODO (william.belcher): This user idx should be the playerId
-			0 /* InUserIndex */
-		);
-		FSlateApplication::Get().ProcessAnalogInputEvent(AnalogInputEvent);
 	}
 
 	void FPixelStreamingInputHandler::BroadcastActiveTouchMoveEvents()
diff --git a/Engine/Plugins/Media/PixelStreaming/Source/PixelStreamingInput/Private/PixelStreamingInputHandler.h b/Engine/Plugins/Media/PixelStreaming/Source/PixelStreamingInput/Private/PixelStreamingInputHandler.h
--- a/Engine/Plugins/Media/PixelStreaming/Source/PixelStreamingInput/Private/PixelStreamingInputHandler.h
+++ b/Engine/Plugins/Media/PixelStreaming/Source/PixelStreamingInput/Private/PixelStreamingInputHandler.h
@@ -101,6 +101,7 @@ namespace UE::PixelStreamingInput
 		virtual void HandleOnXRControllerTransform(FMemoryReader Ar);
 		virtual void HandleOnXRButtonPressed(FMemoryReader Ar);
 		virtual void HandleOnXRButtonTouched(FMemoryReader Ar);
+		virtual void HandleOnXRButtonTouchReleased(FMemoryReader Ar);
 		virtual void HandleOnXRButtonReleased(FMemoryReader Ar);
 		virtual void HandleOnXRAnalog(FMemoryReader Ar);
 		virtual void HandleOnXRSystem(FMemoryReader Ar);
@@ -146,11 +147,24 @@ namespace UE::PixelStreamingInput
 			int32 ControllerIndex;
 		};
 
+		struct FAnalogValue
+		{
+			/* The actual analog value from the controller axis, typical 0.0..1.0 */
+			double Value;
+			/* If value is non-zero then keep applying this analog values across frames.
+			 * This is useful for trigger axis inputs where if a value is not transmitted
+			 * UE will assume a gap in input means a full trigger press (which is not accurate if we were still pressing).
+			 */
+			bool bKeepUnlessZero = false;
+			/* Has this key event already been fired once? */
+			bool bIsRepeat = false;
+		};
+
 		// Keep a cache of the last touch events as we need to fire Touch Moved every frame while touch is down
 		TMap<int32, FCachedTouchEvent> CachedTouchEvents;
-		
+
 		using FKeyId = uint8;
-		using FAnalogValue = double;
+
 		/**
 		 * If more values are received in a single tick (e.g. could be temp network issue),
 		 * then we only forward the latest value.
@@ -162,12 +176,10 @@ namespace UE::PixelStreamingInput
 		 * The values arrive in the order of recording: that means once the player releases the analog,
 		 * the last analog value would be 0.
 		 */
-		TMap<FInputDeviceId, TMap<FKeyId, FAnalogValue>> AnalogEventsReceivedThisTick;
+		TMap<FInputDeviceId, TMap<FKey*, FAnalogValue>> AnalogEventsReceivedThisTick;
 
 		/** Forwards the latest analog input received for each key this tick. */
 		void ProcessLatestAnalogInputFromThisTick();
-		/** Forward a single analog input the engine. */
-		void ProcessAnalog(const FInputDeviceId& ControllerId, FKeyId Key, FAnalogValue AnalogValue);
 
 		// Track which touch events we processed this frame so we can avoid re-processing them
 		TSet<int32> TouchIndicesProcessedThisFrame;
diff --git a/Engine/Plugins/Media/PixelStreaming/Source/PixelStreamingInput/Private/PixelStreamingInputModule.cpp b/Engine/Plugins/Media/PixelStreaming/Source/PixelStreamingInput/Private/PixelStreamingInputModule.cpp
--- a/Engine/Plugins/Media/PixelStreaming/Source/PixelStreamingInput/Private/PixelStreamingInputModule.cpp
+++ b/Engine/Plugins/Media/PixelStreaming/Source/PixelStreamingInput/Private/PixelStreamingInputModule.cpp
@@ -125,6 +125,11 @@ namespace UE::PixelStreamingInput
 																													EType::Float, EType::Float, EType::Float, EType::Float,
 																													EType::Float, EType::Float, EType::Float, EType::Float,
 																													EType::Float, EType::Float, EType::Float, EType::Float,
+																													// HMD 4x4 Transform
+																													EType::Float, EType::Float, EType::Float, EType::Float,
+																													EType::Float, EType::Float, EType::Float, EType::Float,
+																													EType::Float, EType::Float, EType::Float, EType::Float,
+																													EType::Float, EType::Float, EType::Float, EType::Float,
 																												}));
 
 		FPixelStreamingInputProtocol::ToStreamerProtocol.Add("XRHMDTransform", FPixelStreamingInputMessage(110, {	// 4x4 Transform
@@ -143,17 +148,19 @@ namespace UE::PixelStreamingInput
 																														EType::Uint8 
 																														}));
 
-		FPixelStreamingInputProtocol::ToStreamerProtocol.Add("XRButtonPressed", FPixelStreamingInputMessage(112,// Handedness,   ButtonIdx,      IsRepeat
-																												{ EType::Uint8, EType::Uint8, EType::Uint8 }));
-		FPixelStreamingInputProtocol::ToStreamerProtocol.Add("XRButtonTouched", FPixelStreamingInputMessage(113,// Handedness,   ButtonIdx,      IsRepeat
+		FPixelStreamingInputProtocol::ToStreamerProtocol.Add("XRButtonPressed", FPixelStreamingInputMessage(112,// Handedness,  ButtonIdx,    IsRepeat,     PressedAmount
+																												{ EType::Uint8, EType::Uint8, EType::Uint8, EType::Double }));
+		FPixelStreamingInputProtocol::ToStreamerProtocol.Add("XRButtonTouched", FPixelStreamingInputMessage(113,// Handedness,  ButtonIdx,    IsRepeat
 																												{ EType::Uint8, EType::Uint8, EType::Uint8 }));
-		FPixelStreamingInputProtocol::ToStreamerProtocol.Add("XRButtonReleased", FPixelStreamingInputMessage(114,// Handedness,   ButtonIdx,     IsRepeat
+		FPixelStreamingInputProtocol::ToStreamerProtocol.Add("XRButtonReleased", FPixelStreamingInputMessage(114,// Handedness, ButtonIdx,    IsRepeat
 																												{ EType::Uint8, EType::Uint8, EType::Uint8 }));
 
 		// clang-format on
 		FPixelStreamingInputProtocol::ToStreamerProtocol.Add("XRAnalog", FPixelStreamingInputMessage(115, { EType::Uint8, EType::Uint8, EType::Double }));
 
 		FPixelStreamingInputProtocol::ToStreamerProtocol.Add("XRSystem", FPixelStreamingInputMessage(116, { EType::Uint8 }));
+		FPixelStreamingInputProtocol::ToStreamerProtocol.Add("XRButtonTouchReleased", FPixelStreamingInputMessage(117,// Handedness,  ButtonIdx,    IsRepeat
+																												{ EType::Uint8, EType::Uint8, EType::Uint8 }));
 
 		// Old EToPlayerMsg commands
 		FPixelStreamingInputProtocol::FromStreamerProtocol.Add("QualityControlOwnership", FPixelStreamingInputMessage(0));
