Add missing launch argument parsing for -PixelStreamingHMDApplyEyePosition

#submit-tool bypass

diff --git a/Engine/Plugins/Media/PixelStreaming/Source/PixelStreamingHMD/Private/Settings.cpp b/Engine/Plugins/Media/PixelStreaming/Source/PixelStreamingHMD/Private/Settings.cpp
--- a/Engine/Plugins/Media/PixelStreaming/Source/PixelStreamingHMD/Private/Settings.cpp
+++ b/Engine/Plugins/Media/PixelStreaming/Source/PixelStreamingHMD/Private/Settings.cpp
@@ -105,6 +105,7 @@ namespace UE::PixelStreamingHMD::Settings
 		CommandLineParseOption(TEXT("PixelStreamingEnableHMD"), CVarPixelStreamingEnableHMD);
 		CommandLineParseOption(TEXT("PixelStreamingHMDMatchAspectRatio"), CVarPixelStreamingHMDMatchAspectRatio);
 		CommandLineParseOption(TEXT("PixelStreamingHMDApplyEyeRotation"), CVarPixelStreamingHMDApplyEyeRotation);
+		CommandLineParseOption(TEXT("PixelStreamingHMDApplyEyePosition"), CVarPixelStreamingHMDApplyEyePosition);
 		CommandLineParseValue(TEXT("PixelStreamingHMDHFOV="), CVarPixelStreamingHMDHFOV);
 		CommandLineParseValue(TEXT("PixelStreamingHMDVFOV="), CVarPixelStreamingHMDVFOV);
 		CommandLineParseValue(TEXT("PixelStreamingHMDVIPD="), CVarPixelStreamingHMDIPD);
