PR #12053: Fix PixelStreaming/WebXR Y-axis

diff --git a/Engine/Plugins/Media/PixelStreaming/Source/PixelStreamingInput/Private/PixelStreamingInputHandler.cpp b/Engine/Plugins/Media/PixelStreaming/Source/PixelStreamingInput/Private/PixelStreamingInputHandler.cpp
--- a/Engine/Plugins/Media/PixelStreaming/Source/PixelStreamingInput/Private/PixelStreamingInputHandler.cpp
+++ b/Engine/Plugins/Media/PixelStreaming/Source/PixelStreamingInput/Private/PixelStreamingInputHandler.cpp
@@ -1132,7 +1132,8 @@ namespace UE::PixelStreamingInput
 		// so we can safely set `bKeepUnlessZero` to false. However, if we use this for more than joysticks we will have to conditionally set this.
 		FAnalogValue AnalogValue;
 		AnalogValue.bKeepUnlessZero = false;
-		AnalogValue.Value = Payload.Param3;
+		// Y-axis is inverted in WebXR Gamepad API compared to UE
+		AnalogValue.Value = AxisIndex % 2 == 0 ? Payload.Param3 : -Payload.Param3;
 		AnalogValue.bIsRepeat = false;
 
 		UE_LOG(LogPixelStreamingInputHandler, Verbose, TEXT("XR_ANALOG: ControllerId = %d; KeyName = %s; IsRepeat = False; AnalogValue = %.4f; [Queued for Tick()]"), ControllerId.GetId(), *AnalogKeyPtr->ToString(), AnalogValue.Value);
