Updating the Pixel Streaming get infra scripts to point to the new EpicGamesExt repo.

diff --git a/Engine/Plugins/Media/PixelStreaming/Resources/WebServers/get_ps_servers.bat b/Engine/Plugins/Media/PixelStreaming/Resources/WebServers/get_ps_servers.bat
--- a/Engine/Plugins/Media/PixelStreaming/Resources/WebServers/get_ps_servers.bat
+++ b/Engine/Plugins/Media/PixelStreaming/Resources/WebServers/get_ps_servers.bat
@@ -45,7 +45,7 @@ if DEFINED ARG (
     )
     if "%ARG%"=="/r" (
         SET "ReleaseVersion=%2"
-        SET "ReleaseUrl=https://github.com/EpicGames/PixelStreamingInfrastructure/releases/download/!ReleaseVersion!/!ReleaseVersion!.zip"
+        SET "ReleaseUrl=[REDACTED_PROFILE_URL]/PixelStreamingInfrastructure/releases/download/!ReleaseVersion!/!ReleaseVersion!.zip"
         SET IsTag=0
         SET FlagPassed=1
         SHIFT
@@ -55,7 +55,7 @@ if DEFINED ARG (
 )
 
 @Rem Name and version of ps-infra that we are downloading
-SET PSInfraOrg=EpicGames
+SET PSInfraOrg=EpicGamesExt
 SET PSInfraRepo=PixelStreamingInfrastructure
 
 @Rem If a UE version is supplied set the right branch or tag to fetch for that version of UE
@@ -109,9 +109,9 @@ if %IsTag%==0 (
   if NOT DEFINED ReleaseUrl (
     @Rem We don't want to auto-set the release version if the user passed an explicit flag.
     if NOT DEFINED FlagPassed (
-      FOR /F "tokens=* USEBACKQ" %%F IN (`curl -s -f -L https://raw.githubusercontent.com/EpicGames/PixelStreamingInfrastructure/%PSInfraTagOrBranch%/RELEASE_VERSION`) DO (
+      FOR /F "tokens=* USEBACKQ" %%F IN (`curl -s -f -L https://raw.githubusercontent.com/EpicGamesExt/PixelStreamingInfrastructure/%PSInfraTagOrBranch%/RELEASE_VERSION`) DO (
         SET "ReleaseVersion=!PSInfraTagOrBranch!-%%F"
-        SET "ReleaseUrl=https://github.com/EpicGames/PixelStreamingInfrastructure/releases/download/!ReleaseVersion!/!ReleaseVersion!.zip"
+        SET "ReleaseUrl=[REDACTED_PROFILE_URL]/PixelStreamingInfrastructure/releases/download/!ReleaseVersion!/!ReleaseVersion!.zip"
       )
     )
   )
@@ -198,6 +198,6 @@ echo  Where:
 echo    /v      Specify a version of Unreal Engine to download the recommended release for
 echo    /b      Specify a specific branch for the tool to download from repo
 echo    /t      Specify a specific tag for the tool to download from repo
-echo    /r      Specify a specific release url path e.g. https://github.com/EpicGames/PixelStreamingInfrastructure/releases/download/${RELEASE_HERE}.zip
+echo    /r      Specify a specific release url path e.g. [REDACTED_PROFILE_URL]/PixelStreamingInfrastructure/releases/download/${RELEASE_HERE}.zip
 echo    /h      Display this help message
 goto :EOF
\ No newline at end of file
diff --git a/Engine/Plugins/Media/PixelStreaming/Resources/WebServers/get_ps_servers.sh b/Engine/Plugins/Media/PixelStreaming/Resources/WebServers/get_ps_servers.sh
--- a/Engine/Plugins/Media/PixelStreaming/Resources/WebServers/get_ps_servers.sh
+++ b/Engine/Plugins/Media/PixelStreaming/Resources/WebServers/get_ps_servers.sh
@@ -17,19 +17,19 @@ print_help() {
            release for
    -b      Specify a specific branch for the tool to download from repo
    -t      Specify a specific tag for the tool to download from repo
-   -r      Specify a specific release url path e.g. https://github.com/EpicGames/PixelStreamingInfrastructure/releases/download/<RELEASE_VERSION>/<RELEASE_VERSION>.zip
+   -r      Specify a specific release url path e.g. [REDACTED_PROFILE_URL]/PixelStreamingInfrastructure/releases/download/<RELEASE_VERSION>/<RELEASE_VERSION>.zip
    -h      Display this help message
 "
  exit 1
 }
 
 # Set all default variables (e.g. # Name and version of ps-infra that we are downloading)
-PSInfraOrg=EpicGames
+PSInfraOrg=EpicGamesExt
 PSInfraRepo=PixelStreamingInfrastructure
 PSInfraTagOrBranch=UE5.4
 RefType=heads
 IsTag=0
-ReleaseUrlBase=https://github.com/EpicGames/PixelStreamingInfrastructure/releases/download
+ReleaseUrlBase=[REDACTED_PROFILE_URL]/PixelStreamingInfrastructure/releases/download
 # Unset any variables that don't have defaults that we use that may have persisted between bash terminals.
 unset Url
 unset DownloadVersion
@@ -107,10 +107,10 @@ fi
 # We have a branch, no user-specified release, then check repo for the presence of a RELEASE_VERSION file in the current branch.
 if [ "$IsTag" -eq 0 ] && [ -z "$ReleaseUrl" ] && [ -z "$FlagPassed" ]
 then
-  RelUrl=https://raw.githubusercontent.com/EpicGames/PixelStreamingInfrastructure/$PSInfraTagOrBranch/RELEASE_VERSION
+  RelUrl=https://raw.githubusercontent.com/EpicGamesExt/PixelStreamingInfrastructure/$PSInfraTagOrBranch/RELEASE_VERSION
   if curl --output /dev/null --silent -r 0-0 --fail "$RelUrl"; then
     ReleaseVersion="$PSInfraTagOrBranch-$(curl -L -s $RelUrl)"
-    ReleaseUrl=https://github.com/EpicGames/PixelStreamingInfrastructure/releases/download/$ReleaseVersion/$ReleaseVersion.tar.gz
+    ReleaseUrl=[REDACTED_PROFILE_URL]/PixelStreamingInfrastructure/releases/download/$ReleaseVersion/$ReleaseVersion.tar.gz
     echo "Valid RELEASE_VERSION file found in Github repo at $RelUrl"
   else
     echo "RELEASE_VERSION file does not exist at: $RelUrl"
