Fix: PIE viewport streaming when there are other non-PIE game instances running (ex: Motion Design rundown)

diff --git a/Engine/Plugins/Media/PixelStreaming/Source/PixelStreamingEditor/Private/PixelStreamingVideoInputViewport.cpp b/Engine/Plugins/Media/PixelStreaming/Source/PixelStreamingEditor/Private/PixelStreamingVideoInputViewport.cpp
--- a/Engine/Plugins/Media/PixelStreaming/Source/PixelStreamingEditor/Private/PixelStreamingVideoInputViewport.cpp
+++ b/Engine/Plugins/Media/PixelStreaming/Source/PixelStreamingEditor/Private/PixelStreamingVideoInputViewport.cpp
@@ -124,7 +124,7 @@ void FPixelStreamingVideoInputViewport::OnWindowRendered(SWindow& InWindow, void
 void FPixelStreamingVideoInputViewport::OnPIEViewportRendered(FViewport* InViewport)
 {
 	// If gameclient viewport is rendered and it is not PIE we are not interested
-	if (!IsPIESessionRunning())
+	if (!IsPIESessionRunning() || !InViewport || !InViewport->IsPlayInEditorViewport())
 	{
 		return;
 	}
