[Pixel Streaming VR] Fix: Enable eye offsets by default and fix bug in applying them.

diff --git a/Engine/Plugins/Media/PixelStreaming/Source/PixelStreamingHMD/Private/PixelStreamingHMD.cpp b/Engine/Plugins/Media/PixelStreaming/Source/PixelStreamingHMD/Private/PixelStreamingHMD.cpp
--- a/Engine/Plugins/Media/PixelStreaming/Source/PixelStreamingHMD/Private/PixelStreamingHMD.cpp
+++ b/Engine/Plugins/Media/PixelStreaming/Source/PixelStreamingHMD/Private/PixelStreamingHMD.cpp
@@ -33,7 +33,7 @@ FPixelStreamingHMD::~FPixelStreamingHMD()
 void FPixelStreamingHMD::SetEyeViews(FTransform Left, FMatrix LeftProj, FTransform Right, FMatrix RightProj, FTransform HMD)
 {
 	// This is our intialization message we can use this set the base position if it hasn't been set.
-	if(BasePosition == FVector::ZeroVector)
+	if (BasePosition == FVector::ZeroVector)
 	{
 		SetBasePosition(HMD.GetLocation());
 	}
@@ -63,29 +63,27 @@ void FPixelStreamingHMD::SetEyeViews(FTransform Left, FMatrix LeftProj, FTransfo
 	VFoVRads = 2.0f * FMath::Atan(1.0f / LeftProj.M[1][1]);
 
 	// Extract the left/right eye projection offsets
-	CurLeftEyeProjOffsetX  = -LeftProj.M[0][2]; // 0.242512569
-	CurLeftEyeProjOffsetY  = -LeftProj.M[1][2]; // 0.193187475
+	CurLeftEyeProjOffsetX = -LeftProj.M[0][2];	 // 0.242512569
+	CurLeftEyeProjOffsetY = -LeftProj.M[1][2];	 // 0.193187475
 	CurRightEyeProjOffsetX = -RightProj.M[0][2]; // -0.242512569
 	CurRightEyeProjOffsetY = -RightProj.M[1][2]; // 0.193187475
 
 	// Extract near and farclip planes
-    NearClip = LeftProj.M[3][2] / (LeftProj.M[2][2] - 1);
-    FarClip = LeftProj.M[3][2] / (LeftProj.M[2][2] + 1);
+	NearClip = LeftProj.M[3][2] / (LeftProj.M[2][2] - 1);
+	FarClip = LeftProj.M[3][2] / (LeftProj.M[2][2] + 1);
 	SetClippingPlanes(NearClip, FarClip);
 
 	// Calculate target aspect ratio from the projection matrix (left and right eye will have same aspect ratio)
 	TargetAspectRatio = tan(HFoVRads * 0.5f) / tan(VFoVRads * 0.5f);
 
 	TSharedPtr<SWindow> TargetWindow = GEngine->GameViewport->GetWindow();
-	FVector2f SizeInScreen = TargetWindow->GetSizeInScreen();
-	const float InWidth = SizeInScreen.X / 2.f;
-	const float InHeight = SizeInScreen.Y;
-	const float AspectRatio = InWidth / InHeight;
+	FVector2f			SizeInScreen = TargetWindow->GetSizeInScreen();
+	const float			InWidth = SizeInScreen.X / 2.f;
+	const float			InHeight = SizeInScreen.Y;
+	const float			AspectRatio = InWidth / InHeight;
 
 	// If current resolution does not match remote device aspect ratio, we will change resolution to match aspect ratio (though we rate limit res change to every 5s)
-	if(UE::PixelStreamingHMD::Settings::CVarPixelStreamingHMDMatchAspectRatio.GetValueOnAnyThread() &&
-		FMath::Abs(AspectRatio - TargetAspectRatio) > 0.01 &&
-		!bReceivedTransforms)
+	if (UE::PixelStreamingHMD::Settings::CVarPixelStreamingHMDMatchAspectRatio.GetValueOnAnyThread() && FMath::Abs(AspectRatio - TargetAspectRatio) > 0.01 && !bReceivedTransforms)
 	{
 		int TargetHeight = InHeight;
 		int TargetWidth = InHeight * TargetAspectRatio * 2.0f;
@@ -95,7 +93,7 @@ void FPixelStreamingHMD::SetEyeViews(FTransform Left, FMatrix LeftProj, FTransfo
 	}
 
 	// If we know we are doing XR update some CVars for Pixel Streaming to optimise for it.
-	if(!bReceivedTransforms)
+	if (!bReceivedTransforms)
 	{
 		// Couple engine's render rate and streaming rate
 		if (IConsoleVariable* DecoupleFramerateCVar = IConsoleManager::Get().FindConsoleVariable(TEXT("PixelStreaming.DecoupleFramerate")))
@@ -167,14 +165,14 @@ bool FPixelStreamingHMD::GetRelativeEyePose(int32 DeviceId, int32 ViewIndex, FQu
 		return false;
 	}
 
-	if(UE::PixelStreamingHMD::Settings::CVarPixelStreamingHMDApplyEyePosition.GetValueOnAnyThread())
+	if (UE::PixelStreamingHMD::Settings::CVarPixelStreamingHMDApplyEyePosition.GetValueOnAnyThread())
 	{
 		// If not using override IPD get the actual translation of each eye from the HMD transform and apply that.
 		OutPosition = (ViewIndex == EStereoscopicEye::eSSE_LEFT_EYE) ? LeftEyePosOffset : RightEyePosOffset;
 	}
 
 	// Apply eye rotation if this enabled (default: true)
-	if(UE::PixelStreamingHMD::Settings::CVarPixelStreamingHMDApplyEyeRotation.GetValueOnAnyThread())
+	if (UE::PixelStreamingHMD::Settings::CVarPixelStreamingHMDApplyEyeRotation.GetValueOnAnyThread())
 	{
 		OutOrientation = (ViewIndex == EStereoscopicEye::eSSE_LEFT_EYE) ? LeftEyeRotOffset : RightEyeRotOffset;
 	}
@@ -248,7 +246,7 @@ void FPixelStreamingHMD::AdjustViewRect(int32 ViewIndex, int32& X, int32& Y, uin
 
 void FPixelStreamingHMD::CalculateStereoViewOffset(const int32 ViewIndex, FRotator& ViewRotation, const float InWorldToMeters, FVector& ViewLocation)
 {
-	if(ViewIndex == INDEX_NONE)
+	if (ViewIndex == INDEX_NONE)
 	{
 		return;
 	}
@@ -258,26 +256,26 @@ void FPixelStreamingHMD::CalculateStereoViewOffset(const int32 ViewIndex, FRotat
 	// If not received any transforms yet, just do default offset of half IPD
 	if (!bReceivedTransforms)
 	{
-		float IPDCentimeters = OverrideIPD > 0.0f ? OverrideIPD : InterpupillaryDistance * 100.0f;
+		float		IPDCentimeters = OverrideIPD > 0.0f ? OverrideIPD : InterpupillaryDistance * 100.0f;
 		const float PassOffset = (ViewIndex == EStereoscopicEye::eSSE_LEFT_EYE) ? -IPDCentimeters * 0.5f : IPDCentimeters * 0.5f;
 		ViewLocation += ViewRotation.Quaternion().RotateVector(FVector(0, PassOffset, 0));
 	}
 	else
 	{
-		if(OverrideIPD > 0.0f)
+		if (OverrideIPD > 0.0f)
 		{
 			// If using override IPD only translate along the horizontal plane.
 			const float EyeTranslationOffset = (ViewIndex == EStereoscopicEye::eSSE_LEFT_EYE) ? -OverrideIPD * 0.5f : OverrideIPD * 0.5f;
 			ViewLocation += ViewRotation.Quaternion().RotateVector(FVector(0, EyeTranslationOffset, 0));
 		}
-		else if(UE::PixelStreamingHMD::Settings::CVarPixelStreamingHMDApplyEyePosition.GetValueOnAnyThread())
+		else if (UE::PixelStreamingHMD::Settings::CVarPixelStreamingHMDApplyEyePosition.GetValueOnAnyThread())
 		{
 			// If not using override IPD get the actual translation of each eye from the HMD transform and apply that.
-			ViewLocation += (ViewIndex == EStereoscopicEye::eSSE_LEFT_EYE) ? LeftEyePosOffset : RightEyePosOffset;
+			ViewLocation += (ViewIndex == EStereoscopicEye::eSSE_LEFT_EYE) ? ViewRotation.Quaternion().RotateVector(LeftEyePosOffset) : ViewRotation.Quaternion().RotateVector(RightEyePosOffset);
 		}
 
 		// Apply eye rotation if this enabled (default: true)
-		if(UE::PixelStreamingHMD::Settings::CVarPixelStreamingHMDApplyEyeRotation.GetValueOnAnyThread())
+		if (UE::PixelStreamingHMD::Settings::CVarPixelStreamingHMDApplyEyeRotation.GetValueOnAnyThread())
 		{
 			ViewRotation += (ViewIndex == EStereoscopicEye::eSSE_LEFT_EYE) ? LeftEyeRotOffset.Rotator() : RightEyeRotOffset.Rotator();
 		}
@@ -295,11 +293,11 @@ FMatrix FPixelStreamingHMD::GetStereoProjectionMatrix(const int32 ViewIndex) con
 		float OverrideProjectionOffsetX = UE::PixelStreamingHMD::Settings::CVarPixelStreamingHMDProjectionOffsetX.GetValueOnAnyThread();
 		float OverrideProjectionOffsetY = UE::PixelStreamingHMD::Settings::CVarPixelStreamingHMDProjectionOffsetY.GetValueOnAnyThread();
 
-		if(OverrideProjectionOffsetX >= 0.0f)
+		if (OverrideProjectionOffsetX >= 0.0f)
 		{
 			ProjOffsetX = (ViewIndex == EStereoscopicEye::eSSE_LEFT_EYE) ? OverrideProjectionOffsetX : -OverrideProjectionOffsetX;
 		}
-		if(OverrideProjectionOffsetY >= 0.0f)
+		if (OverrideProjectionOffsetY >= 0.0f)
 		{
 			ProjOffsetY = OverrideProjectionOffsetY;
 		}
@@ -319,14 +317,13 @@ FMatrix FPixelStreamingHMD::GetStereoProjectionMatrix(const int32 ViewIndex) con
 
 	// Apply eye off-center translation
 	const FTranslationMatrix OffCenterProjection = FTranslationMatrix(FVector(ProjOffsetX, ProjOffsetY, 0));
-	const float ZNear = GNearClippingPlane_RenderThread;
+	const float				 ZNear = GNearClippingPlane_RenderThread;
 
 	FMatrix ProjMatrix = FMatrix(
 		FPlane(XS, 0.0f, 0.0f, 0.0f),
 		FPlane(0.0f, YS, 0.0f, 0.0f),
 		FPlane(0.0f, 0.0f, 0.0f, 1.0f),
-		FPlane(0.0f, 0.0, ZNear, 0.0f)
-	);
+		FPlane(0.0f, 0.0, ZNear, 0.0f));
 
 	const FMatrix OutMatrix = ProjMatrix * OffCenterProjection;
 	return OutMatrix;
diff --git a/Engine/Plugins/Media/PixelStreaming/Source/PixelStreamingHMD/Private/Settings.cpp b/Engine/Plugins/Media/PixelStreaming/Source/PixelStreamingHMD/Private/Settings.cpp
--- a/Engine/Plugins/Media/PixelStreaming/Source/PixelStreamingHMD/Private/Settings.cpp
+++ b/Engine/Plugins/Media/PixelStreaming/Source/PixelStreamingHMD/Private/Settings.cpp
@@ -47,58 +47,58 @@ namespace UE::PixelStreamingHMD::Settings
 	};
 
 	TAutoConsoleVariable<bool> CVarPixelStreamingEnableHMD(
-			TEXT("PixelStreaming.HMD.Enable"),
-			false,
-			TEXT("Enables HMD specific functionality for Pixel Streaming. Namely input handling and stereoscopic rendering"),
-			ECVF_Default);
+		TEXT("PixelStreaming.HMD.Enable"),
+		false,
+		TEXT("Enables HMD specific functionality for Pixel Streaming. Namely input handling and stereoscopic rendering"),
+		ECVF_Default);
 
 	TAutoConsoleVariable<bool> CVarPixelStreamingHMDMatchAspectRatio(
-			TEXT("PixelStreaming.HMD.MatchAspectRatio"),
-			true,
-			TEXT("If true automatically resize the rendering resolution to match the aspect ratio determined by the HFoV and VFoV"),
-			ECVF_Default);
+		TEXT("PixelStreaming.HMD.MatchAspectRatio"),
+		true,
+		TEXT("If true automatically resize the rendering resolution to match the aspect ratio determined by the HFoV and VFoV"),
+		ECVF_Default);
 
 	TAutoConsoleVariable<bool> CVarPixelStreamingHMDApplyEyePosition(
-			TEXT("PixelStreaming.HMD.ApplyEyePosition"),
-			false,
-			TEXT("If true automatically position each eye's rendering by whatever amount WebXR reports for each left-right XRView. If false do no eye positioning. Default: false"),
-			ECVF_Default);
+		TEXT("PixelStreaming.HMD.ApplyEyePosition"),
+		true,
+		TEXT("If true automatically position each eye's rendering by whatever amount WebXR reports for each left-right XRView. If false do no eye positioning. Default: false"),
+		ECVF_Default);
 
 	TAutoConsoleVariable<bool> CVarPixelStreamingHMDApplyEyeRotation(
-			TEXT("PixelStreaming.HMD.ApplyEyeRotation"),
-			false,
-			TEXT("If true automatically rotate each eye's rendering by whatever amount WebXR reports for each left-right XRView. If false do no eye rotation. Default: false"),
-			ECVF_Default);
+		TEXT("PixelStreaming.HMD.ApplyEyeRotation"),
+		true,
+		TEXT("If true automatically rotate each eye's rendering by whatever amount WebXR reports for each left-right XRView. If false do no eye rotation. Default: false"),
+		ECVF_Default);
 
 	TAutoConsoleVariable<float> CVarPixelStreamingHMDHFOV(
-			TEXT("PixelStreaming.HMD.HFOV"),
-			-1.0f,
-			TEXT("Overrides the horizontal field of view for HMD rendering, values are in degrees and values less than 0.0f disable the override."),
-			ECVF_Default);
+		TEXT("PixelStreaming.HMD.HFOV"),
+		-1.0f,
+		TEXT("Overrides the horizontal field of view for HMD rendering, values are in degrees and values less than 0.0f disable the override."),
+		ECVF_Default);
 
 	TAutoConsoleVariable<float> CVarPixelStreamingHMDVFOV(
-			TEXT("PixelStreaming.HMD.VFOV"),
-			-1.0f,
-			TEXT("Overrides the vertical field of view for HMD rendering, values are in degrees and values less than 0.0f disable the override."),
-			ECVF_Default);
+		TEXT("PixelStreaming.HMD.VFOV"),
+		-1.0f,
+		TEXT("Overrides the vertical field of view for HMD rendering, values are in degrees and values less than 0.0f disable the override."),
+		ECVF_Default);
 
 	TAutoConsoleVariable<float> CVarPixelStreamingHMDIPD(
-			TEXT("PixelStreaming.HMD.IPD"),
-			-1.0f,
-			TEXT("Overrides the HMD IPD (interpupillary distance), values are in centimeters and values less than 0.0f disable the override."),
-			ECVF_Default);
+		TEXT("PixelStreaming.HMD.IPD"),
+		-1.0f,
+		TEXT("Overrides the HMD IPD (interpupillary distance), values are in centimeters and values less than 0.0f disable the override."),
+		ECVF_Default);
 
 	TAutoConsoleVariable<float> CVarPixelStreamingHMDProjectionOffsetX(
-			TEXT("PixelStreaming.HMD.ProjectionOffsetX"),
-			-1.0f,
-			TEXT("Overrides the left/right eye projection matrix x-offset, values are in clip space and values less than 0.0f disable the override."),
-			ECVF_Default);
+		TEXT("PixelStreaming.HMD.ProjectionOffsetX"),
+		-1.0f,
+		TEXT("Overrides the left/right eye projection matrix x-offset, values are in clip space and values less than 0.0f disable the override."),
+		ECVF_Default);
 
 	TAutoConsoleVariable<float> CVarPixelStreamingHMDProjectionOffsetY(
-			TEXT("PixelStreaming.HMD.ProjectionOffsetY"),
-			-1.0f,
-			TEXT("Overrides the left-right eye projection matrix y-offset, values are in clip space and values less than 0.0f disable the override."),
-			ECVF_Default);
+		TEXT("PixelStreaming.HMD.ProjectionOffsetY"),
+		-1.0f,
+		TEXT("Overrides the left-right eye projection matrix y-offset, values are in clip space and values less than 0.0f disable the override."),
+		ECVF_Default);
 
 	void InitialiseSettings()
 	{
@@ -112,4 +112,4 @@ namespace UE::PixelStreamingHMD::Settings
 		CommandLineParseValue(TEXT("PixelStreamingHMDProjectionOffsetX="), CVarPixelStreamingHMDProjectionOffsetX);
 		CommandLineParseValue(TEXT("PixelStreamingHMDProjectionOffsetY="), CVarPixelStreamingHMDProjectionOffsetY);
 	}
-} // UE::PixelStreamingHMD::Settings
\ No newline at end of file
+} // namespace UE::PixelStreamingHMD::Settings
\ No newline at end of file
