[Pixel Streaming] - Fix: NumConnections returned by embedded signalling server was inaccurate if multiple streamers were started and you query before they identify themselves.

diff --git a/Engine/Plugins/Media/PixelStreaming/Source/PixelStreamingServers/Private/SignallingServer.cpp b/Engine/Plugins/Media/PixelStreaming/Source/PixelStreamingServers/Private/SignallingServer.cpp
--- a/Engine/Plugins/Media/PixelStreaming/Source/PixelStreamingServers/Private/SignallingServer.cpp
+++ b/Engine/Plugins/Media/PixelStreaming/Source/PixelStreamingServers/Private/SignallingServer.cpp
@@ -46,7 +46,18 @@ namespace UE::PixelStreamingServers
 			if (bConnected)
 			{
 				// Close the websocket connection so others can use it
+				Probe->Close();
 				Probe.Reset();
+
+				// Even after closing the client WS of the probe above it will take another tick to remove the connection
+				// so we manually remove the probe connection here to prevent any erroneous number of streamers being reported.
+				uint16 ProbeId = -1;
+				bool   bHasProbe = StreamersWS->GetFirstConnection(ProbeId);
+				if (bHasProbe)
+				{
+					StreamersWS->Close(ProbeId);
+				}
+
 				return true;
 			}
 			else
@@ -115,7 +126,7 @@ namespace UE::PixelStreamingServers
 	TArray<FWebSocketHttpMount> FSignallingServer::GenerateDirectoriesToServe() const
 	{
 		FString ServersDir;
-		bool bServersDirExists = Utils::GetWebServersDir(ServersDir);
+		bool	bServersDirExists = Utils::GetWebServersDir(ServersDir);
 		if (bServersDirExists)
 		{
 			ServersDir = ServersDir / TEXT("SignallingWebServer");
@@ -130,7 +141,7 @@ namespace UE::PixelStreamingServers
 		if (!bServersDirExists)
 		{
 			FString OutResourcesDir;
-			bool bResourcesDirExists = Utils::GetResourcesDir(OutResourcesDir);
+			bool	bResourcesDirExists = Utils::GetResourcesDir(OutResourcesDir);
 			FString NotFoundDir = OutResourcesDir / TEXT("NotFound");
 
 			if (bResourcesDirExists && FPaths::DirectoryExists(NotFoundDir))
@@ -214,7 +225,7 @@ namespace UE::PixelStreamingServers
 
 		// We don't want to make the connections shared to prevent someone accidentally holding on to it. So we use it raw here
 		FWebSocketConnection* PlayerWS = (*PlayersWS->GetConnections().Find(PlayerConnectionId)).Get();
-		bool bUESendsOffer = !PlayerWS->GetUrlArgs().Contains(TEXT("OfferToReceive=true"));
+		bool				  bUESendsOffer = !PlayerWS->GetUrlArgs().Contains(TEXT("OfferToReceive=true"));
 
 		// Send "playerConnected" message to streamer which kicks off making a new RTC connection
 		TSharedRef<FJsonObject> OnPlayerConnectedJSON = MakeShared<FJsonObject>();
@@ -271,8 +282,6 @@ namespace UE::PixelStreamingServers
 		TSharedRef<FJsonObject> idJSON = MakeShared<FJsonObject>();
 		idJSON->SetStringField("type", "identify");
 		SendStreamerMessage(ConnectionId, idJSON);
-
-		StreamersWS->NameConnection(ConnectionId, LEGACY_NAME);
 	}
 
 	void FSignallingServer::OnStreamerDisconnected(uint16 ConnectionId)
@@ -295,7 +304,7 @@ namespace UE::PixelStreamingServers
 		const FString Msg = Utils::ToString(Message);
 		UE_LOG(LogPixelStreamingServers, Log, TEXT("From Streamer id=%d: %s"), ConnectionId, *Msg);
 
-		FString MsgType;
+		FString					MsgType;
 		TSharedPtr<FJsonObject> JSONObj = ParseMessage(Msg, MsgType);
 		if (!JSONObj)
 		{
@@ -341,7 +350,7 @@ namespace UE::PixelStreamingServers
 		const FString Msg = Utils::ToString(Message);
 		UE_LOG(LogPixelStreamingServers, Log, TEXT("From Player id=%d: %s"), ConnectionId, *Msg);
 
-		FString MsgType;
+		FString					MsgType;
 		TSharedPtr<FJsonObject> JSONObj = ParseMessage(Msg, MsgType);
 		if (!JSONObj)
 		{
@@ -358,7 +367,7 @@ namespace UE::PixelStreamingServers
 			if (!PlayerSubscriptions.Contains(ConnectionId))
 			{
 				TArray<FString> StreamerConnections = StreamersWS->GetConnectionNames();
-				if(StreamerConnections.Num() == 0)
+				if (StreamerConnections.Num() == 0)
 				{
 					UE_LOG(LogPixelStreamingServers, Error, TEXT("Player %d sent a message, but no streamers were connected"), ConnectionId);
 					return;
@@ -386,7 +395,7 @@ namespace UE::PixelStreamingServers
 
 	void FSignallingServer::OnStreamerPingMessage(uint16 ConnectionId, TSharedPtr<FJsonObject> JSONObj)
 	{
-		const double UnixTime = FDateTime::UtcNow().ToUnixTimestamp();
+		const double			UnixTime = FDateTime::UtcNow().ToUnixTimestamp();
 		TSharedRef<FJsonObject> PongJSON = MakeShared<FJsonObject>();
 		PongJSON->SetStringField("type", "pong");
 		PongJSON->SetNumberField("time", UnixTime);
@@ -409,8 +418,8 @@ namespace UE::PixelStreamingServers
 
 	void FSignallingServer::OnPlayerListStreamersMessage(uint16 ConnectionId, TSharedPtr<FJsonObject> JSONObj)
 	{
-		TSharedRef<FJsonObject> listJSON = MakeShared<FJsonObject>();
-		const TArray<FString> Names = StreamersWS->GetConnectionNames();
+		TSharedRef<FJsonObject>		   listJSON = MakeShared<FJsonObject>();
+		const TArray<FString>		   Names = StreamersWS->GetConnectionNames();
 		TArray<TSharedPtr<FJsonValue>> JsonNames;
 		for (const FString& Name : Names)
 		{
@@ -446,7 +455,7 @@ namespace UE::PixelStreamingServers
 
 	void FSignallingServer::GetNumStreamers(TFunction<void(uint16)> OnNumStreamersReceived)
 	{
-		if(StreamersWS)
+		if (StreamersWS)
 		{
 			OnNumStreamersReceived(StreamersWS->Count());
 		}
diff --git a/Engine/Plugins/Media/PixelStreaming/Source/PixelStreamingServers/Private/Tests/ServerTests.cpp b/Engine/Plugins/Media/PixelStreaming/Source/PixelStreamingServers/Private/Tests/ServerTests.cpp
--- a/Engine/Plugins/Media/PixelStreaming/Source/PixelStreamingServers/Private/Tests/ServerTests.cpp
+++ b/Engine/Plugins/Media/PixelStreaming/Source/PixelStreamingServers/Private/Tests/ServerTests.cpp
@@ -6,11 +6,26 @@
 #include "Misc/Paths.h"
 #include "PixelStreamingServersLog.h"
 #include "ServerUtils.h"
+#include "WebSocketProbe.h"
 
 #if WITH_DEV_AUTOMATION_TESTS
 
 namespace UE::PixelStreamingServers
 {
+	static const int	 HttpPort = 85;
+	static const FString ExpectedWebserverAddress = FString::Printf(TEXT("http://127.0.0.1:%d"), HttpPort);
+	static const FString ExpectedPlayerWSAddress = FString::Printf(TEXT("ws://127.0.0.1:%d"), HttpPort);
+
+	static const int	 SFUPort = 8889;
+	static const FString ExpectedSFUAddress = FString::Printf(TEXT("ws://127.0.0.1:%d"), SFUPort);
+
+	static const int	 StreamerPort = 8989;
+	static const FString ExpectedStreamerAddress = FString::Printf(TEXT("ws://127.0.0.1:%d"), StreamerPort);
+
+	static const int	 MatchmakerPort = 9999;
+	static const FString ExpectedMatchmakerAddress = FString::Printf(TEXT("ws://127.0.0.1:%d"), MatchmakerPort);
+
+	static const bool bTestServerBinary = false;
 
 	DEFINE_LATENT_AUTOMATION_COMMAND_ONE_PARAMETER(FWaitForServerOrTimeout, TSharedPtr<UE::PixelStreamingServers::IServer>, Server);
 	bool FWaitForServerOrTimeout::Update()
@@ -32,42 +47,114 @@ namespace UE::PixelStreamingServers
 		return true;
 	}
 
-	static const int HttpPort = 85;
-	static const FString ExpectedWebserverAddress = FString::Printf(TEXT("http://127.0.0.1:%d"), HttpPort);
-	static const FString ExpectedPlayerWSAddress = FString::Printf(TEXT("ws://127.0.0.1:%d"), HttpPort);
+	DEFINE_LATENT_AUTOMATION_COMMAND_TWO_PARAMETER(FCheckNumStreamers, TSharedPtr<IServer>, Server, uint16, ExpectedNumStreamers);
+	bool FCheckNumStreamers::Update()
+	{
+		if (Server && Server->IsReady())
+		{
+			uint16 ActualNumberOfStreamers = 0;
+
+			Server->GetNumStreamers([&ActualNumberOfStreamers](uint16 NumStreamers) {
+				ActualNumberOfStreamers = NumStreamers;
+			});
+
+			bool	bTestPassed = ActualNumberOfStreamers == ExpectedNumStreamers;
+			FString LogString = FString::Printf(TEXT("Testing num ws connections. Actual=%d | Expected=%d"), ActualNumberOfStreamers, ExpectedNumStreamers);
+			if (bTestPassed)
+			{
+				UE_LOG(LogPixelStreamingServers, Log, TEXT("Success: %s"), *LogString);
+				return true;
+			}
+			else if (GetCurrentRunTime() > 2.0f)
+			{
+				UE_LOG(LogPixelStreamingServers, Error, TEXT("Failed (timed out after 2s): %s"), *LogString);
+				return true;
+			}
+		}
+		return false;
+	}
 
-	static const int SFUPort = 8889;
-	static const FString ExpectedSFUAddress = FString::Printf(TEXT("ws://127.0.0.1:%d"), SFUPort);
+	DEFINE_LATENT_AUTOMATION_COMMAND_THREE_PARAMETER(FStartWebsocketClient, TSharedPtr<IServer>, Server, TSharedPtr<FWebSocketProbe>, Probe, FURL, WSStreamerURL);
+	bool FStartWebsocketClient::Update()
+	{
+		if (GetCurrentRunTime() > 2.0f)
+		{
+			UE_LOG(LogPixelStreamingServers, Error, TEXT("Timed out after 2s of waiting for websocket to connect"));
+			return true;
+		}
 
-	static const int StreamerPort = 8989;
-	static const FString ExpectedStreamerAddress = FString::Printf(TEXT("ws://127.0.0.1:%d"), StreamerPort);
+		// Do not proceed further into the test if server is not ready
+		if (!Server || !Server->IsReady())
+		{
+			return false;
+		}
 
-	static const int MatchmakerPort = 9999;
-	static const FString ExpectedMatchmakerAddress = FString::Printf(TEXT("ws://127.0.0.1:%d"), MatchmakerPort);
+		if (!Probe)
+		{
+			return false;
+		}
 
-	static const bool bTestServerBinary = false;
+		if (Probe->Probe())
+		{
+			return false;
+		}
+		else
+		{
+			UE_LOG(LogPixelStreamingServers, Log, TEXT("Websocket client probe connected."));
+			return true;
+		}
+	}
+
+	DEFINE_LATENT_AUTOMATION_COMMAND_TWO_PARAMETER(FCloseWebsocketClient, TSharedPtr<IServer>, Server, TSharedPtr<FWebSocketProbe>, Probe);
+	bool FCloseWebsocketClient::Update()
+	{
+		if (GetCurrentRunTime() > 2.0f)
+		{
+			UE_LOG(LogPixelStreamingServers, Error, TEXT("Timed out after 2s of waiting for websocket to close"));
+			return true;
+		}
+
+		// Do not proceed further into the test if server is not ready
+		if (!Server || !Server->IsReady())
+		{
+			return false;
+		}
+
+		// Do not proceed further into the test if probe is not setup
+		if (!Probe)
+		{
+			return false;
+		}
+
+		if (Probe->IsConnected())
+		{
+			Probe->Close();
+			UE_LOG(LogPixelStreamingServers, Log, TEXT("Asked websocket client probe to close."));
+			return false;
+		}
+		else
+		{
+			UE_LOG(LogPixelStreamingServers, Log, TEXT("Websocket client probe closed."));
+			return true;
+		}
+	}
 
 	FString GetCirrusBinaryAbsPath()
 	{
-		FString ServerPath = 
-			FPaths::EnginePluginsDir() / 
-			TEXT("Media") / 
-			TEXT("PixelStreaming") / 
-			TEXT("Resources") /
-			TEXT("WebServers") /
-			TEXT("SignallingWebServer");
-
-#if PLATFORM_WINDOWS
+		FString ServerPath =
+			FPaths::EnginePluginsDir() / TEXT("Media") / TEXT("PixelStreaming") / TEXT("Resources") / TEXT("WebServers") / TEXT("SignallingWebServer");
+
+	#if PLATFORM_WINDOWS
 		ServerPath = ServerPath / TEXT("cirrus.exe");
-#elif PLATFORM_LINUX
+	#elif PLATFORM_LINUX
 		ServerPath = ServerPath / TEXT("cirrus");
-#elif PLATFORM_MAC
-        UE_LOG(LogPixelStreamingServers, Error, TEXT("No cirrus binaries exist for Mac!"));
-        return "";
-#else
+	#elif PLATFORM_MAC
+		UE_LOG(LogPixelStreamingServers, Error, TEXT("No cirrus binaries exist for Mac!"));
+		return "";
+	#else
 		UE_LOG(LogPixelStreamingServers, Error, TEXT("Unsupported platform for Pixel Streaming."));
 		return "";
-#endif
+	#endif
 
 		ServerPath = FPaths::ConvertRelativePathToFull(ServerPath);
 
@@ -80,23 +167,23 @@ namespace UE::PixelStreamingServers
 		UE_LOG(LogPixelStreamingServers, Log, TEXT("----------- LaunchDownloadedCirrusTest -----------"));
 
 		TSharedPtr<IServer> SignallingServer = MakeCirrusServer();
-		FLaunchArgs LaunchArgs;
+		FLaunchArgs			LaunchArgs;
 		LaunchArgs.bPollUntilReady = true;
 		LaunchArgs.ReconnectionTimeoutSeconds = 30.0f;
 		LaunchArgs.ReconnectionIntervalSeconds = 2.0f;
 		LaunchArgs.ProcessArgs = FString::Printf(TEXT("--HttpPort=%d --SFUPort=%d --StreamerPort=%d --MatchmakerPort=%d"), HttpPort, SFUPort, StreamerPort, MatchmakerPort);
 
-		if(bTestServerBinary)
+		if (bTestServerBinary)
 		{
 			LaunchArgs.ServerBinaryOverridePath = GetCirrusBinaryAbsPath();
 		}
 
 		bool bLaunched = SignallingServer->Launch(LaunchArgs);
 
-		if(!bLaunched)
+		if (!bLaunched)
 		{
 			// If we were unable to launch this means some files were missing, we early exit here because this will always happen on Horde
-			// and we don't want a permanently failing test on Horde. 
+			// and we don't want a permanently failing test on Horde.
 			// Todo: Determine way to only disable test on Horde but not locally, or to make it actually download the required scripts.
 			return true;
 		}
@@ -106,28 +193,28 @@ namespace UE::PixelStreamingServers
 
 			FString ActualWebserverUrl = Utils::ToString(Endpoints[EEndpoint::Signalling_Webserver]);
 			UE_LOG(LogPixelStreamingServers, Log, TEXT("Http address for webserver. Actual=%s | Expected=%s"), *ActualWebserverUrl, *ExpectedWebserverAddress);
-			TestTrue(FString::Printf(TEXT("Http address for webserver. Actual=%s | Expected=%s"), *ActualWebserverUrl, *ExpectedWebserverAddress), 
-			ActualWebserverUrl == ExpectedWebserverAddress);
+			TestTrue(FString::Printf(TEXT("Http address for webserver. Actual=%s | Expected=%s"), *ActualWebserverUrl, *ExpectedWebserverAddress),
+				ActualWebserverUrl == ExpectedWebserverAddress);
 
 			FString ActualStreamerUrl = Utils::ToString(Endpoints[EEndpoint::Signalling_Streamer]);
 			UE_LOG(LogPixelStreamingServers, Log, TEXT("Websocket address for streamer messages. Actual=%s | Expected=%s"), *ActualStreamerUrl, *ExpectedStreamerAddress);
-			TestTrue(FString::Printf(TEXT("Websocket address for streamer messages. Actual=%s | Expected=%s"), *ActualStreamerUrl, *ExpectedStreamerAddress), 
-			ActualStreamerUrl == ExpectedStreamerAddress);
+			TestTrue(FString::Printf(TEXT("Websocket address for streamer messages. Actual=%s | Expected=%s"), *ActualStreamerUrl, *ExpectedStreamerAddress),
+				ActualStreamerUrl == ExpectedStreamerAddress);
 
 			FString ActualPlayersUrl = Utils::ToString(Endpoints[EEndpoint::Signalling_Players]);
 			UE_LOG(LogPixelStreamingServers, Log, TEXT("Websocket address for player messages. Actual=%s | Expected=%s"), *ActualPlayersUrl, *ExpectedPlayerWSAddress);
-			TestTrue(FString::Printf(TEXT("Websocket address for player messages. Actual=%s | Expected=%s"), *ActualPlayersUrl, *ExpectedPlayerWSAddress), 
-			ActualPlayersUrl == ExpectedPlayerWSAddress);
+			TestTrue(FString::Printf(TEXT("Websocket address for player messages. Actual=%s | Expected=%s"), *ActualPlayersUrl, *ExpectedPlayerWSAddress),
+				ActualPlayersUrl == ExpectedPlayerWSAddress);
 
 			FString ActualSFUUrl = Utils::ToString(Endpoints[EEndpoint::Signalling_SFU]);
 			UE_LOG(LogPixelStreamingServers, Log, TEXT("Websocket address for SFU messages. Actual=%s | Expected=%s"), *ActualSFUUrl, *ExpectedSFUAddress);
-			TestTrue(FString::Printf(TEXT("Websocket address for SFU messages. Actual=%s | Expected=%s"), *ActualSFUUrl, *ExpectedSFUAddress), 
-			ActualSFUUrl == ExpectedSFUAddress);
+			TestTrue(FString::Printf(TEXT("Websocket address for SFU messages. Actual=%s | Expected=%s"), *ActualSFUUrl, *ExpectedSFUAddress),
+				ActualSFUUrl == ExpectedSFUAddress);
 
 			FString ActualMatchmakerUrl = Utils::ToString(Endpoints[EEndpoint::Signalling_Matchmaker]);
 			UE_LOG(LogPixelStreamingServers, Log, TEXT("Websocket address for matchmaker messages. Actual=%s | Expected=%s"), *ActualMatchmakerUrl, *ExpectedMatchmakerAddress);
-			TestTrue(FString::Printf(TEXT("Websocket address for matchmaker messages. Actual=%s | Expected=%s"), *ActualMatchmakerUrl, *ExpectedMatchmakerAddress), 
-			ActualMatchmakerUrl == ExpectedMatchmakerAddress);
+			TestTrue(FString::Printf(TEXT("Websocket address for matchmaker messages. Actual=%s | Expected=%s"), *ActualMatchmakerUrl, *ExpectedMatchmakerAddress),
+				ActualMatchmakerUrl == ExpectedMatchmakerAddress);
 		});
 
 		SignallingServer->OnFailedToReady.AddLambda([this]() {
@@ -146,7 +233,7 @@ namespace UE::PixelStreamingServers
 		UE_LOG(LogPixelStreamingServers, Log, TEXT("----------- LaunchEmbeddedCirrusTest -----------"));
 
 		TSharedPtr<IServer> SignallingServer = MakeSignallingServer();
-		FLaunchArgs LaunchArgs;
+		FLaunchArgs			LaunchArgs;
 		LaunchArgs.bPollUntilReady = true;
 		LaunchArgs.ReconnectionTimeoutSeconds = 30.0f;
 		LaunchArgs.ReconnectionIntervalSeconds = 2.0f;
@@ -156,7 +243,7 @@ namespace UE::PixelStreamingServers
 		UE_LOG(LogPixelStreamingServers, Log, TEXT("Embedded cirrus launched: %s"), bLaunched ? TEXT("true") : TEXT("false"));
 		TestTrue("Embedded cirrus launched.", bLaunched);
 
-		if(!bLaunched)
+		if (!bLaunched)
 		{
 			return false;
 		}
@@ -166,19 +253,18 @@ namespace UE::PixelStreamingServers
 
 			FString ActualWebserverUrl = Utils::ToString(Endpoints[EEndpoint::Signalling_Webserver]);
 			UE_LOG(LogPixelStreamingServers, Log, TEXT("Http address for webserver. Actual=%s | Expected=%s"), *ActualWebserverUrl, *ExpectedWebserverAddress);
-			TestTrue(FString::Printf(TEXT("Http address for webserver. Actual=%s | Expected=%s"), *ActualWebserverUrl, *ExpectedWebserverAddress), 
-			ActualWebserverUrl == ExpectedWebserverAddress);
+			TestTrue(FString::Printf(TEXT("Http address for webserver. Actual=%s | Expected=%s"), *ActualWebserverUrl, *ExpectedWebserverAddress),
+				ActualWebserverUrl == ExpectedWebserverAddress);
 
 			FString ActualStreamerUrl = Utils::ToString(Endpoints[EEndpoint::Signalling_Streamer]);
 			UE_LOG(LogPixelStreamingServers, Log, TEXT("Websocket address for streamer messages. Actual=%s | Expected=%s"), *ActualStreamerUrl, *ExpectedStreamerAddress);
-			TestTrue(FString::Printf(TEXT("Websocket address for streamer messages. Actual=%s | Expected=%s"), *ActualStreamerUrl, *ExpectedStreamerAddress), 
-			ActualStreamerUrl == ExpectedStreamerAddress);
+			TestTrue(FString::Printf(TEXT("Websocket address for streamer messages. Actual=%s | Expected=%s"), *ActualStreamerUrl, *ExpectedStreamerAddress),
+				ActualStreamerUrl == ExpectedStreamerAddress);
 
 			FString ActualPlayersUrl = Utils::ToString(Endpoints[EEndpoint::Signalling_Players]);
 			UE_LOG(LogPixelStreamingServers, Log, TEXT("Websocket address for player messages. Actual=%s | Expected=%s"), *ActualPlayersUrl, *ExpectedPlayerWSAddress);
-			TestTrue(FString::Printf(TEXT("Websocket address for player messages. Actual=%s | Expected=%s"), *ActualPlayersUrl, *ExpectedPlayerWSAddress), 
-			ActualPlayersUrl == ExpectedPlayerWSAddress);
-
+			TestTrue(FString::Printf(TEXT("Websocket address for player messages. Actual=%s | Expected=%s"), *ActualPlayersUrl, *ExpectedPlayerWSAddress),
+				ActualPlayersUrl == ExpectedPlayerWSAddress);
 		});
 
 		SignallingServer->OnFailedToReady.AddLambda([this]() {
@@ -191,6 +277,71 @@ namespace UE::PixelStreamingServers
 		return true;
 	}
 
+	IMPLEMENT_SIMPLE_AUTOMATION_TEST(FTwoWebsocketToEmbeddedCirrusTest, "System.Plugins.PixelStreaming.TwoWebsocketToEmbeddedCirrus", EAutomationTestFlags::EditorContext | EAutomationTestFlags::ClientContext | EAutomationTestFlags::ProductFilter)
+	bool FTwoWebsocketToEmbeddedCirrusTest::RunTest(const FString& Parameters)
+	{
+		UE_LOG(LogPixelStreamingServers, Log, TEXT("----------- TwoWebsocketToEmbeddedCirrus -----------"));
+
+		TSharedPtr<IServer> SignallingServer = MakeSignallingServer();
+		FLaunchArgs			LaunchArgs;
+		LaunchArgs.bPollUntilReady = true;
+		LaunchArgs.ReconnectionTimeoutSeconds = 30.0f;
+		LaunchArgs.ReconnectionIntervalSeconds = 2.0f;
+		LaunchArgs.ProcessArgs = FString::Printf(TEXT("--HttpPort=%d --StreamerPort=%d"), HttpPort, StreamerPort);
+
+		bool bLaunched = SignallingServer->Launch(LaunchArgs);
+		UE_LOG(LogPixelStreamingServers, Log, TEXT("Embedded cirrus launched: %s"), bLaunched ? TEXT("true") : TEXT("false"));
+		TestTrue("Embedded cirrus launched.", bLaunched);
+
+		if (!bLaunched)
+		{
+			return false;
+		}
+
+		FURL WSStreamerURL;
+		WSStreamerURL.Protocol = TEXT("ws");
+		WSStreamerURL.Host = TEXT("127.0.0.1");
+		WSStreamerURL.Port = StreamerPort;
+		WSStreamerURL.Map = FString();
+
+		SignallingServer->OnReady.AddLambda([this](TMap<EEndpoint, FURL> Endpoints) {
+			TestTrue("Got server OnReady.", true);
+		});
+
+		SignallingServer->OnFailedToReady.AddLambda([this]() {
+			TestTrue("Server was not ready.", false);
+		});
+
+		// These websocket clients will be used to test number of connections
+		TArray<FString> Protocols;
+		Protocols.Add(FString(TEXT("binary")));
+		TSharedPtr<FWebSocketProbe> Client1 = TSharedPtr<FWebSocketProbe>(new FWebSocketProbe(WSStreamerURL, Protocols));
+		TSharedPtr<FWebSocketProbe> Client2 = TSharedPtr<FWebSocketProbe>(new FWebSocketProbe(WSStreamerURL, Protocols));
+
+		ADD_LATENT_AUTOMATION_COMMAND(FWaitForServerOrTimeout(SignallingServer));
+		// Test that there should be zero streamers connected after the server is initially up
+		ADD_LATENT_AUTOMATION_COMMAND(FCheckNumStreamers(SignallingServer, 0));
+		// Start ws client 1
+		ADD_LATENT_AUTOMATION_COMMAND(FStartWebsocketClient(SignallingServer, Client1, WSStreamerURL));
+		// Check num streamers is 1
+		ADD_LATENT_AUTOMATION_COMMAND(FCheckNumStreamers(SignallingServer, 1));
+		// Start ws client 2
+		ADD_LATENT_AUTOMATION_COMMAND(FStartWebsocketClient(SignallingServer, Client2, WSStreamerURL));
+		// Check num streamers = 2
+		ADD_LATENT_AUTOMATION_COMMAND(FCheckNumStreamers(SignallingServer, 2));
+		// Close client 1
+		ADD_LATENT_AUTOMATION_COMMAND(FCloseWebsocketClient(SignallingServer, Client1));
+		// Check num streamers is 1
+		ADD_LATENT_AUTOMATION_COMMAND(FCheckNumStreamers(SignallingServer, 1));
+		// Close client 2
+		ADD_LATENT_AUTOMATION_COMMAND(FCloseWebsocketClient(SignallingServer, Client2));
+		// Check num streamers is 0
+		ADD_LATENT_AUTOMATION_COMMAND(FCheckNumStreamers(SignallingServer, 0));
+		// Shut down server
+		ADD_LATENT_AUTOMATION_COMMAND(FCleanupServer(SignallingServer));
+		return true;
+	}
+
 	// Todo test where create and teardown signalling server 10 times in quick succession with probe to ensure ports are freed
 
 } // namespace UE::PixelStreamingServers
diff --git a/Engine/Plugins/Media/PixelStreaming/Source/PixelStreamingServers/Private/WebSocketProbe.h b/Engine/Plugins/Media/PixelStreaming/Source/PixelStreamingServers/Private/WebSocketProbe.h
--- a/Engine/Plugins/Media/PixelStreaming/Source/PixelStreamingServers/Private/WebSocketProbe.h
+++ b/Engine/Plugins/Media/PixelStreaming/Source/PixelStreamingServers/Private/WebSocketProbe.h
@@ -13,15 +13,15 @@
 namespace UE::PixelStreamingServers
 {
 
-	/* 
-	* A utility class that tries to establish a websocket connection. 
-	* Useful for testing whether servers have come online yet.
-	*/
+	/*
+	 * A utility class that tries to establish a websocket connection.
+	 * Useful for testing whether servers have come online yet.
+	 */
 	class FWebSocketProbe
 	{
 	private:
 		TSharedRef<IWebSocket> WebSocket;
-		FThreadSafeBool bShouldAttemptReconnect;
+		FThreadSafeBool		   bShouldAttemptReconnect;
 
 	public:
 		FWebSocketProbe(FURL Url, TArray<FString> Protocols = TArray<FString>())
@@ -29,11 +29,24 @@ namespace UE::PixelStreamingServers
 			, bShouldAttemptReconnect(true)
 		{
 			WebSocket->OnConnectionError().AddLambda([Url, &bShouldAttemptReconnect = bShouldAttemptReconnect](const FString& Error) {
-				UE_LOG(LogPixelStreamingServers, Log, TEXT("Probing websocket %s | Msg= \"%s\" | Retrying..."), *Utils::ToString(Url),  *Error);
+				UE_LOG(LogPixelStreamingServers, Log, TEXT("Probing websocket %s | Msg= \"%s\" | Retrying..."), *Utils::ToString(Url), *Error);
 				bShouldAttemptReconnect = true;
 			});
 		}
 
+		void Close()
+		{
+			if (WebSocket->IsConnected())
+			{
+				WebSocket->Close();
+			}
+		}
+
+		bool IsConnected() const
+		{
+			return WebSocket->IsConnected();
+		}
+
 		bool Probe()
 		{
 			bool bIsConnected = WebSocket->IsConnected();
@@ -48,4 +61,4 @@ namespace UE::PixelStreamingServers
 		}
 	};
 
-} // UE::PixelStreamingServers
\ No newline at end of file
+} // namespace UE::PixelStreamingServers
\ No newline at end of file
diff --git a/Engine/Plugins/Media/PixelStreaming/Source/PixelStreamingServers/Private/WebSocketServerWrapper.h b/Engine/Plugins/Media/PixelStreaming/Source/PixelStreamingServers/Private/WebSocketServerWrapper.h
--- a/Engine/Plugins/Media/PixelStreaming/Source/PixelStreamingServers/Private/WebSocketServerWrapper.h
+++ b/Engine/Plugins/Media/PixelStreaming/Source/PixelStreamingServers/Private/WebSocketServerWrapper.h
@@ -20,23 +20,22 @@ namespace UE::PixelStreamingServers
 	public:
 		FWebSocketConnection(INetworkingWebSocket* InSocketConnection);
 		~FWebSocketConnection();
-		uint16 GetId() const;
+		uint16			GetId() const;
 		TArray<FString> GetUrlArgs() const;
-		bool Send(FString Message) const;
-		void SetCallbacks();
+		bool			Send(FString Message) const;
+		void			SetCallbacks();
 
 	public:
 		FWebSocketPacketReceivedCallBack OnPacketReceivedCallback;
-		FWebSocketInfoCallBack OnClosedCallback;
+		FWebSocketInfoCallBack			 OnClosedCallback;
 
 	private:
 		static FThreadSafeCounter IdGenerator;
-		uint16 Id;
-		TArray<FString> UrlArgs;
-		INetworkingWebSocket* SocketConnection;
+		uint16					  Id;
+		TArray<FString>			  UrlArgs;
+		INetworkingWebSocket*	  SocketConnection;
 	};
 
-
 	class FWebSocketServerWrapper : public FTickableGameObject
 	{
 	public:
@@ -53,25 +52,25 @@ namespace UE::PixelStreamingServers
 		 * Note: This message is sent as binary using the websocket protocol because of the WS implementation that is used.
 		 * @return True if the message able to be sent.
 		 */
-		bool Send(uint16 ConnectionId, FString Message) const;
-		bool Send(const FString& ConnectionName, FString Message) const;
-		bool GetFirstConnection(uint16& OutConnectionId) const;
+		bool											Send(uint16 ConnectionId, FString Message) const;
+		bool											Send(const FString& ConnectionName, FString Message) const;
+		bool											GetFirstConnection(uint16& OutConnectionId) const;
 		TMap<uint16, TUniquePtr<FWebSocketConnection>>& GetConnections() { return Connections; }
 
-		void NameConnection(uint16 ConnectionId, const FString& Name);
-		void RemoveName(const FString& Name);
-		bool GetNamedConnection(const FString& Name, uint16& OutConnectionId) const;
+		void			NameConnection(uint16 ConnectionId, const FString& Name);
+		void			RemoveName(const FString& Name);
+		bool			GetNamedConnection(const FString& Name, uint16& OutConnectionId) const;
 		TArray<FString> GetConnectionNames() const;
 
 		/* Begin FTickableGameObject */
-		virtual bool IsTickableWhenPaused() const { return true; }
-		virtual bool IsTickableInEditor() const { return true; }
+		virtual bool			  IsTickableWhenPaused() const { return true; }
+		virtual bool			  IsTickableInEditor() const { return true; }
 		virtual ETickableTickType GetTickableTickType() const override { return ETickableTickType::Always; }
-		virtual void Tick(float DeltaTime) override;
-		TStatId GetStatId() const { RETURN_QUICK_DECLARE_CYCLE_STAT(FWebSocketServerWrapper, STATGROUP_Tickables); }
+		virtual void			  Tick(float DeltaTime) override;
+		TStatId					  GetStatId() const { RETURN_QUICK_DECLARE_CYCLE_STAT(FWebSocketServerWrapper, STATGROUP_Tickables); }
 		/* End FTickableGameObject */
 
-		int32 Count() const { return NamedConnections.Num(); }
+		int32 Count() const { return Connections.Num(); }
 
 	public:
 		DECLARE_MULTICAST_DELEGATE_OneParam(FOnNewWebSocketConnection, uint16 /*Connection Id*/);
@@ -89,13 +88,13 @@ namespace UE::PixelStreamingServers
 		virtual void OnConnectionClosed(uint16 ConnectionId);
 
 	private:
-		FThreadSafeBool bLaunched;
-		TUniquePtr<IWebSocketServer> WSServer;
-		FWebSocketClientConnectedCallBack OnClientConnectedCallback;
+		FThreadSafeBool								   bLaunched;
+		TUniquePtr<IWebSocketServer>				   WSServer;
+		FWebSocketClientConnectedCallBack			   OnClientConnectedCallback;
 		TMap<uint16, TUniquePtr<FWebSocketConnection>> Connections;
-		TMap<FString, uint16> NamedConnections;
-		bool bEnableWebServer = false;
-		TArray<FWebSocketHttpMount> DirectoriesToServe;
+		TMap<FString, uint16>						   NamedConnections;
+		bool										   bEnableWebServer = false;
+		TArray<FWebSocketHttpMount>					   DirectoriesToServe;
 	};
 
-} // UE::PixelStreamingServers
+} // namespace UE::PixelStreamingServers
