[Pixel Streaming] Fix: Test for websocket probe had inverted logic for continuing to probe.

diff --git a/Engine/Plugins/Media/PixelStreaming/Source/PixelStreamingServers/Private/Tests/ServerTests.cpp b/Engine/Plugins/Media/PixelStreaming/Source/PixelStreamingServers/Private/Tests/ServerTests.cpp
--- a/Engine/Plugins/Media/PixelStreaming/Source/PixelStreamingServers/Private/Tests/ServerTests.cpp
+++ b/Engine/Plugins/Media/PixelStreaming/Source/PixelStreamingServers/Private/Tests/ServerTests.cpp
@@ -96,12 +96,12 @@ namespace UE::PixelStreamingServers
 
 		if (Probe->Probe())
 		{
-			return false;
+			UE_LOG(LogPixelStreamingServers, Log, TEXT("Websocket client probe connected."));
+			return true;
 		}
 		else
 		{
-			UE_LOG(LogPixelStreamingServers, Log, TEXT("Websocket client probe connected."));
-			return true;
+			return false;
 		}
 	}
 
