[Pixel Streaming] Fix - Undeprecate `kick` blueprints.

diff --git a/Engine/Plugins/Media/PixelStreaming/Source/PixelStreaming/Private/PixelStreamingBlueprints.cpp b/Engine/Plugins/Media/PixelStreaming/Source/PixelStreaming/Private/PixelStreamingBlueprints.cpp
--- a/Engine/Plugins/Media/PixelStreaming/Source/PixelStreaming/Private/PixelStreamingBlueprints.cpp
+++ b/Engine/Plugins/Media/PixelStreaming/Source/PixelStreaming/Private/PixelStreamingBlueprints.cpp
@@ -54,7 +54,7 @@ void UPixelStreamingBlueprints::SendFile(FString FilePath, FString MimeType, FSt
 	}
 
 	TArray64<uint8> ByteData;
-	bool bSuccess = FFileHelper::LoadFileToArray(ByteData, *FilePath);
+	bool			bSuccess = FFileHelper::LoadFileToArray(ByteData, *FilePath);
 	if (bSuccess)
 	{
 		Streamer->SendFileData(ByteData, MimeType, FileExtension);
@@ -79,7 +79,7 @@ void UPixelStreamingBlueprints::StreamerSendFile(FString StreamerId, FString Fil
 	}
 
 	TArray64<uint8> ByteData;
-	bool bSuccess = FFileHelper::LoadFileToArray(ByteData, *FilePath);
+	bool			bSuccess = FFileHelper::LoadFileToArray(ByteData, *FilePath);
 	if (bSuccess)
 	{
 		Streamer->SendFileData(ByteData, MimeType, FileExtension);
@@ -177,9 +177,7 @@ void UPixelStreamingBlueprints::KickPlayer(FString PlayerId)
 	{
 		return;
 	}
-	PRAGMA_DISABLE_DEPRECATION_WARNINGS
 	Streamer->KickPlayer(ToPlayerId(PlayerId));
-	PRAGMA_ENABLE_DEPRECATION_WARNINGS
 }
 
 void UPixelStreamingBlueprints::StreamerKickPlayer(FString StreamerId, FString PlayerId)
@@ -194,9 +192,7 @@ void UPixelStreamingBlueprints::StreamerKickPlayer(FString StreamerId, FString P
 	{
 		return;
 	}
-	PRAGMA_DISABLE_DEPRECATION_WARNINGS
 	Streamer->KickPlayer(ToPlayerId(PlayerId));
-	PRAGMA_ENABLE_DEPRECATION_WARNINGS
 }
 
 void UPixelStreamingBlueprints::SetPlayerLayerPreference(FString PlayerId, int SpatialLayerId, int TemporalLayerId)
diff --git a/Engine/Plugins/Media/PixelStreaming/Source/PixelStreaming/Private/PixelStreamingBlueprints.h b/Engine/Plugins/Media/PixelStreaming/Source/PixelStreaming/Private/PixelStreamingBlueprints.h
--- a/Engine/Plugins/Media/PixelStreaming/Source/PixelStreaming/Private/PixelStreamingBlueprints.h
+++ b/Engine/Plugins/Media/PixelStreaming/Source/PixelStreaming/Private/PixelStreamingBlueprints.h
@@ -98,18 +98,18 @@ public:
 	static void StreamerUnfreezeStream(FString StreamerId);
 
 	/**
-	 * Kick a player. DEPRECATED
+	 * Kick a player.
 	 * @param   PlayerId         The ID of the player to kick.
 	 */
-	UFUNCTION(BlueprintCallable, Category = "Pixel Streaming", meta=(DeprecatedFunction, DeprecationMessage = "Kick player has been deprecated, will no longer function and will be removed in future versions."))
+	UFUNCTION(BlueprintCallable, Category = "Pixel Streaming")
 	static void KickPlayer(FString PlayerId);
 
 	/**
-	 * Kick a player. DEPRECATED
+	 * Kick a player.
 	 * @param	StreamerId		The streamer which the player belongs
 	 * @param   PlayerId        The ID of the player to kick.
 	 */
-	UFUNCTION(BlueprintCallable, Category = "Pixel Streaming", meta=(DeprecatedFunction, DeprecationMessage = "Kick player has been deprecated, will no longer function and will be removed in future versions."))
+	UFUNCTION(BlueprintCallable, Category = "Pixel Streaming")
 	static void StreamerKickPlayer(FString StreamerId, FString PlayerId);
 
 	/**
diff --git a/Engine/Plugins/Media/PixelStreaming/Source/PixelStreaming/Public/IPixelStreamingStreamer.h b/Engine/Plugins/Media/PixelStreaming/Source/PixelStreaming/Public/IPixelStreamingStreamer.h
--- a/Engine/Plugins/Media/PixelStreaming/Source/PixelStreaming/Public/IPixelStreamingStreamer.h
+++ b/Engine/Plugins/Media/PixelStreaming/Source/PixelStreaming/Public/IPixelStreamingStreamer.h
@@ -203,10 +203,9 @@ public:
 	virtual void SendFileData(const TArray64<uint8>& ByteData, FString& MimeType, FString& FileExtension) = 0;
 
 	/**
-	 * @brief Kick a player by player id. DEPRECATED
+	 * @brief Kick a player by player id.
 	 * @param PlayerId		- The ID of the player to kick
 	 */
-	UE_DEPRECATED(5.4, "Kick player has been deprecated, will no longer function and will be removed in future versions.")
 	virtual void KickPlayer(FPixelStreamingPlayerId PlayerId) = 0;
 
 	/**
@@ -219,7 +218,7 @@ public:
 
 	/**
 	 * @brief Get the connected players
-	 * 
+	 *
 	 * @return TArray<FPixelStreamingPlayerId> The connected players
 	 */
 	virtual TArray<FPixelStreamingPlayerId> GetConnectedPlayers() = 0;
@@ -292,7 +291,7 @@ public:
 	virtual bool GetConfigOption(const FName& OptionName, FString& OutValue) = 0;
 
 	/**
-	*/
+	 */
 	virtual void PlayerRequestsBitrate(FPixelStreamingPlayerId PlayerId, int MinBitrate, int MaxBitrate) = 0;
 
 	virtual void RefreshStreamBitrate() = 0;
