[Pixel Streaming] Fix: Using -PixelStreamingUseMediaCapture caused a crash on shutdown due to attempting to access a deleted UObject.

diff --git a/Engine/Plugins/Media/PixelStreaming/Source/PixelStreaming/Private/PixelStreamingVideoInputMediaCapture.cpp b/Engine/Plugins/Media/PixelStreaming/Source/PixelStreaming/Private/PixelStreamingVideoInputMediaCapture.cpp
--- a/Engine/Plugins/Media/PixelStreaming/Source/PixelStreaming/Private/PixelStreamingVideoInputMediaCapture.cpp
+++ b/Engine/Plugins/Media/PixelStreaming/Source/PixelStreaming/Private/PixelStreamingVideoInputMediaCapture.cpp
@@ -35,7 +35,9 @@ FPixelStreamingVideoInputMediaCapture::FPixelStreamingVideoInputMediaCapture()
 
 FPixelStreamingVideoInputMediaCapture::~FPixelStreamingVideoInputMediaCapture()
 {
-	if (MediaCapture)
+	// We don't need to remove mediacapture from root and clear delegate if engine is shutting down
+	// as UE will already have killed all UObjects by this point.
+	if (!IsEngineExitRequested() && MediaCapture)
 	{
 		MediaCapture->OnStateChangedNative.RemoveAll(this);
 		MediaCapture->RemoveFromRoot();
@@ -45,7 +47,7 @@ FPixelStreamingVideoInputMediaCapture::~FPixelStreamingVideoInputMediaCapture()
 void FPixelStreamingVideoInputMediaCapture::StartActiveViewportCapture()
 {
 	// If we were bound to the OnFrameEnd delegate to ensure a frame was rendered before starting, then we can unset it here.
-	if(OnFrameEndDelegateHandle.IsSet())
+	if (OnFrameEndDelegateHandle.IsSet())
 	{
 		FCoreDelegates::OnEndFrame.Remove(OnFrameEndDelegateHandle.GetValue());
 		OnFrameEndDelegateHandle.Reset();
@@ -87,44 +89,44 @@ TSharedPtr<FPixelCaptureCapturer> FPixelStreamingVideoInputMediaCapture::CreateC
 {
 	switch (FinalFormat)
 	{
-	case PixelCaptureBufferFormat::FORMAT_RHI:
-	{
-		if (FPixelStreamingSettings::GetSimulcastParameters().Layers.Num() == 1 &&
-			FPixelStreamingSettings::GetSimulcastParameters().Layers[0].Scaling == 1.0)
-		{
-			// If we only have a single layer (and it's scale is 1), we can use the no copy capturer 
-			// as we know the output from the media capture will already be the correct format and scale
-			return FPixelCaptureCapturerRHINoCopy::Create(FinalScale);
-		}
-		else
+		case PixelCaptureBufferFormat::FORMAT_RHI:
 		{
-			// "Safe Texture Copy" polls a fence to ensure a GPU copy is complete
-			// the RDG pathway does not poll a fence so is more unsafe but offers
-			// a significant performance increase
-			if (FPixelStreamingSettings::GetCaptureUseFence())
+			if (FPixelStreamingSettings::GetSimulcastParameters().Layers.Num() == 1
+				&& FPixelStreamingSettings::GetSimulcastParameters().Layers[0].Scaling == 1.0)
 			{
-				return FPixelCaptureCapturerRHI::Create(FinalScale);
+				// If we only have a single layer (and it's scale is 1), we can use the no copy capturer
+				// as we know the output from the media capture will already be the correct format and scale
+				return FPixelCaptureCapturerRHINoCopy::Create(FinalScale);
 			}
 			else
 			{
-				return FPixelCaptureCapturerRHIRDG::Create(FinalScale);
+				// "Safe Texture Copy" polls a fence to ensure a GPU copy is complete
+				// the RDG pathway does not poll a fence so is more unsafe but offers
+				// a significant performance increase
+				if (FPixelStreamingSettings::GetCaptureUseFence())
+				{
+					return FPixelCaptureCapturerRHI::Create(FinalScale);
+				}
+				else
+				{
+					return FPixelCaptureCapturerRHIRDG::Create(FinalScale);
+				}
 			}
 		}
-	}
-	case PixelCaptureBufferFormat::FORMAT_I420:
-	{
-		if (FPixelStreamingSettings::GetVPXUseCompute())
-		{
-			return FPixelCaptureCapturerRHIToI420Compute::Create(FinalScale);
-		}
-		else
+		case PixelCaptureBufferFormat::FORMAT_I420:
 		{
-			return FPixelCaptureCapturerRHIToI420CPU::Create(FinalScale);
+			if (FPixelStreamingSettings::GetVPXUseCompute())
+			{
+				return FPixelCaptureCapturerRHIToI420Compute::Create(FinalScale);
+			}
+			else
+			{
+				return FPixelCaptureCapturerRHIToI420CPU::Create(FinalScale);
+			}
 		}
-	}
-	default:
-		// UE_LOG(LogPixelStreaming, Error, TEXT("Unsupported final format %d"), FinalFormat);
-		return nullptr;
+		default:
+			// UE_LOG(LogPixelStreaming, Error, TEXT("Unsupported final format %d"), FinalFormat);
+			return nullptr;
 	}
 }
 
@@ -137,25 +139,25 @@ void FPixelStreamingVideoInputMediaCapture::OnCaptureActiveViewportStateChanged(
 
 	switch (MediaCapture->GetState())
 	{
-	case EMediaCaptureState::Capturing:
-		UE_LOG(LogPixelStreaming, Log, TEXT("Starting media capture for Pixel Streaming."));
-		break;
-	case EMediaCaptureState::Stopped:
-		if (MediaCapture->WasViewportResized())
-		{
-			UE_LOG(LogPixelStreaming, Log, TEXT("Pixel Streaming capture was stopped due to resize, going to restart capture."));
-			// If it was stopped and viewport resized we assume resize caused the stop, so try a restart of capture here.
-			StartActiveViewportCapture();
-		}
-		else
-		{
-			UE_LOG(LogPixelStreaming, Log, TEXT("Stopping media capture for Pixel Streaming."));
-		}
-		break;
-	case EMediaCaptureState::Error:
-		UE_LOG(LogPixelStreaming, Log, TEXT("Pixel Streaming capture hit an error, capturing will stop."));
-		break;
-	default:
-		break;
+		case EMediaCaptureState::Capturing:
+			UE_LOG(LogPixelStreaming, Log, TEXT("Starting media capture for Pixel Streaming."));
+			break;
+		case EMediaCaptureState::Stopped:
+			if (MediaCapture->WasViewportResized())
+			{
+				UE_LOG(LogPixelStreaming, Log, TEXT("Pixel Streaming capture was stopped due to resize, going to restart capture."));
+				// If it was stopped and viewport resized we assume resize caused the stop, so try a restart of capture here.
+				StartActiveViewportCapture();
+			}
+			else
+			{
+				UE_LOG(LogPixelStreaming, Log, TEXT("Stopping media capture for Pixel Streaming."));
+			}
+			break;
+		case EMediaCaptureState::Error:
+			UE_LOG(LogPixelStreaming, Log, TEXT("Pixel Streaming capture hit an error, capturing will stop."));
+			break;
+		default:
+			break;
 	}
 }
