Update get_ps_servers scripts to get 5.5 branch and update build.cs to exclude 5.5 new node_module dirs.

diff --git a/Engine/Plugins/Media/PixelStreaming/Resources/WebServers/get_ps_servers.bat b/Engine/Plugins/Media/PixelStreaming/Resources/WebServers/get_ps_servers.bat
--- a/Engine/Plugins/Media/PixelStreaming/Resources/WebServers/get_ps_servers.bat
+++ b/Engine/Plugins/Media/PixelStreaming/Resources/WebServers/get_ps_servers.bat
@@ -88,11 +88,15 @@ if DEFINED UEVersion (
     SET PSInfraTagOrBranch=UE5.4
     SET IsTag=0
   )
+  if "%UEVersion%"=="5.5" (
+    SET PSInfraTagOrBranch=UE5.5
+    SET IsTag=0
+  )
 )
 
 @Rem If no arguments select a specific version, fetch the appropriate default
 if NOT DEFINED PSInfraTagOrBranch (
-    SET PSInfraTagOrBranch=UE5.4
+    SET PSInfraTagOrBranch=UE5.5
     SET IsTag=0
 )
 echo Tag or branch: !PSInfraTagOrBranch!
diff --git a/Engine/Plugins/Media/PixelStreaming/Resources/WebServers/get_ps_servers.sh b/Engine/Plugins/Media/PixelStreaming/Resources/WebServers/get_ps_servers.sh
--- a/Engine/Plugins/Media/PixelStreaming/Resources/WebServers/get_ps_servers.sh
+++ b/Engine/Plugins/Media/PixelStreaming/Resources/WebServers/get_ps_servers.sh
@@ -26,7 +26,7 @@ print_help() {
 # Set all default variables (e.g. # Name and version of ps-infra that we are downloading)
 PSInfraOrg=EpicGamesExt
 PSInfraRepo=PixelStreamingInfrastructure
-PSInfraTagOrBranch=UE5.4
+PSInfraTagOrBranch=UE5.5
 RefType=heads
 IsTag=0
 ReleaseUrlBase=[REDACTED_PROFILE_URL]/PixelStreamingInfrastructure/releases/download
@@ -86,12 +86,17 @@ then
     PSInfraTagOrBranch=UE5.4
     IsTag=0
   fi
+  if [ "$UEVersion" = "5.5" ]
+  then
+    PSInfraTagOrBranch=UE5.5
+    IsTag=0
+  fi
 fi
 
 # If no arguments select a specific version, fetch the appropriate default
 if [ -z "$PSInfraTagOrBranch" ]
 then
-  PSInfraTagOrBranch=UE5.4
+  PSInfraTagOrBranch=UE5.5
   IsTag=0
 fi
 echo "Tag or branch: $PSInfraTagOrBranch"
diff --git a/Engine/Plugins/Media/PixelStreaming/Source/PixelStreaming/PixelStreaming.Build.cs b/Engine/Plugins/Media/PixelStreaming/Source/PixelStreaming/PixelStreaming.Build.cs
--- a/Engine/Plugins/Media/PixelStreaming/Source/PixelStreaming/PixelStreaming.Build.cs
+++ b/Engine/Plugins/Media/PixelStreaming/Source/PixelStreaming/PixelStreaming.Build.cs
@@ -139,6 +139,10 @@ namespace UnrealBuildTool.Rules
                 AddFolder("SignallingWebServer");
                 AddFolder("Matchmaker");
                 AddFolder("SFU");
+                AddFolder("Extras");
+                AddFolder("Signalling");
+                AddFolder("Common");
+                AddFolder("SS_Test");
 
                 RuntimeDependencies.Add("$(ProjectDir)/Samples/PixelStreaming/WebServers/get_ps_servers.bat", "$(PluginDir)/Resources/WebServers/get_ps_servers.bat", StagedFileType.NonUFS);
                 RuntimeDependencies.Add("$(ProjectDir)/Samples/PixelStreaming/WebServers/get_ps_servers.sh", "$(PluginDir)/Resources/WebServers/get_ps_servers.sh", StagedFileType.NonUFS);
