[PS1][PS2] Fix Pixel Streaming delegate data race.

Several Delegates are removed while broadcast is called on a separate thread, which causes a ensure to fail that checks for the data race. These data races happen in PixelStreaming1 and PixelStreaming2 when ending a stream.

 This fix calls the broadcast in the game thread.

diff --git a/Engine/Plugins/Media/PixelStreaming/Source/PixelStreaming/Private/Stats.cpp b/Engine/Plugins/Media/PixelStreaming/Source/PixelStreaming/Private/Stats.cpp
--- a/Engine/Plugins/Media/PixelStreaming/Source/PixelStreaming/Private/Stats.cpp
+++ b/Engine/Plugins/Media/PixelStreaming/Source/PixelStreaming/Private/Stats.cpp
@@ -190,13 +190,15 @@ namespace UE::PixelStreaming
 
 	void FStats::FireStatChanged(FPixelStreamingPlayerId PlayerId, FName StatName, float StatValue)
 	{
-		// firing off these delegates is not thread safe so we want to mutex this call.
-		FScopeLock Lock(&StatNotificationCS);
-		if (UPixelStreamingDelegates* Delegates = UPixelStreamingDelegates::GetPixelStreamingDelegates())
-		{
-			Delegates->OnStatChangedNative.Broadcast(PlayerId, StatName, StatValue);
-			Delegates->OnStatChanged.Broadcast(PlayerId, StatName, StatValue);
-		}
+		// Broadcast must be done on the GameThread because the GameThread can remove the delegates.
+		// If removing and broadcast happens simultaneously it causes a datarace failure.
+		AsyncTask(ENamedThreads::GameThread, [PlayerId, StatName, StatValue]() {
+			if (UPixelStreamingDelegates* Delegates = UPixelStreamingDelegates::GetPixelStreamingDelegates())
+			{
+				Delegates->OnStatChangedNative.Broadcast(PlayerId, StatName, StatValue);
+				Delegates->OnStatChanged.Broadcast(PlayerId, StatName, StatValue);
+			}
+		});
 	}
 
 	void FStats::UpdateConsoleAutoComplete(TArray<FAutoCompleteCommand>& AutoCompleteList)
diff --git a/Engine/Plugins/Media/PixelStreaming2/Source/PixelStreaming2/Private/Stats.cpp b/Engine/Plugins/Media/PixelStreaming2/Source/PixelStreaming2/Private/Stats.cpp
--- a/Engine/Plugins/Media/PixelStreaming2/Source/PixelStreaming2/Private/Stats.cpp
+++ b/Engine/Plugins/Media/PixelStreaming2/Source/PixelStreaming2/Private/Stats.cpp
@@ -200,13 +200,15 @@ namespace UE::PixelStreaming2
 
 	void FStats::FireStatChanged(const FString& PlayerId, FName StatName, float StatValue)
 	{
-		// firing off these delegates is not thread safe so we want to mutex this call.
-		FScopeLock Lock(&StatNotificationCS);
-		if (UPixelStreaming2Delegates* Delegates = UPixelStreaming2Delegates::Get())
-		{
-			Delegates->OnStatChangedNative.Broadcast(PlayerId, StatName, StatValue);
-			Delegates->OnStatChanged.Broadcast(PlayerId, StatName, StatValue);
-		}
+		// Broadcast must be done on the GameThread because the GameThread can remove the delegates.
+		// If removing and broadcast happens simultaneously it causes a datarace failure.
+		AsyncTask(ENamedThreads::GameThread, [PlayerId, StatName, StatValue]() {
+			if (UPixelStreaming2Delegates* Delegates = UPixelStreaming2Delegates::Get())
+			{
+				Delegates->OnStatChangedNative.Broadcast(PlayerId, StatName, StatValue);
+				Delegates->OnStatChanged.Broadcast(PlayerId, StatName, StatValue);
+			}
+		});
 	}
 
 	void FStats::UpdateConsoleAutoComplete(TArray<FAutoCompleteCommand>& AutoCompleteList)
