Fixed return statement.

#rnx

diff --git a/Engine/Plugins/Media/PixelStreaming/Source/PixelStreaming/Private/PixelStreamingBlueprints.cpp b/Engine/Plugins/Media/PixelStreaming/Source/PixelStreaming/Private/PixelStreamingBlueprints.cpp
--- a/Engine/Plugins/Media/PixelStreaming/Source/PixelStreaming/Private/PixelStreamingBlueprints.cpp
+++ b/Engine/Plugins/Media/PixelStreaming/Source/PixelStreaming/Private/PixelStreamingBlueprints.cpp
@@ -260,7 +260,7 @@ FString UPixelStreamingBlueprints::GetDefaultStreamerID()
 	IPixelStreamingModule* Module = UE::PixelStreaming::FPixelStreamingModule::GetModule();
 	if (!Module)
 	{
-		FString();
+		return FString();
 	}
 	return Module->GetDefaultStreamerID();
 }
