[PixelStreaming] Fix for crash when SFU and normal player connect to PixelStreaming1.

Problem is the StreamId is used to determine if a stream shares video source for stream sharing. SFU though requires a different encoder so needs a different StreamId to be set. When the StreamId is the same for the SFU and regular stream, the encoder is shared between the streams and is accessed simultaneously which crashes inside NVEnc or AMF depending on the GPU device being used.

This fixes the problem by making sure the SFU and regular streams have a different StreamId.

When adding a FVideoSource from a SFU, set to use a unique Id. FVideoSource will get a unique Id from PixelStreamingVideoInput.

Modifications to PixelStreamingVideoInput class are kept private to allow for a hotfix.

diff --git a/Engine/Plugins/Media/PixelStreaming/Source/PixelStreaming/Private/FrameBufferMultiFormat.h b/Engine/Plugins/Media/PixelStreaming/Source/PixelStreaming/Private/FrameBufferMultiFormat.h
--- a/Engine/Plugins/Media/PixelStreaming/Source/PixelStreaming/Private/FrameBufferMultiFormat.h
+++ b/Engine/Plugins/Media/PixelStreaming/Source/PixelStreaming/Private/FrameBufferMultiFormat.h
@@ -32,6 +32,8 @@ namespace UE::PixelStreaming
 			return nullptr;
 		}
 
+		void SetSourceStreamId(uint32 InStreamId) { StreamId = InStreamId; }
+
 	protected:
 		TSharedPtr<FPixelCaptureCapturerMultiFormat> FrameCapturer;
 		uint32 StreamId;
diff --git a/Engine/Plugins/Media/PixelStreaming/Source/PixelStreaming/Private/PixelStreamingModule.cpp b/Engine/Plugins/Media/PixelStreaming/Source/PixelStreaming/Private/PixelStreamingModule.cpp
--- a/Engine/Plugins/Media/PixelStreaming/Source/PixelStreaming/Private/PixelStreamingModule.cpp
+++ b/Engine/Plugins/Media/PixelStreaming/Source/PixelStreaming/Private/PixelStreamingModule.cpp
@@ -334,7 +334,8 @@ namespace UE::PixelStreaming
 
 	rtc::scoped_refptr<webrtc::VideoTrackSourceInterface> FPixelStreamingModule::CreateExternalVideoSource()
 	{
-		return ExternalVideoSourceGroup->CreateVideoSource([]() { return true; });
+		// New Video source will not work with stream sharing as bInRequireUniqueId is set to true
+		return ExternalVideoSourceGroup->CreateVideoSource(true, []() { return true; });
 	}
 
 	void FPixelStreamingModule::ReleaseExternalVideoSource(const webrtc::VideoTrackSourceInterface* InVideoSource)
diff --git a/Engine/Plugins/Media/PixelStreaming/Source/PixelStreaming/Private/PixelStreamingVideoInput.cpp b/Engine/Plugins/Media/PixelStreaming/Source/PixelStreaming/Private/PixelStreamingVideoInput.cpp
--- a/Engine/Plugins/Media/PixelStreaming/Source/PixelStreaming/Private/PixelStreamingVideoInput.cpp
+++ b/Engine/Plugins/Media/PixelStreaming/Source/PixelStreaming/Private/PixelStreamingVideoInput.cpp
@@ -99,3 +99,8 @@ FString FPixelStreamingVideoInput::ToString()
 {
 	return TEXT("a Video Input");
 }
+
+uint32 FPixelStreamingVideoInput::GetUniqueStreamId()
+{
+	return NextStreamId++;
+}
diff --git a/Engine/Plugins/Media/PixelStreaming/Source/PixelStreaming/Private/Streamer.cpp b/Engine/Plugins/Media/PixelStreaming/Source/PixelStreaming/Private/Streamer.cpp
--- a/Engine/Plugins/Media/PixelStreaming/Source/PixelStreaming/Private/Streamer.cpp
+++ b/Engine/Plugins/Media/PixelStreaming/Source/PixelStreaming/Private/Streamer.cpp
@@ -100,7 +100,7 @@ namespace UE::PixelStreaming
 		Players->Apply([this](FPixelStreamingPlayerId PlayerId, FPlayerContext& PlayerContext) {
 			if (PlayerContext.PeerConnection)
 			{
-				PlayerContext.PeerConnection->SetVideoSource(VideoSourceGroup->CreateVideoSource([this, PlayerId]() { return ShouldPeerGenerateFrames(PlayerId); }));
+				PlayerContext.PeerConnection->SetVideoSource(VideoSourceGroup->CreateVideoSource(PlayerContext.Config.IsSFU, [this, PlayerId]() { return ShouldPeerGenerateFrames(PlayerId); }));
 			}
 		});
 	}
@@ -484,7 +484,7 @@ namespace UE::PixelStreaming
 			{
 				const bool AllowSimulcast = PlayerContext->Config.IsSFU;
 
-				PlayerContext->PeerConnection->SetVideoSource(VideoSourceGroup->CreateVideoSource([this, PlayerId]() { return ShouldPeerGenerateFrames(PlayerId); }));
+				PlayerContext->PeerConnection->SetVideoSource(VideoSourceGroup->CreateVideoSource(PlayerContext->Config.IsSFU, [this, PlayerId]() { return ShouldPeerGenerateFrames(PlayerId); }));
 				PlayerContext->PeerConnection->SetAudioSource(FPixelStreamingPeerConnection::GetApplicationAudioSource());
 				PlayerContext->PeerConnection->SetAudioSink(MakeShared<FAudioSink>());
 			}
diff --git a/Engine/Plugins/Media/PixelStreaming/Source/PixelStreaming/Private/VideoSource.cpp b/Engine/Plugins/Media/PixelStreaming/Source/PixelStreaming/Private/VideoSource.cpp
--- a/Engine/Plugins/Media/PixelStreaming/Source/PixelStreaming/Private/VideoSource.cpp
+++ b/Engine/Plugins/Media/PixelStreaming/Source/PixelStreaming/Private/VideoSource.cpp
@@ -7,11 +7,17 @@
 
 namespace UE::PixelStreaming
 {
-	FVideoSource::FVideoSource(TSharedPtr<FPixelStreamingVideoInput> InVideoInput, const TFunction<bool()>& InShouldGenerateFramesCheck)
+	FVideoSource::FVideoSource(bool bInRequireUniqueId, TSharedPtr<FPixelStreamingVideoInput> InVideoInput, const TFunction<bool()>& InShouldGenerateFramesCheck)
 		: CurrentState(webrtc::MediaSourceInterface::SourceState::kInitializing)
 		, VideoInput(InVideoInput)
 		, ShouldGenerateFramesCheck(InShouldGenerateFramesCheck)
+		, StreamId(0)
+		, bRequireUniqueId(bInRequireUniqueId)
 	{
+		if (bRequireUniqueId)
+		{
+			StreamId = FPixelStreamingVideoInput::GetUniqueStreamId();
+		}
 	}
 
 	void FVideoSource::MaybePushFrame()
@@ -32,12 +38,22 @@ namespace UE::PixelStreaming
 		}
 	}
 
+	void FVideoSource::SetStreamId(rtc::scoped_refptr<webrtc::VideoFrameBuffer> FrameBuffer) const
+	{
+		FFrameBufferMultiFormatBase* FrameBufferMultiFormat = StaticCast<FFrameBufferMultiFormatBase*>(FrameBuffer.get());
+		FrameBufferMultiFormat->SetSourceStreamId(StreamId);
+	}
+
 	void FVideoSource::PushFrame()
 	{
-	    TRACE_CPUPROFILER_EVENT_SCOPE_ON_CHANNEL_STR("PixelStreaming Push Video Frame", PixelStreamingChannel);
+		TRACE_CPUPROFILER_EVENT_SCOPE_ON_CHANNEL_STR("PixelStreaming Push Video Frame", PixelStreamingChannel);
 		static int32 FrameId = 1;
 
 		rtc::scoped_refptr<webrtc::VideoFrameBuffer> FrameBuffer = VideoInput->GetFrameBuffer();
+		if (bRequireUniqueId)
+		{
+			SetStreamId(FrameBuffer);
+		}
 		check(FrameBuffer->width() != 0 && FrameBuffer->height() != 0);
 		webrtc::VideoFrame Frame = webrtc::VideoFrame::Builder()
 									   .set_video_frame_buffer(FrameBuffer)
@@ -54,6 +70,10 @@ namespace UE::PixelStreaming
 		static int32 EmptyFrameId = 1;
 
 		rtc::scoped_refptr<webrtc::VideoFrameBuffer> FrameBuffer = VideoInput->GetEmptyFrameBuffer();
+		if (bRequireUniqueId)
+		{
+			SetStreamId(FrameBuffer);
+		}
 		webrtc::VideoFrame Frame = webrtc::VideoFrame::Builder()
 									   .set_video_frame_buffer(FrameBuffer)
 									   .set_timestamp_us(rtc::TimeMicros())
diff --git a/Engine/Plugins/Media/PixelStreaming/Source/PixelStreaming/Private/VideoSource.h b/Engine/Plugins/Media/PixelStreaming/Source/PixelStreaming/Private/VideoSource.h
--- a/Engine/Plugins/Media/PixelStreaming/Source/PixelStreaming/Private/VideoSource.h
+++ b/Engine/Plugins/Media/PixelStreaming/Source/PixelStreaming/Private/VideoSource.h
@@ -21,7 +21,7 @@ namespace UE::PixelStreaming
 	class FVideoSource : public FAdaptedVideoTrackSource
 	{
 	public:
-		FVideoSource(TSharedPtr<FPixelStreamingVideoInput> InVideoInput, const TFunction<bool()>& InShouldGenerateFramesCheck);
+		FVideoSource(bool bInRequiresUniqueId, TSharedPtr<FPixelStreamingVideoInput> InVideoInput, const TFunction<bool()>& InShouldGenerateFramesCheck);
 		virtual ~FVideoSource() = default;
 
 		void MaybePushFrame();
@@ -37,7 +37,10 @@ namespace UE::PixelStreaming
 		webrtc::MediaSourceInterface::SourceState CurrentState;
 		TSharedPtr<FPixelStreamingVideoInput> VideoInput;
 		TFunction<bool()> ShouldGenerateFramesCheck;
+		uint32 StreamId;
+		bool bRequireUniqueId;
 
+		void SetStreamId(rtc::scoped_refptr<webrtc::VideoFrameBuffer> FrameBuffer) const;
 		void PushFrame();
 		void PushEmptyFrame();
 	};
diff --git a/Engine/Plugins/Media/PixelStreaming/Source/PixelStreaming/Private/VideoSourceGroup.cpp b/Engine/Plugins/Media/PixelStreaming/Source/PixelStreaming/Private/VideoSourceGroup.cpp
--- a/Engine/Plugins/Media/PixelStreaming/Source/PixelStreaming/Private/VideoSourceGroup.cpp
+++ b/Engine/Plugins/Media/PixelStreaming/Source/PixelStreaming/Private/VideoSourceGroup.cpp
@@ -55,12 +55,12 @@ namespace UE::PixelStreaming
 		bCoupleFramerate = Couple;
 	}
 
-	rtc::scoped_refptr<webrtc::VideoTrackSourceInterface> FVideoSourceGroup::CreateVideoSource(const TFunction<bool()>& InShouldGenerateFramesCheck)
+	rtc::scoped_refptr<webrtc::VideoTrackSourceInterface> FVideoSourceGroup::CreateVideoSource(bool bInRequireUniqueId, const TFunction<bool()>& InShouldGenerateFramesCheck)
 	{
 #if WEBRTC_5414
-		rtc::scoped_refptr<FVideoSource> NewVideoSource = rtc::scoped_refptr<FVideoSource>(new FVideoSource(VideoInput, InShouldGenerateFramesCheck));
+		rtc::scoped_refptr<FVideoSource> NewVideoSource = rtc::scoped_refptr<FVideoSource>(new FVideoSource(bInRequireUniqueId, VideoInput, InShouldGenerateFramesCheck));
 #else
-		rtc::scoped_refptr<FVideoSource> NewVideoSource = new FVideoSource(VideoInput, InShouldGenerateFramesCheck);
+		rtc::scoped_refptr<FVideoSource> NewVideoSource = new FVideoSource(InStreamId, VideoInput, InShouldGenerateFramesCheck);
 #endif
 		{
 			FScopeLock Lock(&CriticalSection);
diff --git a/Engine/Plugins/Media/PixelStreaming/Source/PixelStreaming/Private/VideoSourceGroup.h b/Engine/Plugins/Media/PixelStreaming/Source/PixelStreaming/Private/VideoSourceGroup.h
--- a/Engine/Plugins/Media/PixelStreaming/Source/PixelStreaming/Private/VideoSourceGroup.h
+++ b/Engine/Plugins/Media/PixelStreaming/Source/PixelStreaming/Private/VideoSourceGroup.h
@@ -25,7 +25,7 @@ namespace UE::PixelStreaming
 
 		void SetCoupleFramerate(bool Couple);
 
-		rtc::scoped_refptr<webrtc::VideoTrackSourceInterface> CreateVideoSource(const TFunction<bool()>& InShouldGenerateFramesCheck);
+		rtc::scoped_refptr<webrtc::VideoTrackSourceInterface> CreateVideoSource(bool bInRequireUniqueId, const TFunction<bool()>& InShouldGenerateFramesCheck);
 		void RemoveVideoSource(const webrtc::VideoTrackSourceInterface* ToRemove);
 		void RemoveAllVideoSources();
 
diff --git a/Engine/Plugins/Media/PixelStreaming/Source/PixelStreaming/Public/PixelStreamingVideoInput.h b/Engine/Plugins/Media/PixelStreaming/Source/PixelStreaming/Public/PixelStreamingVideoInput.h
--- a/Engine/Plugins/Media/PixelStreaming/Source/PixelStreaming/Public/PixelStreamingVideoInput.h
+++ b/Engine/Plugins/Media/PixelStreaming/Source/PixelStreaming/Public/PixelStreamingVideoInput.h
@@ -8,6 +8,12 @@
 #include "PixelStreamingWebRTCIncludes.h"
 #include "PixelCaptureBufferFormat.h"
 
+// Forward declaration
+namespace UE::PixelStreaming
+{
+	class FVideoSource;
+}
+
 /**
  * The input of the Pixel Streaming system. Frames enter the system when OnFrame is called.
  * The input buffers might need to be converted/adapted to suit the selected video encoder.
@@ -69,7 +75,12 @@ private:
 	void CreateFrameCapturer();
 	void OnCaptureComplete();
 
+	static uint32 GetUniqueStreamId();
+
 	// video inputs need a unique stream id so we can properly divide hardware encoders between inputs
 	uint32 StreamId = 0;
 	static uint32 NextStreamId;
+
+	// FVideoSource needs to set unique StreamId
+	friend class UE::PixelStreaming::FVideoSource;
 };
