Fixing some issues with signalling connection that can happen after the connection has been destroyed.

diff --git a/Engine/Plugins/Media/PixelStreaming/Source/PixelStreaming/Private/PixelStreamingSignallingConnection.cpp b/Engine/Plugins/Media/PixelStreaming/Source/PixelStreaming/Private/PixelStreamingSignallingConnection.cpp
--- a/Engine/Plugins/Media/PixelStreaming/Source/PixelStreaming/Private/PixelStreamingSignallingConnection.cpp
+++ b/Engine/Plugins/Media/PixelStreaming/Source/PixelStreaming/Private/PixelStreamingSignallingConnection.cpp
@@ -23,11 +23,36 @@
 DECLARE_LOG_CATEGORY_EXTERN(LogPixelStreamingSS, Log, VeryVerbose);
 DEFINE_LOG_CATEGORY(LogPixelStreamingSS);
 
+template<typename T>
+void FPixelStreamingSignallingConnection::SafeGameThreadExecute(T&& Func)
+{
+	TWeakPtr<int> WeakSemaphore = AliveSemaphore;
+	UE::PixelStreaming::DoOnGameThread([this, WeakSemaphore, Func]() {
+		if (WeakSemaphore.IsValid())
+		{
+			Func();
+		}
+	});
+}
+
+template<typename T>
+void FPixelStreamingSignallingConnection::SafeGameThreadExecuteAndWait(uint32 Timeout, T&& Func)
+{
+	TWeakPtr<int> WeakSemaphore = AliveSemaphore;
+	UE::PixelStreaming::DoOnGameThreadAndWait(Timeout, [this, WeakSemaphore, Func]() {
+		if (WeakSemaphore.IsValid())
+		{
+			Func();
+		}
+	});
+}
+
 FPixelStreamingSignallingConnection::FPixelStreamingSignallingConnection(TSharedPtr<IPixelStreamingSignallingConnectionObserver> InObserver, FString InStreamerId, TSharedPtr<IWebSocket> InWebSocket)
 	: Observer(InObserver)
 	, StreamerId(InStreamerId)
 	, WebSocket(InWebSocket)
 {
+	AliveSemaphore = MakeShared<int>(1);
 	RegisterHandler("identify", [this](FJsonObjectPtr JsonMsg) { OnIdRequested(); });
 	RegisterHandler("config", [this](FJsonObjectPtr JsonMsg) { OnConfig(JsonMsg); });
 	RegisterHandler("offer", [this](FJsonObjectPtr JsonMsg) { OnSessionDescription(JsonMsg); });
@@ -49,6 +74,7 @@ FPixelStreamingSignallingConnection::~FPixelStreamingSignallingConnection()
 
 	WebSocket.Reset();
 	Observer.Reset();
+	AliveSemaphore.Reset();
 }
 
 void FPixelStreamingSignallingConnection::Connect(FString InUrl, bool bIsReconnect)
@@ -345,15 +371,6 @@ void FPixelStreamingSignallingConnection::SetKeepAlive(bool bKeepAlive)
 	bKeepAliveEnabled = bKeepAlive;
 }
 
-void FPixelStreamingSignallingConnection::KeepAlive()
-{
-	FJsonObjectPtr Json = MakeShared<FJsonObject>();
-	const double UnixTime = FDateTime::UtcNow().ToUnixTimestamp();
-	Json->SetStringField(TEXT("type"), TEXT("ping"));
-	Json->SetNumberField(TEXT("time"), UnixTime);
-	SendMessage(UE::PixelStreaming::ToString(Json, false));
-}
-
 void FPixelStreamingSignallingConnection::OnConnected()
 {
 	UE_LOG(LogPixelStreamingSS, Log, TEXT("Connected to SS %s"), *Url);
@@ -847,12 +864,29 @@ void FPixelStreamingSignallingConnection::StartKeepAliveTimer()
 {
 	// Dereferencing needs to happen on the game thread
 	// we dont need to wait since its just setting the timer
-	UE::PixelStreaming::DoOnGameThread([this]() {
+	SafeGameThreadExecute([this]() {
 		using namespace UE::PixelStreamingSignallingConnection::Private;
 		FTimerManager* TimerManager = GetTimerManager();
 		if (TimerManager && !TimerManager->IsTimerActive(TimerHandle_KeepAlive))
 		{
-			TimerManager->SetTimer(TimerHandle_KeepAlive, FTimerDelegate::CreateRaw(this, &FPixelStreamingSignallingConnection::KeepAlive), KEEP_ALIVE_INTERVAL, true);
+			TimerManager->SetTimer(TimerHandle_KeepAlive, FTimerDelegate::CreateLambda([this, TimerManager, TimerHandle = TimerHandle_KeepAlive, WeakSemaphore = AliveSemaphore.ToWeakPtr()]() {
+				if (WeakSemaphore.IsValid())
+				{
+					FJsonObjectPtr Json = MakeShared<FJsonObject>();
+					const double UnixTime = FDateTime::UtcNow().ToUnixTimestamp();
+					Json->SetStringField(TEXT("type"), TEXT("ping"));
+					Json->SetNumberField(TEXT("time"), UnixTime);
+					SendMessage(UE::PixelStreaming::ToString(Json, false));
+				}
+				else
+				{
+					// Owning connection has gone away for whatever reason. Clear this timer.
+					// TimerHandle captured is const and we cant use TimerHandle_KeepAlive because 'this' might have been deleted.
+					// ClearTimer takes a non const reference so we cant use TimerHandle.
+					FTimerHandle Handle = TimerHandle;
+					TimerManager->ClearTimer(Handle);
+				}
+			}), KEEP_ALIVE_INTERVAL, true);
 		}
 	});
 }
@@ -862,7 +896,7 @@ void FPixelStreamingSignallingConnection::StopKeepAliveTimer()
 	// Dereferencing needs to happen on the game thread
 	// we need to wait because if we're destructing this object we dont
 	// want to call the callback mid/post destruction
-	UE::PixelStreaming::DoOnGameThreadAndWait(MAX_uint32, [this]() {
+	SafeGameThreadExecuteAndWait(MAX_uint32, [this]() {
 		using namespace UE::PixelStreamingSignallingConnection::Private;
 		FTimerManager* TimerManager = GetTimerManager();
 		if (TimerManager)
@@ -875,7 +909,7 @@ void FPixelStreamingSignallingConnection::StopKeepAliveTimer()
 void FPixelStreamingSignallingConnection::StartReconnectTimer()
 {
 	// Dereferencing needs to happen on the game thread
-	UE::PixelStreaming::DoOnGameThread([this]() {
+	SafeGameThreadExecute([this]() {
 		if (IsEngineExitRequested())
 		{
 			return;
@@ -885,15 +919,21 @@ void FPixelStreamingSignallingConnection::StartReconnectTimer()
 		if (TimerManager && !TimerManager->IsTimerActive(TimerHandle_Reconnect))
 		{
 			float ReconnectInterval = UE::PixelStreaming::Settings::CVarPixelStreamingSignalingReconnectInterval.GetValueOnAnyThread();
-			TimerManager->SetTimer(TimerHandle_Reconnect, FTimerDelegate::CreateRaw(this, &FPixelStreamingSignallingConnection::Connect, Url, true), ReconnectInterval, true);
+			FTimerDelegate delegate = FTimerDelegate::CreateLambda([this, WeakSemaphore = AliveSemaphore.ToWeakPtr()]() {
+				if (WeakSemaphore.IsValid())
+				{
+					Connect(Url, true);
+				}
+				});
+			TimerManager->SetTimer(TimerHandle_Reconnect, delegate, ReconnectInterval, true);
 		}
 	});
 }
 
 void FPixelStreamingSignallingConnection::StopReconnectTimer()
 {
-	// Dereferencing needs to happen on the game thread
-	UE::PixelStreaming::DoOnGameThread([this]() {
+	// GWorld dereferencing needs to happen on the game thread
+	SafeGameThreadExecute([this]() {
 		if (IsEngineExitRequested())
 		{
 			return;
diff --git a/Engine/Plugins/Media/PixelStreaming/Source/PixelStreaming/Public/PixelStreamingSignallingConnection.h b/Engine/Plugins/Media/PixelStreaming/Source/PixelStreaming/Public/PixelStreamingSignallingConnection.h
--- a/Engine/Plugins/Media/PixelStreaming/Source/PixelStreaming/Public/PixelStreamingSignallingConnection.h
+++ b/Engine/Plugins/Media/PixelStreaming/Source/PixelStreaming/Public/PixelStreamingSignallingConnection.h
@@ -42,7 +42,6 @@ public:
 
 private:
 	void Connect(FString Url, bool bIsReconnect);
-	void KeepAlive();
 	void OnConnected();
 	void OnConnectionError(const FString& Error);
 	void OnClosed(int32 StatusCode, const FString& Reason, bool bWasClean);
@@ -103,4 +102,16 @@ private:
 	const float KEEP_ALIVE_INTERVAL = 60.0f;
 
 	TMap<FString, TFunction<void(FJsonObjectPtr)>> MessageHandlers;
+
+	// we can used WeakPtrs pointing to this SharedPtr to detect when this class is destroyed
+	// that allows the following Safe... delayed executions to safely check if the connection
+	// no longer exists.
+	TSharedPtr<int> AliveSemaphore;
+
+	/**
+	 * These templates are defined in the cpp since theyre private and we dont want to include the private
+	 * include file "Utils.h" here in this public header.
+	 */
+	template<typename T> void SafeGameThreadExecute(T&& Func);
+	template<typename T> void SafeGameThreadExecuteAndWait(uint32 Timeout, T&& Func);
 };
