Fix crash in pixel streaming audio component when doing a live coding compile

#rnx

diff --git a/Engine/Plugins/Media/PixelStreaming/Source/PixelStreaming/Private/PixelStreamingAudioComponent.cpp b/Engine/Plugins/Media/PixelStreaming/Source/PixelStreaming/Private/PixelStreamingAudioComponent.cpp
--- a/Engine/Plugins/Media/PixelStreaming/Source/PixelStreaming/Private/PixelStreamingAudioComponent.cpp
+++ b/Engine/Plugins/Media/PixelStreaming/Source/PixelStreaming/Private/PixelStreamingAudioComponent.cpp
@@ -112,13 +112,18 @@ void UPixelStreamingAudioComponent::Reset()
 {
 	PlayerToHear = FString();
 	StreamerToHear = FString();
-	SoundGenerator->bShouldGenerateAudio = false;
+
 	if (AudioSink)
 	{
 		AudioSink->RemoveAudioConsumer(this);
 	}
 	AudioSink = nullptr;
-	SoundGenerator->EmptyBuffers();
+
+	if (SoundGenerator)
+	{
+		SoundGenerator->bShouldGenerateAudio = false;
+		SoundGenerator->EmptyBuffers();
+	}
 }
 
 bool UPixelStreamingAudioComponent::IsListeningToPlayer()
