Fix failing automation tests in Pixel Streaming servers to do with websocket shutdown/removal

diff --git a/Engine/Plugins/Media/PixelStreaming/Source/PixelStreamingServers/Private/SignallingServer.cpp b/Engine/Plugins/Media/PixelStreaming/Source/PixelStreamingServers/Private/SignallingServer.cpp
--- a/Engine/Plugins/Media/PixelStreaming/Source/PixelStreamingServers/Private/SignallingServer.cpp
+++ b/Engine/Plugins/Media/PixelStreaming/Source/PixelStreamingServers/Private/SignallingServer.cpp
@@ -42,22 +42,14 @@ namespace UE::PixelStreamingServers
 		}
 		else
 		{
-			bool bConnected = Probe->Probe();
+			bool bConnected = Probe && Probe->Probe();
 			if (bConnected)
 			{
 				// Close the websocket connection so others can use it
 				Probe->Close();
 				Probe.Reset();
-
-				// Even after closing the client WS of the probe above it will take another tick to remove the connection
-				// so we manually remove the probe connection here to prevent any erroneous number of streamers being reported.
-				uint16 ProbeId = -1;
-				bool   bHasProbe = StreamersWS->GetFirstConnection(ProbeId);
-				if (bHasProbe)
-				{
-					StreamersWS->Close(ProbeId);
-				}
-
+				// Note: even after closing the client WS of the probe above it will take another tick to remove the connection
+				// from the ws server so if you query number of streamers during this time we do have to remove the probe from that count
 				return true;
 			}
 			else
@@ -457,7 +449,14 @@ namespace UE::PixelStreamingServers
 	{
 		if (StreamersWS)
 		{
-			OnNumStreamersReceived(StreamersWS->Count());
+			int NConnections = StreamersWS->Count();
+
+			// If the probe is currently connected, it is not a streamer, so don't count it.
+			if(Probe && Probe->IsConnected()) {
+				NConnections = NConnections - 1;
+			}
+
+			OnNumStreamersReceived(NConnections);
 		}
 		else
 		{
diff --git a/Engine/Plugins/Media/PixelStreaming/Source/PixelStreamingServers/Private/WebSocketProbe.h b/Engine/Plugins/Media/PixelStreaming/Source/PixelStreamingServers/Private/WebSocketProbe.h
--- a/Engine/Plugins/Media/PixelStreaming/Source/PixelStreamingServers/Private/WebSocketProbe.h
+++ b/Engine/Plugins/Media/PixelStreaming/Source/PixelStreamingServers/Private/WebSocketProbe.h
@@ -22,11 +22,13 @@ namespace UE::PixelStreamingServers
 	private:
 		TSharedRef<IWebSocket> WebSocket;
 		FThreadSafeBool		   bShouldAttemptReconnect;
+		FThreadSafeBool		   bIsClosing = false;
 
 	public:
 		FWebSocketProbe(FURL Url, TArray<FString> Protocols = TArray<FString>())
 			: WebSocket(FWebSocketsModule::Get().CreateWebSocket(Utils::ToString(Url), Protocols))
 			, bShouldAttemptReconnect(true)
+			, bIsClosing(false)
 		{
 			WebSocket->OnConnectionError().AddLambda([Url, &bShouldAttemptReconnect = bShouldAttemptReconnect](const FString& Error) {
 				UE_LOG(LogPixelStreamingServers, Log, TEXT("Probing websocket %s | Msg= \"%s\" | Retrying..."), *Utils::ToString(Url), *Error);
@@ -36,10 +38,11 @@ namespace UE::PixelStreamingServers
 
 		void Close()
 		{
-			if (WebSocket->IsConnected())
+			if (WebSocket->IsConnected() && !bIsClosing)
 			{
 				WebSocket->Close();
 			}
+			bIsClosing = true;
 		}
 
 		bool IsConnected() const
@@ -55,6 +58,7 @@ namespace UE::PixelStreamingServers
 			{
 				WebSocket->Connect();
 				bShouldAttemptReconnect = false;
+				bIsClosing = false;
 			}
 
 			return bIsConnected;
