[PS1] Fix: Failing tests

diff --git a/Engine/Plugins/Media/PixelStreaming/Source/PixelStreaming/Private/FrameBufferMultiFormat.cpp b/Engine/Plugins/Media/PixelStreaming/Source/PixelStreaming/Private/FrameBufferMultiFormat.cpp
--- a/Engine/Plugins/Media/PixelStreaming/Source/PixelStreaming/Private/FrameBufferMultiFormat.cpp
+++ b/Engine/Plugins/Media/PixelStreaming/Source/PixelStreaming/Private/FrameBufferMultiFormat.cpp
@@ -10,24 +10,20 @@ namespace UE::PixelStreaming
 	{
 	}
 
-	FFrameBufferMultiFormatLayered::FFrameBufferMultiFormatLayered(TSharedPtr<FPixelCaptureCapturerMultiFormat> InFrameCapturer, uint32 InStreamId)
+	FFrameBufferMultiFormatLayered::FFrameBufferMultiFormatLayered(TSharedPtr<FPixelCaptureCapturerMultiFormat> InFrameCapturer, uint32 InStreamId, FIntPoint SourceResolution)
 		: FFrameBufferMultiFormatBase(InFrameCapturer, InStreamId)
+		, SourceResolution(SourceResolution)
 	{
 	}
 
 	int FFrameBufferMultiFormatLayered::width() const
 	{
-		return FrameCapturer ? FrameCapturer->GetLayerSizes()[GetNumLayers() - 1].X : -1;
+		return SourceResolution.X;
 	}
 
 	int FFrameBufferMultiFormatLayered::height() const
 	{
-		return FrameCapturer ? FrameCapturer->GetLayerSizes()[GetNumLayers() - 1].Y : -1;
-	}
-
-	int FFrameBufferMultiFormatLayered::GetNumLayers() const
-	{
-		return FrameCapturer ? FrameCapturer->GetNumLayers() : -1;
+		return SourceResolution.Y;
 	}
 
 	rtc::scoped_refptr<FFrameBufferMultiFormat> FFrameBufferMultiFormatLayered::GetLayer(FIntPoint Resolution) const
diff --git a/Engine/Plugins/Media/PixelStreaming/Source/PixelStreaming/Private/FrameBufferMultiFormat.h b/Engine/Plugins/Media/PixelStreaming/Source/PixelStreaming/Private/FrameBufferMultiFormat.h
--- a/Engine/Plugins/Media/PixelStreaming/Source/PixelStreaming/Private/FrameBufferMultiFormat.h
+++ b/Engine/Plugins/Media/PixelStreaming/Source/PixelStreaming/Private/FrameBufferMultiFormat.h
@@ -45,14 +45,16 @@ namespace UE::PixelStreaming
 	class FFrameBufferMultiFormatLayered : public FFrameBufferMultiFormatBase
 	{
 	public:
-		FFrameBufferMultiFormatLayered(TSharedPtr<FPixelCaptureCapturerMultiFormat> InFrameCapturer, uint32 InStreamId);
+		FFrameBufferMultiFormatLayered(TSharedPtr<FPixelCaptureCapturerMultiFormat> InFrameCapturer, uint32 InStreamId, FIntPoint SourceResolution);
 		virtual ~FFrameBufferMultiFormatLayered() = default;
 
 		virtual int width() const override;
 		virtual int height() const override;
 
-		int											GetNumLayers() const;
 		rtc::scoped_refptr<FFrameBufferMultiFormat> GetLayer(FIntPoint Resolution) const;
+
+	private:
+		FIntPoint SourceResolution;
 	};
 
 	/**
diff --git a/Engine/Plugins/Media/PixelStreaming/Source/PixelStreaming/Private/PixelStreamingVideoInput.cpp b/Engine/Plugins/Media/PixelStreaming/Source/PixelStreaming/Private/PixelStreamingVideoInput.cpp
--- a/Engine/Plugins/Media/PixelStreaming/Source/PixelStreaming/Private/PixelStreamingVideoInput.cpp
+++ b/Engine/Plugins/Media/PixelStreaming/Source/PixelStreaming/Private/PixelStreamingVideoInput.cpp
@@ -32,26 +32,28 @@ void FPixelStreamingVideoInput::OnFrame(const IPixelCaptureInputFrame& InputFram
 	}
 
 	Ready = true;
-	FrameCapturer->Capture(InputFrame);
+	
 	LastFrameWidth = InputFrame.GetWidth();
 	LastFrameHeight = InputFrame.GetHeight();
+
+	FrameCapturer->Capture(InputFrame);
 }
 
 rtc::scoped_refptr<webrtc::VideoFrameBuffer> FPixelStreamingVideoInput::GetFrameBuffer()
 {
 #if WEBRTC_5414
-	return rtc::make_ref_counted<UE::PixelStreaming::FFrameBufferMultiFormatLayered>(FrameCapturer, StreamId);
+	return rtc::make_ref_counted<UE::PixelStreaming::FFrameBufferMultiFormatLayered>(FrameCapturer, StreamId, FIntPoint(LastFrameWidth, LastFrameHeight));
 #else
-	return new rtc::RefCountedObject<UE::PixelStreaming::FFrameBufferMultiFormatLayered>(FrameCapturer, StreamId);
+	return new rtc::RefCountedObject<UE::PixelStreaming::FFrameBufferMultiFormatLayered>(FrameCapturer, StreamId, FIntPoint(LastFrameWidth, LastFrameHeight));
 #endif
 }
 
 rtc::scoped_refptr<webrtc::VideoFrameBuffer> FPixelStreamingVideoInput::GetEmptyFrameBuffer()
 {
 #if WEBRTC_5414
-	return rtc::make_ref_counted<UE::PixelStreaming::FFrameBufferMultiFormatLayered>(nullptr, StreamId);
+	return rtc::make_ref_counted<UE::PixelStreaming::FFrameBufferMultiFormatLayered>(nullptr, StreamId, FIntPoint(LastFrameWidth, LastFrameHeight));
 #else
-	return new rtc::RefCountedObject<UE::PixelStreaming::FFrameBufferMultiFormatLayered>(nullptr, StreamId);
+	return new rtc::RefCountedObject<UE::PixelStreaming::FFrameBufferMultiFormatLayered>(nullptr, StreamId, FIntPoint(LastFrameWidth, LastFrameHeight));
 #endif
 }
 
diff --git a/Engine/Plugins/Media/PixelStreaming/Source/PixelStreaming/Private/Tests/CodecUtils.h b/Engine/Plugins/Media/PixelStreaming/Source/PixelStreaming/Private/Tests/CodecUtils.h
--- a/Engine/Plugins/Media/PixelStreaming/Source/PixelStreaming/Private/Tests/CodecUtils.h
+++ b/Engine/Plugins/Media/PixelStreaming/Source/PixelStreaming/Private/Tests/CodecUtils.h
@@ -22,7 +22,6 @@ namespace
 
 		TSharedPtr<IPixelStreamingStreamer> Streamer = CreateStreamer(StreamerPort);
 		TSharedPtr<FPixelStreamingVideoInputI420> VideoInput = MakeShared<FPixelStreamingVideoInputI420>();
-		VideoInput->AddOutputFormat(PixelCaptureBufferFormat::FORMAT_I420);
 		Streamer->SetVideoInput(VideoInput);
 
 		TSharedPtr<FMockPlayer> Player = CreatePlayer(FMockPlayer::EMode::AcceptOffers);
@@ -55,7 +54,6 @@ namespace
 
 		TSharedPtr<IPixelStreamingStreamer> Streamer = CreateStreamer(StreamerPort);
 		TSharedPtr<FPixelStreamingVideoInputI420> VideoInput = MakeShared<FPixelStreamingVideoInputI420>();
-		VideoInput->AddOutputFormat(PixelCaptureBufferFormat::FORMAT_I420);
 		Streamer->SetVideoInput(VideoInput);
 
 		TSharedPtr<FMockPlayer> Player = CreatePlayer(FMockPlayer::EMode::AcceptOffers);
diff --git a/Engine/Plugins/Media/PixelStreaming/Source/PixelStreaming/Private/Tests/DirectionUtils.h b/Engine/Plugins/Media/PixelStreaming/Source/PixelStreaming/Private/Tests/DirectionUtils.h
--- a/Engine/Plugins/Media/PixelStreaming/Source/PixelStreaming/Private/Tests/DirectionUtils.h
+++ b/Engine/Plugins/Media/PixelStreaming/Source/PixelStreaming/Private/Tests/DirectionUtils.h
@@ -20,7 +20,6 @@ namespace
 
 		TSharedPtr<IPixelStreamingStreamer> Streamer = CreateStreamer(StreamerPort);
 		TSharedPtr<FPixelStreamingVideoInputI420> VideoInput = MakeShared<FPixelStreamingVideoInputI420>();
-		VideoInput->AddOutputFormat(PixelCaptureBufferFormat::FORMAT_I420);
 		Streamer->SetVideoInput(VideoInput);
 
 		TSharedPtr<FMockPlayer> Player = CreatePlayer(FMockPlayer::EMode::AcceptOffers);
