[PS1, PS2] Fix: Crash using -stompmalloc due to bad buffer sizes

diff --git a/Engine/Plugins/Media/PixelStreaming/Source/PixelStreaming/Private/AudioInputMixer.cpp b/Engine/Plugins/Media/PixelStreaming/Source/PixelStreaming/Private/AudioInputMixer.cpp
--- a/Engine/Plugins/Media/PixelStreaming/Source/PixelStreaming/Private/AudioInputMixer.cpp
+++ b/Engine/Plugins/Media/PixelStreaming/Source/PixelStreaming/Private/AudioInputMixer.cpp
@@ -171,11 +171,13 @@ namespace UE::PixelStreaming
 			UE_LOG(LogPixelStreaming, VeryVerbose, TEXT("(FAudioInput) Sample rate conversion ratio is : %f."), SampleRateConversionRatio);
 			Resampler.Init(Audio::EResamplingMethod::Linear, SampleRateConversionRatio, InNumChannels);
 
-			int32 NumConvertedSamples = InSampleRate / SampleRateConversionRatio;
+			int32 NumOriginalSamples = NumSamples / InNumChannels;
+			int32 NumConvertedSamples = FMath::CeilToInt(NumSamples / InNumChannels * SampleRateConversionRatio);
 			int32 OutputSamples = INDEX_NONE;
-			AudioBuffer.AddZeroed(NumConvertedSamples);
+			AudioBuffer.AddZeroed(NumConvertedSamples * InNumChannels);
+
 			// Perform the sample rate conversion
-			int32 ErrorCode = Resampler.ProcessAudio(const_cast<float*>(InBuffer), NumSamples, false, AudioBuffer.GetData(), AudioBuffer.Num(), OutputSamples);
+			int32 ErrorCode = Resampler.ProcessAudio(const_cast<float*>(InBuffer), NumOriginalSamples, false, AudioBuffer.GetData(), NumConvertedSamples, OutputSamples);
 			verifyf(OutputSamples <= NumConvertedSamples, TEXT("OutputSamples > NumConvertedSamples"));
 			if (ErrorCode != 0)
 			{
diff --git a/Engine/Plugins/Media/PixelStreaming2/Source/PixelStreaming2/Private/AudioCapturer.cpp b/Engine/Plugins/Media/PixelStreaming2/Source/PixelStreaming2/Private/AudioCapturer.cpp
--- a/Engine/Plugins/Media/PixelStreaming2/Source/PixelStreaming2/Private/AudioCapturer.cpp
+++ b/Engine/Plugins/Media/PixelStreaming2/Source/PixelStreaming2/Private/AudioCapturer.cpp
@@ -59,11 +59,13 @@ namespace UE::PixelStreaming2
 			float SampleRateConversionRatio = static_cast<float>(SampleRate) / static_cast<float>(InSampleRate);
 			Resampler.Init(Audio::EResamplingMethod::Linear, SampleRateConversionRatio, InNumChannels);
 
-			int32 NumConvertedSamples = InNumSamples / SampleRateConversionRatio;
+			int32 NumOriginalSamples = InNumSamples / InNumChannels;
+			int32 NumConvertedSamples = FMath::CeilToInt(InNumSamples / InNumChannels * SampleRateConversionRatio);
 			int32 OutputSamples = INDEX_NONE;
-			AudioBuffer.AddZeroed(NumConvertedSamples);
+			AudioBuffer.AddZeroed(NumConvertedSamples * InNumChannels);
+
 			// Perform the sample rate conversion
-			int32 ErrorCode = Resampler.ProcessAudio(const_cast<float*>(AudioData), InNumSamples, false, AudioBuffer.GetData(), AudioBuffer.Num(), OutputSamples);
+			int32 ErrorCode = Resampler.ProcessAudio(const_cast<float*>(AudioData), NumOriginalSamples, false, AudioBuffer.GetData(), NumConvertedSamples, OutputSamples);
 			verifyf(OutputSamples <= NumConvertedSamples, TEXT("OutputSamples > NumConvertedSamples"));
 			if (ErrorCode != 0)
 			{
