[PS1, PS2] Fix a crash when setting text on a SMultiLineEditableText with PixelStreaming.

Input handler was incorrectly casting to SMultiLineEditableTextBox* instead of SMultiLineEditableText*

diff --git a/Engine/Plugins/Media/PixelStreaming/Source/PixelStreamingInput/Private/PixelStreamingInputHandler.cpp b/Engine/Plugins/Media/PixelStreaming/Source/PixelStreamingInput/Private/PixelStreamingInputHandler.cpp
--- a/Engine/Plugins/Media/PixelStreaming/Source/PixelStreamingInput/Private/PixelStreamingInputHandler.cpp
+++ b/Engine/Plugins/Media/PixelStreaming/Source/PixelStreamingInput/Private/PixelStreamingInputHandler.cpp
@@ -23,7 +23,7 @@
 #include "PixelStreamingInputEnums.h"
 #include "PixelStreamingInputConversion.h"
 #include "Widgets/Input/SEditableText.h"
-#include "Widgets/Input/SMultiLineEditableTextBox.h"
+#include "Widgets/Text/SMultiLineEditableText.h"
 #include "PixelStreamingInputDevice.h"
 
 #if WITH_EDITOR
@@ -1261,7 +1261,7 @@ namespace UE::PixelStreamingInput
 				}
 				else if (bIsMultiLineEditableTextType)
 				{
-					SMultiLineEditableTextBox* TextBox = static_cast<SMultiLineEditableTextBox*>(FocusedWidget.Get());
+					SMultiLineEditableText* TextBox = static_cast<SMultiLineEditableText*>(FocusedWidget.Get());
 					TextBox->SetText(FText::FromString(Text));
 				}
 
diff --git a/Engine/Plugins/Media/PixelStreaming2/Source/PixelStreaming2Input/Private/InputHandler.cpp b/Engine/Plugins/Media/PixelStreaming2/Source/PixelStreaming2Input/Private/InputHandler.cpp
--- a/Engine/Plugins/Media/PixelStreaming2/Source/PixelStreaming2Input/Private/InputHandler.cpp
+++ b/Engine/Plugins/Media/PixelStreaming2/Source/PixelStreaming2Input/Private/InputHandler.cpp
@@ -21,7 +21,7 @@
 #include "IPixelStreaming2HMDModule.h"
 #include "PixelStreaming2InputEnums.h"
 #include "Widgets/Input/SEditableText.h"
-#include "Widgets/Input/SMultiLineEditableTextBox.h"
+#include "Widgets/Text/SMultiLineEditableText.h"
 #include "InputDevice.h"
 #include "DefaultDataProtocol.h"
 #include "Logging.h"
diff --git a/Engine/Plugins/Media/PixelStreaming2/Source/PixelStreaming2RTC/Private/RTCInputHandler.cpp b/Engine/Plugins/Media/PixelStreaming2/Source/PixelStreaming2RTC/Private/RTCInputHandler.cpp
--- a/Engine/Plugins/Media/PixelStreaming2/Source/PixelStreaming2RTC/Private/RTCInputHandler.cpp
+++ b/Engine/Plugins/Media/PixelStreaming2/Source/PixelStreaming2RTC/Private/RTCInputHandler.cpp
@@ -13,7 +13,7 @@
 #include "PixelStreaming2PluginSettings.h"
 #include "PixelStreaming2Utils.h"
 #include "Widgets/Input/SEditableText.h"
-#include "Widgets/Input/SMultiLineEditableTextBox.h"
+#include "Widgets/Text/SMultiLineEditableText.h"
 
 namespace UE::PixelStreaming2
 {
@@ -601,7 +601,7 @@ namespace UE::PixelStreaming2
 				}
 				else if (bIsMultiLineEditableTextType)
 				{
-					SMultiLineEditableTextBox* TextBox = static_cast<SMultiLineEditableTextBox*>(FocusedWidget.Get());
+					SMultiLineEditableText* TextBox = static_cast<SMultiLineEditableText*>(FocusedWidget.Get());
 					TextBox->SetText(FText::FromString(Text));
 				}
 
