[PS1] Fix signalling server unable to connect when port is already in use.

This fixes the unit tests when PS1 would try to connect to a used port.

This change backports changes from PS2 in CL 40640200

diff --git a/Engine/Plugins/Media/PixelStreaming/Source/PixelStreaming/Private/Tests/TestUtils.cpp b/Engine/Plugins/Media/PixelStreaming/Source/PixelStreaming/Private/Tests/TestUtils.cpp
--- a/Engine/Plugins/Media/PixelStreaming/Source/PixelStreaming/Private/Tests/TestUtils.cpp
+++ b/Engine/Plugins/Media/PixelStreaming/Source/PixelStreaming/Private/Tests/TestUtils.cpp
@@ -6,6 +6,7 @@
 #include "PixelStreamingPrivate.h"
 #include "IPixelStreamingModule.h"
 #include "PixelCaptureInputFrameI420.h"
+#include "SocketUtils.h"
 #include "Streamer.h"
 
 #if WITH_DEV_AUTOMATION_TESTS
@@ -16,6 +17,7 @@ namespace UE::PixelStreaming
 	{
 		// Start of IANA un-registerable ports (49152 - 65535)
 		static int NextStreamerPort = 49152;
+		NextStreamerPort = GetNextAvailablePort(NextStreamerPort);
 		return NextStreamerPort++;
 	}
 
@@ -23,6 +25,7 @@ namespace UE::PixelStreaming
 	{
 		// Half of IANA un-registerable ports (49152 - 65535)
 		static int NextPlayerPort = 57344;
+		NextPlayerPort = GetNextAvailablePort(NextPlayerPort);
 		return NextPlayerPort++;
 	}
 
diff --git a/Engine/Plugins/Media/PixelStreaming/Source/PixelStreamingServers/Internal/SocketUtils.h b/Engine/Plugins/Media/PixelStreaming/Source/PixelStreamingServers/Internal/SocketUtils.h
new file mode 100644
--- /dev/null
+++ b/Engine/Plugins/Media/PixelStreaming/Source/PixelStreamingServers/Internal/SocketUtils.h
@@ -0,0 +1,10 @@
+// Copyright Epic Games, Inc. All Rights Reserved.
+
+#pragma once
+
+#include "Misc/Optional.h"
+
+namespace UE::PixelStreaming
+{
+	PIXELSTREAMINGSERVERS_API extern int GetNextAvailablePort(TOptional<int> StartingPort = FNullOpt(0));
+} // namespace UE::PixelStreaming
diff --git a/Engine/Plugins/Media/PixelStreaming/Source/PixelStreamingServers/PixelStreamingServers.build.cs b/Engine/Plugins/Media/PixelStreaming/Source/PixelStreamingServers/PixelStreamingServers.build.cs
--- a/Engine/Plugins/Media/PixelStreaming/Source/PixelStreamingServers/PixelStreamingServers.build.cs
+++ b/Engine/Plugins/Media/PixelStreaming/Source/PixelStreamingServers/PixelStreamingServers.build.cs
@@ -24,6 +24,7 @@ namespace UnrealBuildTool.Rules
 				"Json",
 				"WebSocketNetworking",
 				"WebSockets",
+				"Sockets"
 			});
 
 			PrivateDependencyModuleNames.AddRange(new string[] {});
diff --git a/Engine/Plugins/Media/PixelStreaming/Source/PixelStreamingServers/Private/PixelStreamingServersModule.cpp b/Engine/Plugins/Media/PixelStreaming/Source/PixelStreamingServers/Private/PixelStreamingServersModule.cpp
--- a/Engine/Plugins/Media/PixelStreaming/Source/PixelStreamingServers/Private/PixelStreamingServersModule.cpp
+++ b/Engine/Plugins/Media/PixelStreaming/Source/PixelStreamingServers/Private/PixelStreamingServersModule.cpp
@@ -12,10 +12,4 @@ void FPixelStreamingServersModule::StartupModule()
 	// Any startup logic require goes here.
 }
 
-int FPixelStreamingServersModule::NextPort()
-{
-	// Todo: Checking for in-use ports.
-	return (4000 + NextGeneratedPort.Increment()) % 65535;
-}
-
 IMPLEMENT_MODULE(FPixelStreamingServersModule, PixelStreamingServers)
diff --git a/Engine/Plugins/Media/PixelStreaming/Source/PixelStreamingServers/Private/PixelStreamingServersModule.h b/Engine/Plugins/Media/PixelStreaming/Source/PixelStreamingServers/Private/PixelStreamingServersModule.h
--- a/Engine/Plugins/Media/PixelStreaming/Source/PixelStreamingServers/Private/PixelStreamingServersModule.h
+++ b/Engine/Plugins/Media/PixelStreaming/Source/PixelStreamingServers/Private/PixelStreamingServersModule.h
@@ -11,7 +11,6 @@ class FPixelStreamingServersModule : public IModuleInterface
 public:
 	virtual ~FPixelStreamingServersModule() = default;
 	virtual void StartupModule() override;
-	int NextPort();
 
 public:
 	static FPixelStreamingServersModule& Get();
diff --git a/Engine/Plugins/Media/PixelStreaming/Source/PixelStreamingServers/Private/ServerUtils.cpp b/Engine/Plugins/Media/PixelStreaming/Source/PixelStreamingServers/Private/ServerUtils.cpp
--- a/Engine/Plugins/Media/PixelStreaming/Source/PixelStreamingServers/Private/ServerUtils.cpp
+++ b/Engine/Plugins/Media/PixelStreaming/Source/PixelStreamingServers/Private/ServerUtils.cpp
@@ -11,6 +11,7 @@
 #include "Serialization/JsonReader.h"
 #include "Policies/CondensedJsonPrintPolicy.h"
 #include "Misc/Paths.h"
+#include "SocketUtils.h"
 #include "GenericPlatform/GenericPlatformFile.h"
 
 namespace UE::PixelStreamingServers::Utils
@@ -204,15 +205,40 @@ namespace UE::PixelStreamingServers::Utils
 		return LaunchChildProcess(OutScriptPath, Args, FString(TEXT("Download ps servers")), true /*bRunAsScript*/);
 	}
 
-	void PopulateCirrusEndPoints(FLaunchArgs& InLaunchArgs, TMap<EEndpoint, FURL>& OutEndPoints)
+	bool PopulateCirrusEndPoints(FLaunchArgs& InLaunchArgs, TMap<EEndpoint, FURL>& OutEndPoints)
 	{
 		FPixelStreamingServersModule& Module = FPixelStreamingServersModule::Get();
 
 		// Query for ports, or set them if they don't exist
-		FString StreamerPort = Utils::QueryOrSetProcessArgs(InLaunchArgs, TEXT("--StreamerPort="), FString::FromInt(Module.NextPort()));
-		FString SFUPort = Utils::QueryOrSetProcessArgs(InLaunchArgs, TEXT("--SFUPort="), FString::FromInt(Module.NextPort()));
-		FString MatchmakerPort = Utils::QueryOrSetProcessArgs(InLaunchArgs, TEXT("--MatchmakerPort="), FString::FromInt(Module.NextPort()));
-		FString HttpPort = Utils::QueryOrSetProcessArgs(InLaunchArgs, TEXT("--HttpPort="), FString::FromInt(Module.NextPort()));
+		int FallbackStreamerPort = UE::PixelStreaming::GetNextAvailablePort();
+		int FallbackSFUPort = UE::PixelStreaming::GetNextAvailablePort();
+		int FallbackHttpPort = UE::PixelStreaming::GetNextAvailablePort();
+		int FallbackMatchmakerPort = UE::PixelStreaming::GetNextAvailablePort();
+		if (FallbackStreamerPort == -1)
+		{
+			UE_LOG(LogPixelStreamingServers, Warning, TEXT("Failed to find an available port for streamer connections"));
+			return false;
+		}
+		if (FallbackSFUPort == -1)
+		{
+			UE_LOG(LogPixelStreamingServers, Warning, TEXT("Failed to find an available port for SFU connections"));
+			return false;
+		}
+		if (FallbackHttpPort == -1)
+		{
+			UE_LOG(LogPixelStreamingServers, Warning, TEXT("Failed to find an available port for http connections"));
+			return false;
+		}
+		if (FallbackMatchmakerPort == -1)
+		{
+			UE_LOG(LogPixelStreamingServers, Warning, TEXT("Failed to find an available port for matchmaker connections"));
+			return false;
+		}
+
+		FString StreamerPort = Utils::QueryOrSetProcessArgs(InLaunchArgs, TEXT("--StreamerPort="), FString::FromInt(FallbackStreamerPort));
+		FString SFUPort = Utils::QueryOrSetProcessArgs(InLaunchArgs, TEXT("--SFUPort="), FString::FromInt(FallbackSFUPort));
+		FString MatchmakerPort = Utils::QueryOrSetProcessArgs(InLaunchArgs, TEXT("--MatchmakerPort="), FString::FromInt(FallbackMatchmakerPort));
+		FString HttpPort = Utils::QueryOrSetProcessArgs(InLaunchArgs, TEXT("--HttpPort="), FString::FromInt(FallbackHttpPort));
 
 		// Construct endpoint urls
 
@@ -256,6 +282,8 @@ namespace UE::PixelStreamingServers::Utils
 		OutEndPoints.Add(EEndpoint::Signalling_SFU, SFUUrl);
 		OutEndPoints.Add(EEndpoint::Signalling_Matchmaker, MatchmakerUrl);
 		OutEndPoints.Add(EEndpoint::Signalling_Webserver, WebServerUrl);
+
+		return true;
 	}
 
 	FString ToString(FURL Url)
diff --git a/Engine/Plugins/Media/PixelStreaming/Source/PixelStreamingServers/Private/ServerUtils.h b/Engine/Plugins/Media/PixelStreaming/Source/PixelStreamingServers/Private/ServerUtils.h
--- a/Engine/Plugins/Media/PixelStreaming/Source/PixelStreamingServers/Private/ServerUtils.h
+++ b/Engine/Plugins/Media/PixelStreaming/Source/PixelStreamingServers/Private/ServerUtils.h
@@ -67,8 +67,9 @@ namespace UE::PixelStreamingServers
 		 * Populates the cirrus endpoints by extracting endpoint information from the `InLaunchArgs.ProcessArgs`, and if not found, using defaults.
 		 * @param InLaunchArgs The launch args containing the process args that we will be inspecting.
 		 * @param OutEndPoints The end points for the server if we launch it with these launch args.
+		 * @param return True if the end point was added.
 		 **/
-		void PopulateCirrusEndPoints(FLaunchArgs& InLaunchArgs, TMap<EEndpoint, FURL>& OutEndPoints);
+		bool PopulateCirrusEndPoints(FLaunchArgs& InLaunchArgs, TMap<EEndpoint, FURL>& OutEndPoints);
 
 		/**
 		 * Convert a `FURL` into a `FString`. `FURL::ToString()` already exists; however, it appends a problematic trailing slash to the output string.
diff --git a/Engine/Plugins/Media/PixelStreaming/Source/PixelStreamingServers/Private/SocketUtils.cpp b/Engine/Plugins/Media/PixelStreaming/Source/PixelStreamingServers/Private/SocketUtils.cpp
new file mode 100644
--- /dev/null
+++ b/Engine/Plugins/Media/PixelStreaming/Source/PixelStreamingServers/Private/SocketUtils.cpp
@@ -0,0 +1,33 @@
+// Copyright Epic Games, Inc. All Rights Reserved.
+
+#include "SocketUtils.h"
+
+#include "PixelStreamingServersLog.h"
+#include "SocketSubsystem.h"
+
+namespace UE::PixelStreaming
+{
+	static FThreadSafeCounter NextGeneratedPort = 0;
+
+	int	GetNextAvailablePort(TOptional<int> StartingPort)
+	{
+		int32					  CandidateNextPort = StartingPort.Get((4000 + NextGeneratedPort.Increment()) % 65535);
+		int32					  NumRemainingPorts = 65535 - CandidateNextPort;
+		ISocketSubsystem*		  SocketSubsystem = ISocketSubsystem::Get(PLATFORM_SOCKETSUBSYSTEM);
+		TSharedRef<FInternetAddr> LocalHostAddr = SocketSubsystem->CreateInternetAddr(FNetworkProtocolTypes::IPv4);
+		LocalHostAddr->SetPort(CandidateNextPort);
+
+		// The following logic is done inside a scope to ensure that the FUniqueSocket is destroyed and unbound before we return the port number
+		{
+			FUniqueSocket Socket = SocketSubsystem->CreateUniqueSocket(NAME_Stream, TEXT("DummySocket"), FNetworkProtocolTypes::IPv4);
+			int32		  BoundPort = SocketSubsystem->BindNextPort(Socket.Get(), LocalHostAddr.Get(), NumRemainingPorts, 1);
+			if (BoundPort == 0)
+			{
+				UE_LOG(LogPixelStreamingServers, Warning, TEXT("Failed to find an available port!"));
+				return -1;
+			}
+			CandidateNextPort = BoundPort;
+		}
+		return CandidateNextPort;
+	}
+} // namespace UE::PixelStreaming
