Allow PixelStreamingEditor to be loaded in -game editor instances.

diff --git a/Engine/Plugins/Media/PixelStreaming/PixelStreaming.uplugin b/Engine/Plugins/Media/PixelStreaming/PixelStreaming.uplugin
--- a/Engine/Plugins/Media/PixelStreaming/PixelStreaming.uplugin
+++ b/Engine/Plugins/Media/PixelStreaming/PixelStreaming.uplugin
@@ -31,7 +31,7 @@
 		},
 		{
             "Name": "PixelStreamingEditor",
-            "Type": "Editor",
+            "Type": "UncookedOnly",
             "LoadingPhase": "PreDefault",
             "PlatformAllowList": [ "Win64", "Linux", "Mac" ],
             "TargetDenyList": [ "Server" ]
