Prevent calling into FMainFrameModule when running in -game.

diff --git a/Engine/Plugins/Media/PixelStreaming/Source/PixelStreamingEditor/Private/PixelStreamingEditorModule.cpp b/Engine/Plugins/Media/PixelStreaming/Source/PixelStreamingEditor/Private/PixelStreamingEditorModule.cpp
--- a/Engine/Plugins/Media/PixelStreaming/Source/PixelStreamingEditor/Private/PixelStreamingEditorModule.cpp
+++ b/Engine/Plugins/Media/PixelStreaming/Source/PixelStreamingEditor/Private/PixelStreamingEditorModule.cpp
@@ -133,27 +133,30 @@ void FPixelStreamingEditorModule::InitEditorStreaming(IPixelStreamingModule& Mod
 
 	EditorStreamer->SetConfigOption(FName(*FString(TEXT("DefaultToHover"))), TEXT("true"));
 
-	IMainFrameModule::Get().OnMainFrameCreationFinished().AddLambda([&](TSharedPtr<SWindow> RootWindow, bool bIsRunningStartupDialog) {
-		MaybeResizeEditor(RootWindow);
+	if (GIsEditor)
+	{
+		IMainFrameModule::Get().OnMainFrameCreationFinished().AddLambda([&](TSharedPtr<SWindow> RootWindow, bool bIsRunningStartupDialog) {
+			MaybeResizeEditor(RootWindow);
+			
+			if (UE::EditorPixelStreaming::Settings::CVarEditorPixelStreamingStartOnLaunch.GetValueOnAnyThread())
+			{
+				// Default our source to the full editor
+				UE::EditorPixelStreaming::EStreamTypes Source = UE::EditorPixelStreaming::EStreamTypes::Editor;
 
-		if (UE::EditorPixelStreaming::Settings::CVarEditorPixelStreamingStartOnLaunch.GetValueOnAnyThread())
-		{
-            // Default our source to the full editor
-            UE::EditorPixelStreaming::EStreamTypes Source = UE::EditorPixelStreaming::EStreamTypes::Editor;
-
-            FString SourceStr = UE::EditorPixelStreaming::Settings::CVarEditorPixelStreamingSource.GetValueOnAnyThread();
-            if(SourceStr == TEXT("Editor") || SourceStr == TEXT("editor"))
-            {
-                Source = UE::EditorPixelStreaming::EStreamTypes::Editor;
-            }
-            else if(SourceStr == TEXT("LevelEditor") || SourceStr == TEXT("leveleditor"))
-            {
-                Source = UE::EditorPixelStreaming::EStreamTypes::LevelEditorViewport;
-            }
-
-			StartStreaming(Source);
-		}
-	});
+				FString SourceStr = UE::EditorPixelStreaming::Settings::CVarEditorPixelStreamingSource.GetValueOnAnyThread();
+				if(SourceStr == TEXT("Editor") || SourceStr == TEXT("editor"))
+				{
+					Source = UE::EditorPixelStreaming::EStreamTypes::Editor;
+				}
+				else if(SourceStr == TEXT("LevelEditor") || SourceStr == TEXT("leveleditor"))
+				{
+					Source = UE::EditorPixelStreaming::EStreamTypes::LevelEditorViewport;
+				}
+				
+				StartStreaming(Source);
+			}
+		});
+	}
 
 	if (UPixelStreamingDelegates* Delegates = UPixelStreamingDelegates::GetPixelStreamingDelegates())
 	{
@@ -192,13 +195,16 @@ void FPixelStreamingEditorModule::StartStreaming(UE::EditorPixelStreaming::EStre
 					return;
 				}
 
-				TSharedPtr<SWindow> ParentWindow = IMainFrameModule::Get().GetParentWindow();
-				ParentWindow->Resize(FVector2D(Width, Height));
-				FSlateApplication::Get().OnSizeChanged(ParentWindow->GetNativeWindow().ToSharedRef(), Width, Height);
-                // Triggers the NullApplication to rebuild its DisplayMetrics with the new resolution and inform slate 
-                // about the updated virtual desktop size
-                FSystemResolution::RequestResolutionChange(Width, Height, GSystemResolution.WindowMode);
-	            IConsoleManager::Get().CallAllConsoleVariableSinks();
+				if (GIsEditor)
+				{
+					TSharedPtr<SWindow> ParentWindow = IMainFrameModule::Get().GetParentWindow();
+					ParentWindow->Resize(FVector2D(Width, Height));
+					FSlateApplication::Get().OnSizeChanged(ParentWindow->GetNativeWindow().ToSharedRef(), Width, Height);
+					// Triggers the NullApplication to rebuild its DisplayMetrics with the new resolution and inform slate 
+					// about the updated virtual desktop size
+					FSystemResolution::RequestResolutionChange(Width, Height, GSystemResolution.WindowMode);
+					IConsoleManager::Get().CallAllConsoleVariableSinks();					
+				}
 			});
 	}
 
