[PS1, PS2] Fix unreachable code in PixelStreaming module when compiling UE_SERVER. Code after return is now excluded when compiling with UE_SERVER.

diff --git a/Engine/Plugins/Media/PixelStreaming/Source/PixelStreaming/Private/PixelStreamingModule.cpp b/Engine/Plugins/Media/PixelStreaming/Source/PixelStreaming/Private/PixelStreamingModule.cpp
--- a/Engine/Plugins/Media/PixelStreaming/Source/PixelStreaming/Private/PixelStreamingModule.cpp
+++ b/Engine/Plugins/Media/PixelStreaming/Source/PixelStreaming/Private/PixelStreamingModule.cpp
@@ -78,7 +78,7 @@ namespace UE::PixelStreaming
 #if UE_SERVER
 		// Hack to no-op the rest of the module so Blueprints can still work
 		return;
-#endif
+#else
 
 		// Initialise all settings from command line args etc
 		Settings::InitialiseSettings();
@@ -157,6 +157,7 @@ namespace UE::PixelStreaming
 		// Call FStats::Get() to initialize the singleton
 		FStats::Get();
 		bStartupCompleted = true;
+#endif // UE_SERVER
 	}
 
 	void FPixelStreamingModule::ShutdownModule()
diff --git a/Engine/Plugins/Media/PixelStreaming2/Source/PixelStreaming2RTC/Private/PixelStreaming2RTCModule.cpp b/Engine/Plugins/Media/PixelStreaming2/Source/PixelStreaming2RTC/Private/PixelStreaming2RTCModule.cpp
--- a/Engine/Plugins/Media/PixelStreaming2/Source/PixelStreaming2RTC/Private/PixelStreaming2RTCModule.cpp
+++ b/Engine/Plugins/Media/PixelStreaming2/Source/PixelStreaming2RTC/Private/PixelStreaming2RTCModule.cpp
@@ -50,7 +50,7 @@ namespace UE::PixelStreaming2
 #if UE_SERVER
 		// Hack to no-op the rest of the module so Blueprints can still work
 		return;
-#endif
+#else
 		if (!IsStreamingSupported())
 		{
 			return;
@@ -153,6 +153,7 @@ namespace UE::PixelStreaming2
 		}
 
 		bStartupCompleted = true;
+#endif // UE_SERVER
 	}
 
 	void FPixelStreaming2RTCModule::ShutdownModule()
