[PS1, PS2] Fix: Internal signalling server hitting an ensure during creation due FTickableGameObject not being created on the game thread

#FailedValidations Preflight

diff --git a/Engine/Plugins/Media/PixelStreaming/Source/PixelStreamingServers/Private/ServerBase.cpp b/Engine/Plugins/Media/PixelStreaming/Source/PixelStreamingServers/Private/ServerBase.cpp
--- a/Engine/Plugins/Media/PixelStreaming/Source/PixelStreamingServers/Private/ServerBase.cpp
+++ b/Engine/Plugins/Media/PixelStreaming/Source/PixelStreamingServers/Private/ServerBase.cpp
@@ -4,10 +4,21 @@
 #include "PixelStreamingServersLog.h"
 
 namespace UE::PixelStreamingServers
-{
+{	
+	FServerBase::FServerBase()
+		: FTickableGameObject(ETickableTickType::Never)
+	{
+	}
+
+	FServerBase::~FServerBase()
+	{
+		SetTickableTickType(ETickableTickType::Never);
+	}
 
 	bool FServerBase::Launch(FLaunchArgs& InLaunchArgs)
 	{
+		SetTickableTickType(ETickableTickType::Always);
+		
 		LaunchArgs = InLaunchArgs;
 		bPollUntilReady = InLaunchArgs.bPollUntilReady;
 
diff --git a/Engine/Plugins/Media/PixelStreaming/Source/PixelStreamingServers/Private/ServerBase.h b/Engine/Plugins/Media/PixelStreaming/Source/PixelStreamingServers/Private/ServerBase.h
--- a/Engine/Plugins/Media/PixelStreaming/Source/PixelStreamingServers/Private/ServerBase.h
+++ b/Engine/Plugins/Media/PixelStreaming/Source/PixelStreamingServers/Private/ServerBase.h
@@ -13,7 +13,8 @@ namespace UE::PixelStreamingServers
 	{
 	
 	public:
-		virtual ~FServerBase() = default;
+		FServerBase();
+		virtual ~FServerBase() override;
 
 		/* Begin IServer */
 		virtual void Stop() = 0;
diff --git a/Engine/Plugins/Media/PixelStreaming/Source/PixelStreamingServers/Private/SignallingServer.cpp b/Engine/Plugins/Media/PixelStreaming/Source/PixelStreamingServers/Private/SignallingServer.cpp
--- a/Engine/Plugins/Media/PixelStreaming/Source/PixelStreamingServers/Private/SignallingServer.cpp
+++ b/Engine/Plugins/Media/PixelStreaming/Source/PixelStreamingServers/Private/SignallingServer.cpp
@@ -19,6 +19,7 @@ namespace UE::PixelStreamingServers
 		PlayerMessageHandlers.Add("subscribe").BindRaw(this, &FSignallingServer::OnPlayerSubscribeMessage);
 		PlayerMessageHandlers.Add("unsubscribe").BindRaw(this, &FSignallingServer::OnPlayerUnsubscribeMessage);
 		PlayerMessageHandlers.Add("stats").BindRaw(this, &FSignallingServer::OnPlayerStatsMessage);
+		PlayerMessageHandlers.Add("ping").BindRaw(this, &FSignallingServer::OnPlayerPingMessage);
 	}
 
 	void FSignallingServer::Stop()
@@ -445,6 +446,15 @@ namespace UE::PixelStreamingServers
 		UE_LOG(LogPixelStreamingServers, Log, TEXT("Player %d stats = \n %s"), ConnectionId, *Utils::ToString(JSONObj.ToSharedRef()));
 	}
 
+	void FSignallingServer::OnPlayerPingMessage(uint16 ConnectionId, TSharedPtr<FJsonObject> JSONObj)
+	{
+		const double			UnixTime = FDateTime::UtcNow().ToUnixTimestamp();
+		TSharedRef<FJsonObject> PongJSON = MakeShared<FJsonObject>();
+		PongJSON->SetStringField("type", "pong");
+		PongJSON->SetNumberField("time", UnixTime);
+		SendPlayerMessage(ConnectionId, PongJSON);
+	}
+
 	void FSignallingServer::GetNumStreamers(TFunction<void(uint16)> OnNumStreamersReceived)
 	{
 		if (StreamersWS)
diff --git a/Engine/Plugins/Media/PixelStreaming/Source/PixelStreamingServers/Private/SignallingServer.h b/Engine/Plugins/Media/PixelStreaming/Source/PixelStreamingServers/Private/SignallingServer.h
--- a/Engine/Plugins/Media/PixelStreaming/Source/PixelStreamingServers/Private/SignallingServer.h
+++ b/Engine/Plugins/Media/PixelStreaming/Source/PixelStreamingServers/Private/SignallingServer.h
@@ -58,6 +58,7 @@ namespace UE::PixelStreamingServers
 		void OnPlayerSubscribeMessage(uint16 ConnectionId, TSharedPtr<FJsonObject> JSONObj);
 		void OnPlayerUnsubscribeMessage(uint16 ConnectionId, TSharedPtr<FJsonObject> JSONObj);
 		void OnPlayerStatsMessage(uint16 ConnectionId, TSharedPtr<FJsonObject> JSONObj);
+		void OnPlayerPingMessage(uint16 ConnectionId, TSharedPtr<FJsonObject> JSONObj);
 		
 	protected:
 		TUniquePtr<FWebSocketProbe> Probe;
diff --git a/Engine/Plugins/Media/PixelStreaming/Source/PixelStreamingServers/Private/WebSocketServerWrapper.cpp b/Engine/Plugins/Media/PixelStreaming/Source/PixelStreamingServers/Private/WebSocketServerWrapper.cpp
--- a/Engine/Plugins/Media/PixelStreaming/Source/PixelStreamingServers/Private/WebSocketServerWrapper.cpp
+++ b/Engine/Plugins/Media/PixelStreaming/Source/PixelStreamingServers/Private/WebSocketServerWrapper.cpp
@@ -64,7 +64,8 @@ namespace UE::PixelStreamingServers
 	*/
 
 	FWebSocketServerWrapper::FWebSocketServerWrapper()
-		: bLaunched(false)
+		: FTickableGameObject(ETickableTickType::Never)
+		, bLaunched(false)
 	{
 
 	}
@@ -82,6 +83,8 @@ namespace UE::PixelStreamingServers
 
 	bool FWebSocketServerWrapper::Launch(uint16 Port)
 	{
+		SetTickableTickType(ETickableTickType::Always);
+
 		WSServer = FModuleManager::Get().LoadModuleChecked<IWebSocketNetworkingModule>(TEXT("WebSocketNetworking")).CreateServer();
 		
 		if(bEnableWebServer)
@@ -159,6 +162,8 @@ namespace UE::PixelStreamingServers
 
 	void FWebSocketServerWrapper::Stop()
 	{
+		SetTickableTickType(ETickableTickType::Never);
+		
 		bLaunched = false;
 		Connections.Empty();
 		OnClientConnectedCallback.Unbind();
diff --git a/Engine/Plugins/Media/PixelStreaming2/Source/PixelStreaming2Servers/Private/ServerBase.cpp b/Engine/Plugins/Media/PixelStreaming2/Source/PixelStreaming2Servers/Private/ServerBase.cpp
--- a/Engine/Plugins/Media/PixelStreaming2/Source/PixelStreaming2Servers/Private/ServerBase.cpp
+++ b/Engine/Plugins/Media/PixelStreaming2/Source/PixelStreaming2Servers/Private/ServerBase.cpp
@@ -5,9 +5,20 @@
 
 namespace UE::PixelStreaming2Servers
 {
+	FServerBase::FServerBase()
+		: FTickableGameObject(ETickableTickType::Never)
+	{
+	}
+
+	FServerBase::~FServerBase()
+	{
+		SetTickableTickType(ETickableTickType::Never);
+	}
 
 	bool FServerBase::Launch(FLaunchArgs& InLaunchArgs)
 	{
+		SetTickableTickType(ETickableTickType::Always);
+
 		LaunchArgs = InLaunchArgs;
 		bPollUntilReady = InLaunchArgs.bPollUntilReady;
 
diff --git a/Engine/Plugins/Media/PixelStreaming2/Source/PixelStreaming2Servers/Private/ServerBase.h b/Engine/Plugins/Media/PixelStreaming2/Source/PixelStreaming2Servers/Private/ServerBase.h
--- a/Engine/Plugins/Media/PixelStreaming2/Source/PixelStreaming2Servers/Private/ServerBase.h
+++ b/Engine/Plugins/Media/PixelStreaming2/Source/PixelStreaming2Servers/Private/ServerBase.h
@@ -11,9 +11,9 @@ namespace UE::PixelStreaming2Servers
 	// Base class for all servers, provides functionality for servers to be polled for readiness while they start up.
 	class FServerBase : public IServer, public FTickableGameObject
 	{
-	
 	public:
-		virtual ~FServerBase() = default;
+		FServerBase();
+		virtual ~FServerBase() override;
 
 		/* Begin IServer */
 		virtual void Stop() = 0;
diff --git a/Engine/Plugins/Media/PixelStreaming2/Source/PixelStreaming2Servers/Private/WebSocketServerWrapper.cpp b/Engine/Plugins/Media/PixelStreaming2/Source/PixelStreaming2Servers/Private/WebSocketServerWrapper.cpp
--- a/Engine/Plugins/Media/PixelStreaming2/Source/PixelStreaming2Servers/Private/WebSocketServerWrapper.cpp
+++ b/Engine/Plugins/Media/PixelStreaming2/Source/PixelStreaming2Servers/Private/WebSocketServerWrapper.cpp
@@ -64,7 +64,8 @@ namespace UE::PixelStreaming2Servers
 	 */
 
 	FWebSocketServerWrapper::FWebSocketServerWrapper()
-		: bLaunched(false)
+		: FTickableGameObject(ETickableTickType::Never)
+		, bLaunched(false)
 	{
 	}
 
@@ -83,6 +84,8 @@ namespace UE::PixelStreaming2Servers
 
 	bool FWebSocketServerWrapper::Launch(uint16 Port)
 	{
+		SetTickableTickType(ETickableTickType::Always);
+
 		WSServer = FModuleManager::Get().LoadModuleChecked<IWebSocketNetworkingModule>(TEXT("WebSocketNetworking")).CreateServer();
 
 		if (bEnableWebServer)
@@ -160,6 +163,8 @@ namespace UE::PixelStreaming2Servers
 
 	void FWebSocketServerWrapper::Stop()
 	{
+		SetTickableTickType(ETickableTickType::Never);
+		
 		bLaunched = false;
 		Connections.Empty();
 		OnClientConnectedCallback.Unbind();
