[PixelStreaming. PixelStreaming2] Add: Project setting to disable toolbar

#Synced-CL 48443255

diff --git a/Engine/Plugins/Media/PixelStreaming/Source/PixelStreamingEditor/Private/PixelStreamingToolbar.cpp b/Engine/Plugins/Media/PixelStreaming/Source/PixelStreamingEditor/Private/PixelStreamingToolbar.cpp
--- a/Engine/Plugins/Media/PixelStreaming/Source/PixelStreamingEditor/Private/PixelStreamingToolbar.cpp
+++ b/Engine/Plugins/Media/PixelStreaming/Source/PixelStreamingEditor/Private/PixelStreamingToolbar.cpp
@@ -34,6 +34,7 @@
 #include "PixelStreamingEditorUtils.h"
 #include "SlateFwd.h"
 #include "Widgets/Input/SNumericEntryBox.h"
+#include "EditorPixelStreamingSettings.h"
 
 #include "PixelStreamingCoderUtils.h"
 #include "Video/Encoders/Configs/VideoEncoderConfigAV1.h"
@@ -204,6 +205,11 @@ namespace UE::EditorPixelStreaming
 
 	void FPixelStreamingToolbar::RegisterMenus()
 	{
+		if (!GetDefault<UPixelStreamingSettings>()->EnablePixelStreamingToolbar)
+		{
+			return;
+		}
+
 		FToolMenuOwnerScoped OwnerScoped(this);
 		{
 			UToolMenu* CustomToolBar = UToolMenus::Get()->ExtendMenu("LevelEditor.LevelEditorToolBar.User");
diff --git a/Engine/Plugins/Media/PixelStreaming/Source/PixelStreamingInput/Public/EditorPixelStreamingSettings.h b/Engine/Plugins/Media/PixelStreaming/Source/PixelStreamingInput/Public/EditorPixelStreamingSettings.h
--- a/Engine/Plugins/Media/PixelStreaming/Source/PixelStreamingInput/Public/EditorPixelStreamingSettings.h
+++ b/Engine/Plugins/Media/PixelStreaming/Source/PixelStreamingInput/Public/EditorPixelStreamingSettings.h
@@ -14,6 +14,13 @@ class PIXELSTREAMINGINPUT_API UPixelStreamingSettings : public UDeveloperSetting
 	GENERATED_UCLASS_BODY()
 
 public:
+	UPROPERTY(config, EditAnywhere, Category = "PixelStreaming", meta = (
+		ConfigRestartRequired = true,
+		DisplayName = "Enable Pixel Streaming Toolbar",
+		ToolTip = "Adds the Pixel Streaming toolbar to the main viewport toolbar."
+		))
+	bool EnablePixelStreamingToolbar = true;
+	
 	/**
 	 * Pixel streaming always requires various software cursors so they will be
 	 * visible in the video stream sent to the browser to allow the user to
diff --git a/Engine/Plugins/Media/PixelStreaming2/Source/PixelStreaming2Editor/Private/PixelStreaming2Toolbar.cpp b/Engine/Plugins/Media/PixelStreaming2/Source/PixelStreaming2Editor/Private/PixelStreaming2Toolbar.cpp
--- a/Engine/Plugins/Media/PixelStreaming2/Source/PixelStreaming2Editor/Private/PixelStreaming2Toolbar.cpp
+++ b/Engine/Plugins/Media/PixelStreaming2/Source/PixelStreaming2Editor/Private/PixelStreaming2Toolbar.cpp
@@ -284,6 +284,11 @@ namespace UE::EditorPixelStreaming2
 
 	void FPixelStreaming2Toolbar::RegisterMenus()
 	{
+		if (!GetDefault<UPixelStreaming2PluginSettings>()->EnablePixelStreamingToolbar)
+		{
+			return;
+		}
+
 		FToolMenuOwnerScoped OwnerScoped(this);
 		{
 			UToolMenu* CustomToolBar = UToolMenus::Get()->ExtendMenu("LevelEditor.LevelEditorToolBar.User");
diff --git a/Engine/Plugins/Media/PixelStreaming2/Source/PixelStreaming2Settings/Internal/PixelStreaming2PluginSettings.h b/Engine/Plugins/Media/PixelStreaming2/Source/PixelStreaming2Settings/Internal/PixelStreaming2PluginSettings.h
--- a/Engine/Plugins/Media/PixelStreaming2/Source/PixelStreaming2Settings/Internal/PixelStreaming2PluginSettings.h
+++ b/Engine/Plugins/Media/PixelStreaming2/Source/PixelStreaming2Settings/Internal/PixelStreaming2PluginSettings.h
@@ -88,6 +88,13 @@ class UPixelStreaming2PluginSettings : public UDeveloperSettings
 
 public:
 	// clang-format off
+	UPROPERTY(config, EditAnywhere, Category = "PixelStreaming", meta = (
+		ConfigRestartRequired = true,
+		DisplayName = "Enable Pixel Streaming 2 Toolbar",
+		ToolTip = "Adds the Pixel Streaming 2 toolbar to the main viewport toolbar."
+		))
+	bool EnablePixelStreamingToolbar = true;
+
 	static UE_API TAutoConsoleVariable<bool> CVarLogStats;
 	UPROPERTY(config, EditAnywhere, Category = "PixelStreaming", meta = (
 		DisplayName = "Log Pixel Streaming Stats",
